* The RemoveDependencyNode from the AssetRegistry is now resilient to nodes not not existing in the cache, and safely ignores them.

diff --git a/Engine/Source/Runtime/AssetRegistry/Private/AssetRegistry.cpp b/Engine/Source/Runtime/AssetRegistry/Private/AssetRegistry.cpp
--- a/Engine/Source/Runtime/AssetRegistry/Private/AssetRegistry.cpp
+++ b/Engine/Source/Runtime/AssetRegistry/Private/AssetRegistry.cpp
@@ -1427,7 +1427,15 @@ FDependsNode* FAssetRegistry::CreateOrFindDependsNode(FName ObjectPath)
 
 bool FAssetRegistry::RemoveDependsNode( FName PackageName )
 {
-	FDependsNode* Node = *(CachedDependsNodes.Find( PackageName ));
+	FDependsNode** NodeEntry = CachedDependsNodes.Find(PackageName);
+
+	// If it has already been removed, then we're done.
+	if ( NodeEntry == NULL )
+	{
+		return true;
+	}
+
+	FDependsNode* Node = *(NodeEntry);
 
 	if ( Node != nullptr )
 	{
