UE4: There is no longer a warning if the asset registry cache file does not exist. This is expected in the first run or after the cache is manually cleared.

diff --git a/Engine/Source/Runtime/AssetRegistry/Private/NameTableArchive.cpp b/Engine/Source/Runtime/AssetRegistry/Private/NameTableArchive.cpp
--- a/Engine/Source/Runtime/AssetRegistry/Private/NameTableArchive.cpp
+++ b/Engine/Source/Runtime/AssetRegistry/Private/NameTableArchive.cpp
@@ -10,7 +10,7 @@ FNameTableArchiveReader::FNameTableArchiveReader()
 
 bool FNameTableArchiveReader::LoadFile(const TCHAR* Filename, int32 SerializationVersion)
 {
-	if ( FFileHelper::LoadFileToArray(Reader, Filename) )
+	if ( FFileHelper::LoadFileToArray(Reader, Filename, FILEREAD_Silent) )
 	{
 		int32 MagicNumber = 0;
 		*this << MagicNumber;
