#UE4 Fix for 333838

diff --git a/Engine/Source/Runtime/AssetRegistry/Public/AssetData.h b/Engine/Source/Runtime/AssetRegistry/Public/AssetData.h
--- a/Engine/Source/Runtime/AssetRegistry/Public/AssetData.h
+++ b/Engine/Source/Runtime/AssetRegistry/Public/AssetData.h
@@ -58,6 +58,14 @@ public:
 	{
 		if ( InAsset != NULL )
 		{
+			const UClass* InClass = Cast<UClass>(InAsset);
+			if (InClass && InClass->ClassGeneratedBy)
+			{
+				// this UObject was generated by another class, use the generator to identify
+				// the source asset:
+				InAsset = InClass->ClassGeneratedBy;
+			}
+
 			const UPackage* Outermost = InAsset->GetOutermost();
 
 			// Find the group names
