UE4: Added some code to allow legacy maps to show up in the asset registy when launching with -worldassets. This is mainly for testing the new feature and should be removed upon completion.

diff --git a/Engine/Source/Runtime/AssetRegistry/Private/AssetDataGatherer.cpp b/Engine/Source/Runtime/AssetRegistry/Private/AssetDataGatherer.cpp
--- a/Engine/Source/Runtime/AssetRegistry/Private/AssetDataGatherer.cpp
+++ b/Engine/Source/Runtime/AssetRegistry/Private/AssetDataGatherer.cpp
@@ -60,11 +60,21 @@ bool FAssetDataGatherer::Init()
 
 uint32 FAssetDataGatherer::Run()
 {
+	int32 CacheSerializationVersion = CACHE_SERIALIZATION_VERSION;
+	
+	static const bool bUsingWorldAssets = FParse::Param(FCommandLine::Get(), TEXT("WorldAssets"));
+	if ( bUsingWorldAssets )
+	{
+		// Bump the serialization version to refresh the cache when switching between -WorldAssets and without.
+		// This is a temporary hack just while WorldAssets are under development
+		CacheSerializationVersion++;
+	}
+
 	if ( bLoadAndSaveCache )
 	{
 		// load the cached data
 		FNameTableArchiveReader CachedAssetDataReader;
-		if (CachedAssetDataReader.LoadFile(*CacheFilename, CACHE_SERIALIZATION_VERSION))
+		if (CachedAssetDataReader.LoadFile(*CacheFilename, CacheSerializationVersion))
 		{
 			SerializeCache(CachedAssetDataReader);
 		}
@@ -208,7 +218,7 @@ uint32 FAssetDataGatherer::Run()
 				// If we are caching discovered assets and this is the first time we had no work to do, save off the cache now in case the user terminates unexpectedly
 				if (bLoadAndSaveCache && !bSavedCacheAfterInitialDiscovery)
 				{
-					FNameTableArchiveWriter CachedAssetDataWriter(CACHE_SERIALIZATION_VERSION);
+					FNameTableArchiveWriter CachedAssetDataWriter(CacheSerializationVersion);
 					SerializeCache(CachedAssetDataWriter);
 					CachedAssetDataWriter.SaveToFile(*CacheFilename);
 
@@ -223,7 +233,7 @@ uint32 FAssetDataGatherer::Run()
 
 	if ( bLoadAndSaveCache )
 	{
-		FNameTableArchiveWriter CachedAssetData(CACHE_SERIALIZATION_VERSION);
+		FNameTableArchiveWriter CachedAssetData(CacheSerializationVersion);
 		SerializeCache(CachedAssetData);
 		CachedAssetData.SaveToFile(*CacheFilename);
 	}
diff --git a/Engine/Source/Runtime/AssetRegistry/Private/PackageReader.cpp b/Engine/Source/Runtime/AssetRegistry/Private/PackageReader.cpp
--- a/Engine/Source/Runtime/AssetRegistry/Private/PackageReader.cpp
+++ b/Engine/Source/Runtime/AssetRegistry/Private/PackageReader.cpp
@@ -135,6 +135,14 @@ bool FPackageReader::ReadAssetRegistryData (TArray<FBackgroundAssetData*>& Asset
 		return true;
 	}
 
+	// Worlds that were saved before they were marked public do not have asset data so we will synthesize it here to make sure we see all legacy umaps
+	if ( bUsingWorldAssets && bIsMapPackage && PackageFileSummary.GetFileVersionUE4() < VER_UE4_PUBLIC_WORLDS )
+	{
+		const FString AssetName = FPackageName::GetLongPackageAssetName(PackageName);
+		TMap<FString, FString> TagsAndValues;
+		AssetDataList.Add(new FBackgroundAssetData(PackageName, PackagePath, TEXT(""), AssetName, TEXT("World"), TagsAndValues, PackageFileSummary.ChunkIDs));
+	}
+
 	// Load the object count
 	int32 ObjectCount = 0;
 	*this << ObjectCount;
