UE4 Refactoring. Changed OVERRIDE and FINAL macros to keywords override and final.

diff --git a/Engine/Source/Runtime/AIModule/Classes/Navigation/CrowdFollowingComponent.h b/Engine/Source/Runtime/AIModule/Classes/Navigation/CrowdFollowingComponent.h
--- a/Engine/Source/Runtime/AIModule/Classes/Navigation/CrowdFollowingComponent.h
+++ b/Engine/Source/Runtime/AIModule/Classes/Navigation/CrowdFollowingComponent.h
@@ -25,22 +25,22 @@ class AIMODULE_API UCrowdFollowingComponent : public UPathFollowingComponent, pu
 	FVector CrowdAgentMoveDirection;
 
 	// ICrowdAgentInterface BEGIN
-	virtual FVector GetCrowdAgentLocation() const OVERRIDE;
-	virtual FVector GetCrowdAgentVelocity() const OVERRIDE;
-	virtual void GetCrowdAgentCollisions(float& CylinderRadius, float& CylinderHalfHeight) const OVERRIDE;
-	virtual float GetCrowdAgentMaxSpeed() const OVERRIDE;
+	virtual FVector GetCrowdAgentLocation() const override;
+	virtual FVector GetCrowdAgentVelocity() const override;
+	virtual void GetCrowdAgentCollisions(float& CylinderRadius, float& CylinderHalfHeight) const override;
+	virtual float GetCrowdAgentMaxSpeed() const override;
 	// ICrowdAgentInterface END
 
 	// PathFollowingComponent BEGIN
-	virtual void Initialize() OVERRIDE;
-	virtual void Cleanup() OVERRIDE;
-	virtual void AbortMove(const FString& Reason, FAIRequestID RequestID = FAIRequestID::CurrentRequest, bool bResetVelocity = true, bool bSilent = false) OVERRIDE;
-	virtual void PauseMove(FAIRequestID RequestID = FAIRequestID::CurrentRequest, bool bResetVelocity = true) OVERRIDE;
-	virtual void ResumeMove(FAIRequestID RequestID = FAIRequestID::CurrentRequest) OVERRIDE;
-	virtual FVector GetMoveFocus(bool bAllowStrafe) const OVERRIDE;
-	virtual void OnLanded() OVERRIDE;
-	virtual void OnPathFinished(EPathFollowingResult::Type Result) OVERRIDE;
-	virtual void OnPathUpdated() OVERRIDE;
+	virtual void Initialize() override;
+	virtual void Cleanup() override;
+	virtual void AbortMove(const FString& Reason, FAIRequestID RequestID = FAIRequestID::CurrentRequest, bool bResetVelocity = true, bool bSilent = false) override;
+	virtual void PauseMove(FAIRequestID RequestID = FAIRequestID::CurrentRequest, bool bResetVelocity = true) override;
+	virtual void ResumeMove(FAIRequestID RequestID = FAIRequestID::CurrentRequest) override;
+	virtual FVector GetMoveFocus(bool bAllowStrafe) const override;
+	virtual void OnLanded() override;
+	virtual void OnPathFinished(EPathFollowingResult::Type Result) override;
+	virtual void OnPathUpdated() override;
 	// PathFollowingComponent END
 
 	/** update params in crowd manager */
@@ -131,14 +131,14 @@ protected:
 	TEnumAsByte<ECrowdAvoidanceQuality::Type> AvoidanceQuality;
 
 	// PathFollowingComponent BEGIN
-	virtual int32 DetermineStartingPathPoint(const FNavigationPath* ConsideredPath) const OVERRIDE;
-	virtual void SetMoveSegment(int32 SegmentStartIndex) OVERRIDE;
-	virtual void UpdatePathSegment() OVERRIDE;
-	virtual void FollowPathSegment(float DeltaTime) OVERRIDE;
-	virtual bool ShouldCheckPathOnResume() const OVERRIDE;
-	virtual bool IsOnPath() const OVERRIDE;
-	virtual bool UpdateMovementComponent(bool bForce) OVERRIDE;
-	virtual void Reset() OVERRIDE;
+	virtual int32 DetermineStartingPathPoint(const FNavigationPath* ConsideredPath) const override;
+	virtual void SetMoveSegment(int32 SegmentStartIndex) override;
+	virtual void UpdatePathSegment() override;
+	virtual void FollowPathSegment(float DeltaTime) override;
+	virtual bool ShouldCheckPathOnResume() const override;
+	virtual bool IsOnPath() const override;
+	virtual bool UpdateMovementComponent(bool bForce) override;
+	virtual void Reset() override;
 	// PathFollowingComponent END
 
 	void SwitchToNextPathPart();
diff --git a/Engine/Source/Runtime/AIModule/Classes/Navigation/CrowdManager.h b/Engine/Source/Runtime/AIModule/Classes/Navigation/CrowdManager.h
--- a/Engine/Source/Runtime/AIModule/Classes/Navigation/CrowdManager.h
+++ b/Engine/Source/Runtime/AIModule/Classes/Navigation/CrowdManager.h
@@ -142,7 +142,7 @@ class AIMODULE_API UCrowdManager : public UObject
 	GENERATED_UCLASS_BODY()
 
 	virtual void Tick(float DeltaTime);
-	virtual void BeginDestroy() OVERRIDE;
+	virtual void BeginDestroy() override;
 
 	/** adds new agent to crowd */
 	void RegisterAgent(const class ICrowdAgentInterface* Agent);
@@ -197,7 +197,7 @@ class AIMODULE_API UCrowdManager : public UObject
 	void AdjustAgentPathStart(const UCrowdFollowingComponent* AgentComponent, const FNavMeshPath* Path, int32& PathStartIdx) const;
 
 #if WITH_EDITOR
-	virtual void PostEditChangeProperty(struct FPropertyChangedEvent& PropertyChangedEvent) OVERRIDE;
+	virtual void PostEditChangeProperty(struct FPropertyChangedEvent& PropertyChangedEvent) override;
 
 	void DebugTick() const;
 #endif
diff --git a/Engine/Source/Runtime/AIModule/Classes/Navigation/NavigationComponent.h b/Engine/Source/Runtime/AIModule/Classes/Navigation/NavigationComponent.h
--- a/Engine/Source/Runtime/AIModule/Classes/Navigation/NavigationComponent.h
+++ b/Engine/Source/Runtime/AIModule/Classes/Navigation/NavigationComponent.h
@@ -20,13 +20,13 @@ class AIMODULE_API UNavigationComponent : public UActorComponent, public INaviga
 	const class AActor* GoalActor;
 
 	// Begin UActorComponent Interface
-	virtual void InitializeComponent() OVERRIDE;
-	virtual void TickComponent(float DeltaTime, enum ELevelTick TickType, FActorComponentTickFunction *ThisTickFunction) OVERRIDE;
-	virtual void SetComponentTickEnabledAsync(bool bEnabled) OVERRIDE;
+	virtual void InitializeComponent() override;
+	virtual void TickComponent(float DeltaTime, enum ELevelTick TickType, FActorComponentTickFunction *ThisTickFunction) override;
+	virtual void SetComponentTickEnabledAsync(bool bEnabled) override;
 	// End UActorComponent Interface
 
 	// Begin INavigationPathGenerator Interface
-	virtual FNavPathSharedPtr GetGeneratedPath(class INavAgentInterface* Agent) OVERRIDE { return Path; }
+	virtual FNavPathSharedPtr GetGeneratedPath(class INavAgentInterface* Agent) override { return Path; }
 	// End INavigationPathGenerator Interface
 
 	/** updates cached pointers to relevant owner's components */
diff --git a/Engine/Source/Runtime/AIModule/Classes/Navigation/PathFollowingComponent.h b/Engine/Source/Runtime/AIModule/Classes/Navigation/PathFollowingComponent.h
--- a/Engine/Source/Runtime/AIModule/Classes/Navigation/PathFollowingComponent.h
+++ b/Engine/Source/Runtime/AIModule/Classes/Navigation/PathFollowingComponent.h
@@ -87,7 +87,7 @@ class AIMODULE_API UPathFollowingComponent : public UActorComponent, public IAIR
 	FMoveCompletedSignature OnMoveFinished;
 
 	// Begin UActorComponent Interface
-	virtual void TickComponent(float DeltaTime, enum ELevelTick TickType, FActorComponentTickFunction *ThisTickFunction) OVERRIDE;
+	virtual void TickComponent(float DeltaTime, enum ELevelTick TickType, FActorComponentTickFunction *ThisTickFunction) override;
 	// End UActorComponent Interface
 
 	/** initialize component to use */
@@ -234,10 +234,10 @@ class AIMODULE_API UPathFollowingComponent : public UActorComponent, public IAIR
 	virtual void FinishUsingCustomLink(class INavLinkCustomInterface* CustomNavLink);
 
 	// IAIResourceInterface begin
-	virtual void LockResource(EAILockSource::Type LockSource) OVERRIDE;
-	virtual void ClearResourceLock(EAILockSource::Type LockSource) OVERRIDE;
-	virtual void ForceUnlockResource() OVERRIDE;
-	virtual bool IsResourceLocked() const OVERRIDE;
+	virtual void LockResource(EAILockSource::Type LockSource) override;
+	virtual void ClearResourceLock(EAILockSource::Type LockSource) override;
+	virtual void ForceUnlockResource() override;
+	virtual bool IsResourceLocked() const override;
 	// IAIResourceInterface end
 
 protected:
diff --git a/Engine/Source/Runtime/AIModule/Classes/Perception/PawnSensingComponent.h b/Engine/Source/Runtime/AIModule/Classes/Perception/PawnSensingComponent.h
--- a/Engine/Source/Runtime/AIModule/Classes/Perception/PawnSensingComponent.h
+++ b/Engine/Source/Runtime/AIModule/Classes/Perception/PawnSensingComponent.h
@@ -117,7 +117,7 @@ public:
 	virtual bool CanHear(const FVector& NoiseLoc, float Loudness, bool bFailedLOS) const;
 
 	// Begin UActorComponent interface.
-	virtual void InitializeComponent() OVERRIDE;
+	virtual void InitializeComponent() override;
 	// End UActorComponent interface.
 
 	/** Get position where hearing/seeing occurs (i.e. ear/eye position).  If we ever need different positions for hearing/seeing, we'll deal with that then! */
diff --git a/Engine/Source/Runtime/AIModule/Private/AIModule.cpp b/Engine/Source/Runtime/AIModule/Private/AIModule.cpp
--- a/Engine/Source/Runtime/AIModule/Private/AIModule.cpp
+++ b/Engine/Source/Runtime/AIModule/Private/AIModule.cpp
@@ -9,10 +9,10 @@ DEFINE_LOG_CATEGORY_STATIC(LogAIModule, Log, All);
 class FAIModule : public IAIModule
 {
 	// Begin IModuleInterface
-	virtual void StartupModule() OVERRIDE;
-	virtual void ShutdownModule() OVERRIDE;
+	virtual void StartupModule() override;
+	virtual void ShutdownModule() override;
 	
-	virtual UAISystemBase* CreateAISystemInstance(UWorld* World) OVERRIDE;
+	virtual UAISystemBase* CreateAISystemInstance(UWorld* World) override;
 	// End IModuleInterface
 };
 
diff --git a/Engine/Source/Runtime/AIModule/Private/AISystemExec.cpp b/Engine/Source/Runtime/AIModule/Private/AISystemExec.cpp
--- a/Engine/Source/Runtime/AIModule/Private/AISystemExec.cpp
+++ b/Engine/Source/Runtime/AIModule/Private/AISystemExec.cpp
@@ -12,7 +12,7 @@ struct FAISystemExec : public FSelfRegisteringExec
 	}
 
 	// Begin FExec Interface
-	virtual bool Exec(UWorld* Inworld, const TCHAR* Cmd, FOutputDevice& Ar) OVERRIDE;
+	virtual bool Exec(UWorld* Inworld, const TCHAR* Cmd, FOutputDevice& Ar) override;
 	// End FExec Interface
 };
 
diff --git a/Engine/Source/Runtime/Advertising/Android/AndroidAdvertising/Public/AndroidAdvertising.h b/Engine/Source/Runtime/Advertising/Android/AndroidAdvertising/Public/AndroidAdvertising.h
--- a/Engine/Source/Runtime/Advertising/Android/AndroidAdvertising/Public/AndroidAdvertising.h
+++ b/Engine/Source/Runtime/Advertising/Android/AndroidAdvertising/Public/AndroidAdvertising.h
@@ -8,7 +8,7 @@
 class FAndroidAdvertisingProvider : public IAdvertisingProvider
 {
 public:
-	virtual void ShowAdBanner( bool bShowOnBottomOfScreen ) OVERRIDE;
-	virtual void HideAdBanner() OVERRIDE;
-	virtual void CloseAdBanner() OVERRIDE;
+	virtual void ShowAdBanner( bool bShowOnBottomOfScreen ) override;
+	virtual void HideAdBanner() override;
+	virtual void CloseAdBanner() override;
 };
diff --git a/Engine/Source/Runtime/Advertising/IOS/IOSAdvertising/Public/IOSAdvertising.h b/Engine/Source/Runtime/Advertising/IOS/IOSAdvertising/Public/IOSAdvertising.h
--- a/Engine/Source/Runtime/Advertising/IOS/IOSAdvertising/Public/IOSAdvertising.h
+++ b/Engine/Source/Runtime/Advertising/IOS/IOSAdvertising/Public/IOSAdvertising.h
@@ -8,7 +8,7 @@
 class FIOSAdvertisingProvider : public IAdvertisingProvider
 {
 public:
-	virtual void ShowAdBanner( bool bShowOnBottomOfScreen ) OVERRIDE;
-	virtual void HideAdBanner() OVERRIDE;
-	virtual void CloseAdBanner() OVERRIDE;
+	virtual void ShowAdBanner( bool bShowOnBottomOfScreen ) override;
+	virtual void HideAdBanner() override;
+	virtual void CloseAdBanner() override;
 };
diff --git a/Engine/Source/Runtime/Analytics/AnalyticsET/Private/AnalyticsET.cpp b/Engine/Source/Runtime/Analytics/AnalyticsET/Private/AnalyticsET.cpp
--- a/Engine/Source/Runtime/Analytics/AnalyticsET/Private/AnalyticsET.cpp
+++ b/Engine/Source/Runtime/Analytics/AnalyticsET/Private/AnalyticsET.cpp
@@ -24,22 +24,22 @@ public:
 
 	// FTickerObjectBase
 
-	bool Tick(float DeltaSeconds) OVERRIDE;
+	bool Tick(float DeltaSeconds) override;
 
 	// IAnalyticsProvider
 
-	virtual bool StartSession(const TArray<FAnalyticsEventAttribute>& Attributes) OVERRIDE;
-	virtual void EndSession() OVERRIDE;
+	virtual bool StartSession(const TArray<FAnalyticsEventAttribute>& Attributes) override;
+	virtual void EndSession() override;
 	/** ET PC implementation doesn't cache events */
-	virtual void FlushEvents() OVERRIDE;
+	virtual void FlushEvents() override;
 
-	virtual void SetUserID(const FString& InUserID) OVERRIDE;
-	virtual FString GetUserID() const OVERRIDE;
+	virtual void SetUserID(const FString& InUserID) override;
+	virtual FString GetUserID() const override;
 
-	virtual FString GetSessionID() const OVERRIDE;
-	virtual bool SetSessionID(const FString& InSessionID) OVERRIDE;
+	virtual FString GetSessionID() const override;
+	virtual bool SetSessionID(const FString& InSessionID) override;
 
-	virtual void RecordEvent(const FString& EventName, const TArray<FAnalyticsEventAttribute>& Attributes) OVERRIDE;
+	virtual void RecordEvent(const FString& EventName, const TArray<FAnalyticsEventAttribute>& Attributes) override;
 
 	virtual ~FAnalyticsProviderET();
 
diff --git a/Engine/Source/Runtime/Analytics/AnalyticsMulticast/Private/AnalyticsMulticast.cpp b/Engine/Source/Runtime/Analytics/AnalyticsMulticast/Private/AnalyticsMulticast.cpp
--- a/Engine/Source/Runtime/Analytics/AnalyticsMulticast/Private/AnalyticsMulticast.cpp
+++ b/Engine/Source/Runtime/Analytics/AnalyticsMulticast/Private/AnalyticsMulticast.cpp
@@ -18,21 +18,21 @@ class FAnalyticsProviderMulticast : public IAnalyticsProvider
 public:
 	FAnalyticsProviderMulticast(const FAnalyticsMulticast::Config& ConfigValues, const FAnalytics::FProviderConfigurationDelegate& GetConfigValue);
 
-	virtual bool StartSession(const TArray<FAnalyticsEventAttribute>& Attributes) OVERRIDE;
-	virtual void EndSession() OVERRIDE;
-	virtual void FlushEvents() OVERRIDE;
+	virtual bool StartSession(const TArray<FAnalyticsEventAttribute>& Attributes) override;
+	virtual void EndSession() override;
+	virtual void FlushEvents() override;
 
-	virtual void SetUserID(const FString& InUserID) OVERRIDE;
-	virtual FString GetUserID() const OVERRIDE;
+	virtual void SetUserID(const FString& InUserID) override;
+	virtual FString GetUserID() const override;
 
-	virtual FString GetSessionID() const OVERRIDE;
-	virtual bool SetSessionID(const FString& InSessionID) OVERRIDE;
+	virtual FString GetSessionID() const override;
+	virtual bool SetSessionID(const FString& InSessionID) override;
 
-	virtual void RecordEvent(const FString& EventName, const TArray<FAnalyticsEventAttribute>& Attributes) OVERRIDE;
-	virtual void RecordUserAttribute(const TArray<FAnalyticsEventAttribute>& Attributes) OVERRIDE;
-	virtual void RecordItemPurchase(const FString& ItemId, const FString& Currency, int PerItemCost, int ItemQuantity) OVERRIDE;
-	virtual void RecordCurrencyPurchase(const FString& GameCurrencyType, int GameCurrencyAmount, const FString& RealCurrencyType, float RealMoneyCost, const FString& PaymentProvider) OVERRIDE;
-	virtual void RecordCurrencyGiven(const FString& GameCurrencyType, int GameCurrencyAmount) OVERRIDE;
+	virtual void RecordEvent(const FString& EventName, const TArray<FAnalyticsEventAttribute>& Attributes) override;
+	virtual void RecordUserAttribute(const TArray<FAnalyticsEventAttribute>& Attributes) override;
+	virtual void RecordItemPurchase(const FString& ItemId, const FString& Currency, int PerItemCost, int ItemQuantity) override;
+	virtual void RecordCurrencyPurchase(const FString& GameCurrencyType, int GameCurrencyAmount, const FString& RealCurrencyType, float RealMoneyCost, const FString& PaymentProvider) override;
+	virtual void RecordCurrencyGiven(const FString& GameCurrencyType, int GameCurrencyAmount) override;
 
 	virtual ~FAnalyticsProviderMulticast();
 
diff --git a/Engine/Source/Runtime/Analytics/AnalyticsSwrve/Private/AnalyticsSwrve.cpp b/Engine/Source/Runtime/Analytics/AnalyticsSwrve/Private/AnalyticsSwrve.cpp
--- a/Engine/Source/Runtime/Analytics/AnalyticsSwrve/Private/AnalyticsSwrve.cpp
+++ b/Engine/Source/Runtime/Analytics/AnalyticsSwrve/Private/AnalyticsSwrve.cpp
@@ -21,23 +21,23 @@ class FAnalyticsProviderSwrve : public IAnalyticsProvider
 public:
 	FAnalyticsProviderSwrve(const FAnalyticsSwrve::Config& ConfigValues);
 
-	virtual bool StartSession(const TArray<FAnalyticsEventAttribute>& Attributes) OVERRIDE;
-	virtual void EndSession() OVERRIDE;
+	virtual bool StartSession(const TArray<FAnalyticsEventAttribute>& Attributes) override;
+	virtual void EndSession() override;
 	/** Swrve PC implementation doesn't cache events */
-	virtual void FlushEvents() OVERRIDE {}
+	virtual void FlushEvents() override {}
 
-	virtual void SetUserID(const FString& InUserID) OVERRIDE;
-	virtual FString GetUserID() const OVERRIDE;
+	virtual void SetUserID(const FString& InUserID) override;
+	virtual FString GetUserID() const override;
 
-	virtual FString GetSessionID() const OVERRIDE;
-	virtual bool SetSessionID(const FString& InSessionID) OVERRIDE;
+	virtual FString GetSessionID() const override;
+	virtual bool SetSessionID(const FString& InSessionID) override;
 
-	virtual void RecordEvent(const FString& EventName, const TArray<FAnalyticsEventAttribute>& Attributes) OVERRIDE;
-	virtual void RecordUserAttribute(const TArray<FAnalyticsEventAttribute>& Attributes) OVERRIDE;
+	virtual void RecordEvent(const FString& EventName, const TArray<FAnalyticsEventAttribute>& Attributes) override;
+	virtual void RecordUserAttribute(const TArray<FAnalyticsEventAttribute>& Attributes) override;
 
-	virtual void RecordItemPurchase(const FString& ItemId, const FString& Currency, int PerItemCost, int ItemQuantity) OVERRIDE;
-	virtual void RecordCurrencyPurchase(const FString& GameCurrencyType, int GameCurrencyAmount, const FString& RealCurrencyType, float RealMoneyCost, const FString& PaymentProvider) OVERRIDE;
-	virtual void RecordCurrencyGiven(const FString& GameCurrencyType, int GameCurrencyAmount) OVERRIDE;
+	virtual void RecordItemPurchase(const FString& ItemId, const FString& Currency, int PerItemCost, int ItemQuantity) override;
+	virtual void RecordCurrencyPurchase(const FString& GameCurrencyType, int GameCurrencyAmount, const FString& RealCurrencyType, float RealMoneyCost, const FString& PaymentProvider) override;
+	virtual void RecordCurrencyGiven(const FString& GameCurrencyType, int GameCurrencyAmount) override;
 
 	virtual ~FAnalyticsProviderSwrve();
 
diff --git a/Engine/Source/Runtime/Android/AndroidAudio/Private/AndroidAudioDevice.cpp b/Engine/Source/Runtime/Android/AndroidAudio/Private/AndroidAudioDevice.cpp
--- a/Engine/Source/Runtime/Android/AndroidAudio/Private/AndroidAudioDevice.cpp
+++ b/Engine/Source/Runtime/Android/AndroidAudio/Private/AndroidAudioDevice.cpp
@@ -18,7 +18,7 @@ class FSLESAudioDeviceModule : public IAudioDeviceModule
 public:
 
 	/** Creates a new instance of the audio device implemented by the module. */
-	virtual FAudioDevice* CreateAudioDevice() OVERRIDE
+	virtual FAudioDevice* CreateAudioDevice() override
 	{
 		return new FSLESAudioDevice;
 	}
diff --git a/Engine/Source/Runtime/Android/AndroidAudio/Public/AndroidAudioDevice.h b/Engine/Source/Runtime/Android/AndroidAudio/Public/AndroidAudioDevice.h
--- a/Engine/Source/Runtime/Android/AndroidAudio/Public/AndroidAudioDevice.h
+++ b/Engine/Source/Runtime/Android/AndroidAudio/Public/AndroidAudioDevice.h
@@ -235,23 +235,23 @@ public:
 	FSLESAudioDevice() {} 
 	virtual ~FSLESAudioDevice() {} 
 
-	virtual FName GetRuntimeFormat(USoundWave* SoundWave) OVERRIDE
+	virtual FName GetRuntimeFormat(USoundWave* SoundWave) override
 	{
 		static FName NAME_OGG(TEXT("OGG"));
 		return NAME_OGG;
 	}
 
-	virtual bool HasCompressedAudioInfoClass(USoundWave* SoundWave) OVERRIDE;
+	virtual bool HasCompressedAudioInfoClass(USoundWave* SoundWave) override;
 
-	virtual bool SupportsRealtimeDecompression() const OVERRIDE
+	virtual bool SupportsRealtimeDecompression() const override
 	{
 		return true;
 	}
 
-	virtual class ICompressedAudioInfo* CreateCompressedAudioInfo(USoundWave* SoundWave) OVERRIDE;
+	virtual class ICompressedAudioInfo* CreateCompressedAudioInfo(USoundWave* SoundWave) override;
 
 	/** Starts up any platform specific hardware/APIs */
-	virtual bool InitializeHardware() OVERRIDE;
+	virtual bool InitializeHardware() override;
  
 	/**
 	 * Update the audio device and calculates the cached inverse transform later
@@ -276,7 +276,7 @@ public:
 
 protected:
 
-	virtual FSoundSource* CreateSoundSource() OVERRIDE; 
+	virtual FSoundSource* CreateSoundSource() override; 
 
 
 	/**
diff --git a/Engine/Source/Runtime/AssetRegistry/Private/AssetDataGatherer.h b/Engine/Source/Runtime/AssetRegistry/Private/AssetDataGatherer.h
--- a/Engine/Source/Runtime/AssetRegistry/Private/AssetDataGatherer.h
+++ b/Engine/Source/Runtime/AssetRegistry/Private/AssetDataGatherer.h
@@ -13,10 +13,10 @@ public:
 	virtual ~FAssetDataGatherer();
 
 	// FRunnable implementation
-	virtual bool Init() OVERRIDE;
-	virtual uint32 Run() OVERRIDE;
-	virtual void Stop() OVERRIDE;
-	virtual void Exit() OVERRIDE;
+	virtual bool Init() override;
+	virtual uint32 Run() override;
+	virtual void Stop() override;
+	virtual void Exit() override;
 
 	/** Signals to end the thread and waits for it to close before returning */
 	void EnsureCompletion();
diff --git a/Engine/Source/Runtime/AssetRegistry/Private/AssetRegistry.h b/Engine/Source/Runtime/AssetRegistry/Private/AssetRegistry.h
--- a/Engine/Source/Runtime/AssetRegistry/Private/AssetRegistry.h
+++ b/Engine/Source/Runtime/AssetRegistry/Private/AssetRegistry.h
@@ -14,65 +14,65 @@ public:
 	virtual ~FAssetRegistry();
 
 	// IAssetRegistry implementation
-	virtual bool GetAssetsByPackageName(FName PackageName, TArray<FAssetData>& OutAssetData) const OVERRIDE;
-	virtual bool GetAssetsByPath(FName PackagePath, TArray<FAssetData>& OutAssetData, bool bRecursive = false) const OVERRIDE;
-	virtual bool GetAssetsByClass(FName ClassName, TArray<FAssetData>& OutAssetData, bool bSearchSubClasses = false) const OVERRIDE;
-	virtual bool GetAssetsByTagValues(const TMultiMap<FName, FString>& AssetTagsAndValues, TArray<FAssetData>& OutAssetData) const OVERRIDE;
-	virtual bool GetAssets(const FARFilter& Filter, TArray<FAssetData>& OutAssetData) const OVERRIDE;
-	virtual FAssetData GetAssetByObjectPath( const FName ObjectPath ) const OVERRIDE;
-	virtual bool GetAllAssets(TArray<FAssetData>& OutAssetData) const OVERRIDE;
-	virtual bool GetDependencies(FName PackageName, TArray<FName>& OutDependencies) const OVERRIDE;
-	virtual bool GetReferencers(FName PackageName, TArray<FName>& OutReferencers) const OVERRIDE;
-	virtual bool GetAncestorClassNames(FName ClassName, TArray<FName>& OutAncestorClassNames) const OVERRIDE;
-	virtual void GetDerivedClassNames(const TArray<FName>& ClassNames, const TSet<FName>& ExcludedClassNames, TSet<FName>& OutDerivedClassNames) const OVERRIDE;
-	virtual void GetAllCachedPaths(TArray<FString>& OutPathList) const OVERRIDE;
-	virtual void GetSubPaths(const FString& InBasePath, TArray<FString>& OutPathList, bool bInRecurse) const OVERRIDE;
-	virtual void RunAssetsThroughFilter (TArray<FAssetData>& AssetDataList, const FARFilter& Filter) const OVERRIDE;
-	virtual EAssetAvailability::Type GetAssetAvailability(const FAssetData& AssetData) const OVERRIDE;	
-	virtual float GetAssetAvailabilityProgress(const FAssetData& AssetData, EAssetAvailabilityProgressReportingType::Type ReportType) const OVERRIDE;
-	virtual bool GetAssetAvailabilityProgressTypeSupported(EAssetAvailabilityProgressReportingType::Type ReportType) const OVERRIDE;
-	virtual void PrioritizeAssetInstall(const FAssetData& AssetData) const OVERRIDE;
-	virtual bool AddPath(const FString& PathToAdd) OVERRIDE;
-	virtual bool RemovePath(const FString& PathToRemove) OVERRIDE;
-	virtual void SearchAllAssets(bool bSynchronousSearch) OVERRIDE;
-	virtual void ScanPathsSynchronous(const TArray<FString>& InPaths, bool bForceRescan = false) OVERRIDE;
-	virtual void Serialize(FArchive& Ar) OVERRIDE;
-	virtual void SaveRegistryData(FArchive& Ar, TMap<FName, FAssetData*>& Data, int32 AssetCount) OVERRIDE;
+	virtual bool GetAssetsByPackageName(FName PackageName, TArray<FAssetData>& OutAssetData) const override;
+	virtual bool GetAssetsByPath(FName PackagePath, TArray<FAssetData>& OutAssetData, bool bRecursive = false) const override;
+	virtual bool GetAssetsByClass(FName ClassName, TArray<FAssetData>& OutAssetData, bool bSearchSubClasses = false) const override;
+	virtual bool GetAssetsByTagValues(const TMultiMap<FName, FString>& AssetTagsAndValues, TArray<FAssetData>& OutAssetData) const override;
+	virtual bool GetAssets(const FARFilter& Filter, TArray<FAssetData>& OutAssetData) const override;
+	virtual FAssetData GetAssetByObjectPath( const FName ObjectPath ) const override;
+	virtual bool GetAllAssets(TArray<FAssetData>& OutAssetData) const override;
+	virtual bool GetDependencies(FName PackageName, TArray<FName>& OutDependencies) const override;
+	virtual bool GetReferencers(FName PackageName, TArray<FName>& OutReferencers) const override;
+	virtual bool GetAncestorClassNames(FName ClassName, TArray<FName>& OutAncestorClassNames) const override;
+	virtual void GetDerivedClassNames(const TArray<FName>& ClassNames, const TSet<FName>& ExcludedClassNames, TSet<FName>& OutDerivedClassNames) const override;
+	virtual void GetAllCachedPaths(TArray<FString>& OutPathList) const override;
+	virtual void GetSubPaths(const FString& InBasePath, TArray<FString>& OutPathList, bool bInRecurse) const override;
+	virtual void RunAssetsThroughFilter (TArray<FAssetData>& AssetDataList, const FARFilter& Filter) const override;
+	virtual EAssetAvailability::Type GetAssetAvailability(const FAssetData& AssetData) const override;	
+	virtual float GetAssetAvailabilityProgress(const FAssetData& AssetData, EAssetAvailabilityProgressReportingType::Type ReportType) const override;
+	virtual bool GetAssetAvailabilityProgressTypeSupported(EAssetAvailabilityProgressReportingType::Type ReportType) const override;
+	virtual void PrioritizeAssetInstall(const FAssetData& AssetData) const override;
+	virtual bool AddPath(const FString& PathToAdd) override;
+	virtual bool RemovePath(const FString& PathToRemove) override;
+	virtual void SearchAllAssets(bool bSynchronousSearch) override;
+	virtual void ScanPathsSynchronous(const TArray<FString>& InPaths, bool bForceRescan = false) override;
+	virtual void Serialize(FArchive& Ar) override;
+	virtual void SaveRegistryData(FArchive& Ar, TMap<FName, FAssetData*>& Data, int32 AssetCount) override;
 
 	DECLARE_DERIVED_EVENT( FAssetRegistry, IAssetRegistry::FPathAddedEvent, FPathAddedEvent);
-	virtual FPathAddedEvent& OnPathAdded() OVERRIDE { return PathAddedEvent; }
+	virtual FPathAddedEvent& OnPathAdded() override { return PathAddedEvent; }
 
 	DECLARE_DERIVED_EVENT( FAssetRegistry, IAssetRegistry::FPathRemovedEvent, FPathRemovedEvent);
-	virtual FPathRemovedEvent& OnPathRemoved() OVERRIDE { return PathRemovedEvent; }
+	virtual FPathRemovedEvent& OnPathRemoved() override { return PathRemovedEvent; }
 
-	virtual void AssetCreated(UObject* NewAsset) OVERRIDE;
-	virtual void AssetDeleted(UObject* DeletedAsset) OVERRIDE;
-	virtual void AssetRenamed(const UObject* RenamedAsset, const FString& OldObjectPath) OVERRIDE;
+	virtual void AssetCreated(UObject* NewAsset) override;
+	virtual void AssetDeleted(UObject* DeletedAsset) override;
+	virtual void AssetRenamed(const UObject* RenamedAsset, const FString& OldObjectPath) override;
 
 	DECLARE_DERIVED_EVENT( FAssetRegistry, IAssetRegistry::FAssetAddedEvent, FAssetAddedEvent);
-	virtual FAssetAddedEvent& OnAssetAdded() OVERRIDE { return AssetAddedEvent; }
+	virtual FAssetAddedEvent& OnAssetAdded() override { return AssetAddedEvent; }
 
 	DECLARE_DERIVED_EVENT( FAssetRegistry, IAssetRegistry::FAssetRemovedEvent, FAssetRemovedEvent);
-	virtual FAssetRemovedEvent& OnAssetRemoved() OVERRIDE { return AssetRemovedEvent; }
+	virtual FAssetRemovedEvent& OnAssetRemoved() override { return AssetRemovedEvent; }
 
 	DECLARE_DERIVED_EVENT( FAssetRegistry, IAssetRegistry::FAssetRenamedEvent, FAssetRenamedEvent);
-	virtual FAssetRenamedEvent& OnAssetRenamed() OVERRIDE { return AssetRenamedEvent; }
+	virtual FAssetRenamedEvent& OnAssetRenamed() override { return AssetRenamedEvent; }
 
 	DECLARE_DERIVED_EVENT( FAssetRegistry, IAssetRegistry::FInMemoryAssetCreatedEvent, FInMemoryAssetCreatedEvent );
-	virtual FInMemoryAssetCreatedEvent& OnInMemoryAssetCreated() OVERRIDE { return InMemoryAssetCreatedEvent; }
+	virtual FInMemoryAssetCreatedEvent& OnInMemoryAssetCreated() override { return InMemoryAssetCreatedEvent; }
 
 	DECLARE_DERIVED_EVENT( FAssetRegistry, IAssetRegistry::FInMemoryAssetDeletedEvent, FInMemoryAssetDeletedEvent );
-	virtual FInMemoryAssetDeletedEvent& OnInMemoryAssetDeleted() OVERRIDE { return InMemoryAssetDeletedEvent; }
+	virtual FInMemoryAssetDeletedEvent& OnInMemoryAssetDeleted() override { return InMemoryAssetDeletedEvent; }
 
 	DECLARE_DERIVED_EVENT( FAssetRegistry, IAssetRegistry::FFilesLoadedEvent, FFilesLoadedEvent );
-	virtual FFilesLoadedEvent& OnFilesLoaded() OVERRIDE { return FileLoadedEvent; }
+	virtual FFilesLoadedEvent& OnFilesLoaded() override { return FileLoadedEvent; }
 
 	DECLARE_DERIVED_EVENT( FAssetRegistry, IAssetRegistry::FFileLoadProgressUpdatedEvent, FFileLoadProgressUpdatedEvent );
-	virtual FFileLoadProgressUpdatedEvent& OnFileLoadProgressUpdated() OVERRIDE { return FileLoadProgressUpdatedEvent; }
+	virtual FFileLoadProgressUpdatedEvent& OnFileLoadProgressUpdated() override { return FileLoadProgressUpdatedEvent; }
 
-	virtual bool IsLoadingAssets() const OVERRIDE;
+	virtual bool IsLoadingAssets() const override;
 
-	virtual void Tick (float DeltaTime) OVERRIDE;
+	virtual void Tick (float DeltaTime) override;
 
 private:
 
diff --git a/Engine/Source/Runtime/AutomationWorker/Private/AutomationWorkerModule.cpp b/Engine/Source/Runtime/AutomationWorker/Private/AutomationWorkerModule.cpp
--- a/Engine/Source/Runtime/AutomationWorker/Private/AutomationWorkerModule.cpp
+++ b/Engine/Source/Runtime/AutomationWorker/Private/AutomationWorkerModule.cpp
@@ -556,7 +556,7 @@ static class FAutomationTestCmd : private FSelfRegisteringExec
 {
 public:
 	/** Console commands, see embeded usage statement **/
-	virtual bool Exec( UWorld*, const TCHAR* Cmd, FOutputDevice& Ar) OVERRIDE
+	virtual bool Exec( UWorld*, const TCHAR* Cmd, FOutputDevice& Ar) override
 	{
 		return DirectAutomationCommand(Cmd, &Ar);
 	}
diff --git a/Engine/Source/Runtime/AutomationWorker/Private/AutomationWorkerModule.h b/Engine/Source/Runtime/AutomationWorker/Private/AutomationWorkerModule.h
--- a/Engine/Source/Runtime/AutomationWorker/Private/AutomationWorkerModule.h
+++ b/Engine/Source/Runtime/AutomationWorker/Private/AutomationWorkerModule.h
@@ -13,9 +13,9 @@ public:
 
 	// Begin IModuleInterface interface
 
-	virtual void StartupModule() OVERRIDE;
+	virtual void StartupModule() override;
 
-	virtual bool SupportsDynamicReloading() OVERRIDE;
+	virtual bool SupportsDynamicReloading() override;
 
 	// End IModuleInterface interface
 
@@ -23,9 +23,9 @@ public:
 
 	// Begin IAutomationWorkerModule interface
 
-	virtual void Tick() OVERRIDE;
+	virtual void Tick() override;
 
-	virtual void RunTest( const FString& InTestToRun, const int32 InRoleIndex, FStopTestEvent const& InStopTestEvent ) OVERRIDE;
+	virtual void RunTest( const FString& InTestToRun, const int32 InRoleIndex, FStopTestEvent const& InStopTestEvent ) override;
 
 	// End IAutomationWorkerModule interface
 
diff --git a/Engine/Source/Runtime/Core/Private/Android/AndroidFile.cpp b/Engine/Source/Runtime/Core/Private/Android/AndroidFile.cpp
--- a/Engine/Source/Runtime/Core/Private/Android/AndroidFile.cpp
+++ b/Engine/Source/Runtime/Core/Private/Android/AndroidFile.cpp
@@ -64,27 +64,27 @@ public:
 		FileHandle = -1;
 	}
 
-	virtual int64 Tell() OVERRIDE
+	virtual int64 Tell() override
 	{
 		check(IsValid());
 		return lseek64(FileHandle, 0, SEEK_CUR);
 	}
 
-	virtual bool Seek(int64 NewPosition) OVERRIDE
+	virtual bool Seek(int64 NewPosition) override
 	{
 		check(IsValid());
 		check(NewPosition >= 0);
 		return lseek64(FileHandle, NewPosition, SEEK_SET) != -1;
 	}
 
-	virtual bool SeekFromEnd(int64 NewPositionRelativeToEnd = 0) OVERRIDE
+	virtual bool SeekFromEnd(int64 NewPositionRelativeToEnd = 0) override
 	{
 		check(IsValid());
 		check(NewPositionRelativeToEnd <= 0);
 		return lseek64(FileHandle, NewPositionRelativeToEnd, SEEK_END) != -1;
 	}
 
-	virtual bool Read(uint8* Destination, int64 BytesToRead) OVERRIDE
+	virtual bool Read(uint8* Destination, int64 BytesToRead) override
 	{
 		check(IsValid());
 		while (BytesToRead)
@@ -102,7 +102,7 @@ public:
 		return true;
 	}
 
-	virtual bool Write(const uint8* Source, int64 BytesToWrite) OVERRIDE
+	virtual bool Write(const uint8* Source, int64 BytesToWrite) override
 	{
 		check(IsValid());
 		while (BytesToWrite)
@@ -151,14 +151,14 @@ public:
 		AAsset_close(File);
 	}
 
-	virtual int64 Tell() OVERRIDE
+	virtual int64 Tell() override
 	{
 		check(IsValid());
 		int64 pos = lseek(FileHandle, 0, SEEK_CUR);
 		return pos - Start; // We are treating 'tell' as a virtual location from file Start
 	}
 
-	virtual bool Seek(int64 NewPosition) OVERRIDE
+	virtual bool Seek(int64 NewPosition) override
 	{
 		check(IsValid());
 		// we need to offset all positions by the Start offset
@@ -168,7 +168,7 @@ public:
 		return lseek(FileHandle, NewPosition, SEEK_SET) != -1;
 	}
 
-    virtual bool SeekFromEnd(int64 NewPositionRelativeToEnd = 0) OVERRIDE
+    virtual bool SeekFromEnd(int64 NewPositionRelativeToEnd = 0) override
 	{
 		check(IsValid());
 		check(NewPositionRelativeToEnd <= 0);
@@ -177,7 +177,7 @@ public:
 		return lseek(FileHandle, position, SEEK_SET) != -1;
 	}
 
-	virtual bool Read(uint8* Destination, int64 BytesToRead) OVERRIDE
+	virtual bool Read(uint8* Destination, int64 BytesToRead) override
 	{
 		check(IsValid());
 		while (BytesToRead)
@@ -195,7 +195,7 @@ public:
 		return true;
 	}
 
-	virtual bool Write(const uint8* Source, int64 BytesToWrite) OVERRIDE
+	virtual bool Write(const uint8* Source, int64 BytesToWrite) override
 	{
 		check(IsValid());
 		while (BytesToWrite)
@@ -252,7 +252,7 @@ public:
 		AssetMgr = AndroidThunkCpp_GetAssetManager();
 	}
 
-	virtual bool FileExists(const TCHAR* Filename) OVERRIDE
+	virtual bool FileExists(const TCHAR* Filename) override
 	{
 		struct stat FileInfo;
 		FString NormalizedFilename = NormalizeFilename(Filename);
@@ -271,7 +271,7 @@ public:
 		return S_ISREG(FileInfo.st_mode);
 	}
 
-	virtual int64 FileSize(const TCHAR* Filename) OVERRIDE
+	virtual int64 FileSize(const TCHAR* Filename) override
 	{
 		struct stat FileInfo;
 		FileInfo.st_size = -1;
@@ -303,7 +303,7 @@ public:
 		return FileInfo.st_size;
 	}
 
-	virtual bool DeleteFile(const TCHAR* Filename) OVERRIDE
+	virtual bool DeleteFile(const TCHAR* Filename) override
 	{
 		// only delete from write path
 		FString AndroidFilename = ConvertToAndroidPath(NormalizeFilename(Filename), true);
@@ -314,7 +314,7 @@ public:
 		return unlink(TCHAR_TO_UTF8(*AndroidFilename)) == 0;
 	}
 
-	virtual bool IsReadOnly(const TCHAR* Filename) OVERRIDE
+	virtual bool IsReadOnly(const TCHAR* Filename) override
 	{
 		FString NormalizedFilename = NormalizeFilename(Filename);
 		FString Filepath = ConvertToAndroidPath(NormalizedFilename, false);
@@ -342,7 +342,7 @@ public:
 		return false;
 	}
 
-	virtual bool MoveFile(const TCHAR* To, const TCHAR* From) OVERRIDE
+	virtual bool MoveFile(const TCHAR* To, const TCHAR* From) override
 	{
 		// move to the write path
 		FString ToAndroidFilename = ConvertToAndroidPath(NormalizeFilename(To), true);
@@ -354,7 +354,7 @@ public:
 		return rename(TCHAR_TO_UTF8(*FromAndroidFilename), TCHAR_TO_UTF8(*ToAndroidFilename)) != -1;
 	}
 
-	virtual bool SetReadOnly(const TCHAR* Filename, bool bNewReadOnlyValue) OVERRIDE
+	virtual bool SetReadOnly(const TCHAR* Filename, bool bNewReadOnlyValue) override
 	{
 		struct stat FileInfo;
 		FString AndroidFilename = ConvertToAndroidPath(NormalizeFilename(Filename), false);
@@ -378,7 +378,7 @@ public:
 		return false;
 	}
 
-	virtual FDateTime GetTimeStamp(const TCHAR* Filename) OVERRIDE
+	virtual FDateTime GetTimeStamp(const TCHAR* Filename) override
 	{
 		// get file times
 		struct stat FileInfo;
@@ -405,7 +405,7 @@ public:
 		return AndroidEpoch + TimeSinceEpoch;
 	}
 
-	virtual void SetTimeStamp(const TCHAR* Filename, const FDateTime DateTime) OVERRIDE
+	virtual void SetTimeStamp(const TCHAR* Filename, const FDateTime DateTime) override
 	{
 		// get file times
 		struct stat FileInfo;
@@ -426,7 +426,7 @@ public:
 		utime(TCHAR_TO_UTF8(*AndroidFilename), &Times);
 	}
 
-	virtual FDateTime GetAccessTimeStamp(const TCHAR* Filename) OVERRIDE
+	virtual FDateTime GetAccessTimeStamp(const TCHAR* Filename) override
 	{
 		// get file times
 		struct stat FileInfo;
@@ -452,7 +452,7 @@ public:
 		return AndroidEpoch + TimeSinceEpoch;
 	}
 
-	virtual IFileHandle* OpenRead(const TCHAR* Filename) OVERRIDE
+	virtual IFileHandle* OpenRead(const TCHAR* Filename) override
 	{
 		FString NormalizedFilename = NormalizeFilename(Filename);
 		// check the read path
@@ -481,7 +481,7 @@ public:
 		return NULL;
 	}
 
-	virtual IFileHandle* OpenWrite(const TCHAR* Filename, bool bAppend, bool bAllowRead) OVERRIDE
+	virtual IFileHandle* OpenWrite(const TCHAR* Filename, bool bAppend, bool bAllowRead) override
 	{
 		FString AndroidFilename = ConvertToAndroidPath(NormalizeFilename(Filename), true);
 		if (AndroidFilename.Contains(TEXT(".pak")) && GOBBinAPK)
@@ -520,7 +520,7 @@ public:
 		return NULL;
 	}
 
-	virtual bool DirectoryExists(const TCHAR* Directory) OVERRIDE
+	virtual bool DirectoryExists(const TCHAR* Directory) override
 	{
 		struct stat FileInfo;
 		FString NormalizedDirectory = NormalizeFilename(Directory);
@@ -536,19 +536,19 @@ public:
 		return S_ISDIR(FileInfo.st_mode);
 	}
 
-	virtual bool CreateDirectory(const TCHAR* Directory) OVERRIDE
+	virtual bool CreateDirectory(const TCHAR* Directory) override
 	{
 		FString AndroidDirectory = ConvertToAndroidPath(NormalizeFilename(Directory), true);
 		return mkdir(TCHAR_TO_UTF8(*AndroidDirectory), 0766) || (errno == EEXIST);
 	}
 
-	virtual bool DeleteDirectory(const TCHAR* Directory) OVERRIDE
+	virtual bool DeleteDirectory(const TCHAR* Directory) override
 	{
 		FString AndroidDirectory = ConvertToAndroidPath(NormalizeFilename(Directory), true);
 		return rmdir(TCHAR_TO_UTF8(*AndroidDirectory));
 	}
 
-	bool IterateDirectory(const TCHAR* Directory, FDirectoryVisitor& Visitor) OVERRIDE
+	bool IterateDirectory(const TCHAR* Directory, FDirectoryVisitor& Visitor) override
 	{
 		// we haven't hijacked a directory yet, and this one ends in Paks, which means the engine is looking for .pak files
 		if (!bLookedForObbs && FString(Directory).EndsWith("/Paks/"))
diff --git a/Engine/Source/Runtime/Core/Private/Android/AndroidPlatformOutputDevicesPrivate.h b/Engine/Source/Runtime/Core/Private/Android/AndroidPlatformOutputDevicesPrivate.h
--- a/Engine/Source/Runtime/Core/Private/Android/AndroidPlatformOutputDevicesPrivate.h
+++ b/Engine/Source/Runtime/Core/Private/Android/AndroidPlatformOutputDevicesPrivate.h
@@ -22,7 +22,7 @@ public:
 	 * @param	Data	Text to log
 	 * @param	Event	Event name used for suppression purposes
 	 */
-	virtual void Serialize( const TCHAR* Msg, ELogVerbosity::Type Verbosity, const class FName& Category ) OVERRIDE;
+	virtual void Serialize( const TCHAR* Msg, ELogVerbosity::Type Verbosity, const class FName& Category ) override;
 };
 
 
@@ -38,7 +38,7 @@ public:
 	 * @param	Data	Text to log
 	 * @param	Event	Event name used for suppression purposes
 	 */
-	virtual void Serialize( const TCHAR* Msg, ELogVerbosity::Type Verbosity, const class FName& Category ) OVERRIDE;
+	virtual void Serialize( const TCHAR* Msg, ELogVerbosity::Type Verbosity, const class FName& Category ) override;
 
 	/**
 	 * Error handling function that is being called from within the system wide global
diff --git a/Engine/Source/Runtime/Core/Private/Apple/ApplePlatformFile.cpp b/Engine/Source/Runtime/Core/Private/Apple/ApplePlatformFile.cpp
--- a/Engine/Source/Runtime/Core/Private/Apple/ApplePlatformFile.cpp
+++ b/Engine/Source/Runtime/Core/Private/Apple/ApplePlatformFile.cpp
@@ -63,7 +63,7 @@ public:
 		}
 		FileHandle = -1;
 	}
-	virtual int64 Tell() OVERRIDE
+	virtual int64 Tell() override
 	{
 #if MANAGE_FILE_HANDLES
 		if( IsManaged() )
@@ -77,7 +77,7 @@ public:
 			return lseek(FileHandle, 0, SEEK_CUR);
 		}
 	}
-	virtual bool Seek(int64 NewPosition) OVERRIDE
+	virtual bool Seek(int64 NewPosition) override
 	{
 		check(NewPosition >= 0);
 
@@ -94,7 +94,7 @@ public:
 			return lseek(FileHandle, NewPosition, SEEK_SET) != -1;
 		}
 	}
-	virtual bool SeekFromEnd(int64 NewPositionRelativeToEnd = 0) OVERRIDE
+	virtual bool SeekFromEnd(int64 NewPositionRelativeToEnd = 0) override
 	{
 		check(NewPositionRelativeToEnd <= 0);
 
@@ -111,7 +111,7 @@ public:
 			return lseek(FileHandle, NewPositionRelativeToEnd, SEEK_END) != -1;
 		}
 	}
-	virtual bool Read(uint8* Destination, int64 BytesToRead) OVERRIDE
+	virtual bool Read(uint8* Destination, int64 BytesToRead) override
 	{
 #if MANAGE_FILE_HANDLES
 		if( IsManaged() )
@@ -127,7 +127,7 @@ public:
 			return ReadInternal(Destination, BytesToRead) == BytesToRead;
 		}
 	}
-	virtual bool Write(const uint8* Source, int64 BytesToWrite) OVERRIDE
+	virtual bool Write(const uint8* Source, int64 BytesToWrite) override
 	{
 		check(IsValid());
 		while (BytesToWrite)
@@ -144,7 +144,7 @@ public:
 		}
 		return true;
 	}
-	virtual int64 Size() OVERRIDE
+	virtual int64 Size() override
 	{
 #if MANAGE_FILE_HANDLES
 		if( IsManaged() )
diff --git a/Engine/Source/Runtime/Core/Private/Async/TaskGraph.cpp b/Engine/Source/Runtime/Core/Private/Async/TaskGraph.cpp
--- a/Engine/Source/Runtime/Core/Private/Async/TaskGraph.cpp
+++ b/Engine/Source/Runtime/Core/Private/Async/TaskGraph.cpp
@@ -178,7 +178,7 @@ public:
 	}
 
 	/** Tick single-threaded. */
-	virtual void Tick() OVERRIDE
+	virtual void Tick() override
 	{
 		if (Queue(0).QuitWhenIdle.GetValue() == 0)
 		{
@@ -466,7 +466,7 @@ public:
 	/**
 	 * Return single threaded interface when multithreading is disabled.
 	 */
-	virtual FSingleThreadRunnable* GetSingleThreadInterface() OVERRIDE
+	virtual FSingleThreadRunnable* GetSingleThreadInterface() override
 	{
 		return this;
 	}
@@ -715,7 +715,7 @@ public:
 	 *	@param	ThreadToExecuteOn; Either a named thread for a threadlocked task or ENamedThreads::AnyThread for a task that is to run on a worker thread
 	 *	@param	CurrentThreadIfKnown; This should be the current thread if it is known, or otherwise use ENamedThreads::AnyThread and the current thread will be determined.
 	**/
-	virtual void QueueTask(FBaseGraphTask* Task, ENamedThreads::Type ThreadToExecuteOn, ENamedThreads::Type CurrentThreadIfKnown = ENamedThreads::AnyThread) OVERRIDE
+	virtual void QueueTask(FBaseGraphTask* Task, ENamedThreads::Type ThreadToExecuteOn, ENamedThreads::Type CurrentThreadIfKnown = ENamedThreads::AnyThread) override
 	{
 		checkThreadGraph(NextUnnamedThreadMod);
 		if (CurrentThreadIfKnown == ENamedThreads::AnyThread)
@@ -761,17 +761,17 @@ public:
 	}
 
 
-	virtual	int32 GetNumWorkerThreads() OVERRIDE
+	virtual	int32 GetNumWorkerThreads() override
 	{
 		return NumThreads - NumNamedThreads;
 	}
 
-	virtual ENamedThreads::Type GetCurrentThreadIfKnown() OVERRIDE
+	virtual ENamedThreads::Type GetCurrentThreadIfKnown() override
 	{
 		return GetCurrentThread();
 	}
 
-	virtual bool IsThreadProcessingTasks(ENamedThreads::Type ThreadToCheck) OVERRIDE
+	virtual bool IsThreadProcessingTasks(ENamedThreads::Type ThreadToCheck) override
 	{
 		int32 QueueIndex = ENamedThreads::GetQueueIndex(ThreadToCheck);
 		ThreadToCheck = ENamedThreads::GetThreadIndex(ThreadToCheck);
@@ -781,7 +781,7 @@ public:
 
 	// External Thread API
 
-	virtual void AttachToThread(ENamedThreads::Type CurrentThread) OVERRIDE
+	virtual void AttachToThread(ENamedThreads::Type CurrentThread) override
 	{
 		CurrentThread = ENamedThreads::GetThreadIndex(CurrentThread);
 		check(NextUnnamedThreadMod);
@@ -790,7 +790,7 @@ public:
 		Thread(CurrentThread).InitializeForCurrentThread();
 	}
 
-	virtual void ProcessThreadUntilIdle(ENamedThreads::Type CurrentThread) OVERRIDE
+	virtual void ProcessThreadUntilIdle(ENamedThreads::Type CurrentThread) override
 	{
 		int32 QueueIndex = ENamedThreads::GetQueueIndex(CurrentThread);
 		CurrentThread = ENamedThreads::GetThreadIndex(CurrentThread);
@@ -799,7 +799,7 @@ public:
 		Thread(CurrentThread).ProcessTasks(QueueIndex, false);
 	}
 
-	virtual void ProcessThreadUntilRequestReturn(ENamedThreads::Type CurrentThread) OVERRIDE
+	virtual void ProcessThreadUntilRequestReturn(ENamedThreads::Type CurrentThread) override
 	{
 		int32 QueueIndex = ENamedThreads::GetQueueIndex(CurrentThread);
 		CurrentThread = ENamedThreads::GetThreadIndex(CurrentThread);
@@ -808,7 +808,7 @@ public:
 		Thread(CurrentThread).ProcessTasksUntilQuit(QueueIndex);
 	}
 
-	virtual void RequestReturn(ENamedThreads::Type CurrentThread) OVERRIDE
+	virtual void RequestReturn(ENamedThreads::Type CurrentThread) override
 	{
 		int32 QueueIndex = ENamedThreads::GetQueueIndex(CurrentThread);
 		CurrentThread = ENamedThreads::GetThreadIndex(CurrentThread);
@@ -816,7 +816,7 @@ public:
 		Thread(CurrentThread).RequestQuit(QueueIndex);
 	}
 
-	virtual void WaitUntilTasksComplete(const FGraphEventArray& Tasks, ENamedThreads::Type CurrentThreadIfKnown = ENamedThreads::AnyThread) OVERRIDE
+	virtual void WaitUntilTasksComplete(const FGraphEventArray& Tasks, ENamedThreads::Type CurrentThreadIfKnown = ENamedThreads::AnyThread) override
 	{
 		ENamedThreads::Type CurrentThread = CurrentThreadIfKnown;
 		if (CurrentThreadIfKnown == ENamedThreads::AnyThread)
diff --git a/Engine/Source/Runtime/Core/Private/Features/ModularFeatures.h b/Engine/Source/Runtime/Core/Private/Features/ModularFeatures.h
--- a/Engine/Source/Runtime/Core/Private/Features/ModularFeatures.h
+++ b/Engine/Source/Runtime/Core/Private/Features/ModularFeatures.h
@@ -13,14 +13,14 @@ class FModularFeatures : public IModularFeatures
 public:
 
 	/** IModularFeatures interface */
-	virtual int32 GetModularFeatureImplementationCount( const FName Type ) OVERRIDE;
-	virtual IModularFeature* GetModularFeatureImplementation( const FName Type, const int32 Index ) OVERRIDE;
-	virtual void RegisterModularFeature( const FName Type, class IModularFeature* ModularFeature ) OVERRIDE;
-	virtual void UnregisterModularFeature( const FName Type, class IModularFeature* ModularFeature ) OVERRIDE;
+	virtual int32 GetModularFeatureImplementationCount( const FName Type ) override;
+	virtual IModularFeature* GetModularFeatureImplementation( const FName Type, const int32 Index ) override;
+	virtual void RegisterModularFeature( const FName Type, class IModularFeature* ModularFeature ) override;
+	virtual void UnregisterModularFeature( const FName Type, class IModularFeature* ModularFeature ) override;
 	DECLARE_DERIVED_EVENT(FModularFeatures, IModularFeatures::FOnModularFeatureRegistered, FOnModularFeatureRegistered);
-	virtual IModularFeatures::FOnModularFeatureRegistered& OnModularFeatureRegistered() OVERRIDE;
+	virtual IModularFeatures::FOnModularFeatureRegistered& OnModularFeatureRegistered() override;
 	DECLARE_DERIVED_EVENT(FModularFeatures, IModularFeatures::FOnModularFeatureUnregistered, FOnModularFeatureUnregistered);
-	virtual IModularFeatures::FOnModularFeatureUnregistered& OnModularFeatureUnregistered() OVERRIDE;
+	virtual IModularFeatures::FOnModularFeatureUnregistered& OnModularFeatureUnregistered() override;
 
 private:
 
diff --git a/Engine/Source/Runtime/Core/Private/GenericPlatform/GenericPlatformHttp.cpp b/Engine/Source/Runtime/Core/Private/GenericPlatform/GenericPlatformHttp.cpp
--- a/Engine/Source/Runtime/Core/Private/GenericPlatform/GenericPlatformHttp.cpp
+++ b/Engine/Source/Runtime/Core/Private/GenericPlatform/GenericPlatformHttp.cpp
@@ -13,27 +13,27 @@ class FGenericPlatformHttpRequest : public IHttpRequest
 {
 public:
 	// IHttpBase
-	virtual FString GetURL() OVERRIDE { return TEXT(""); }
-	virtual FString GetURLParameter(const FString& ParameterName) OVERRIDE { return TEXT(""); }
-	virtual FString GetHeader(const FString& HeaderName) OVERRIDE { return TEXT(""); }
-	virtual TArray<FString> GetAllHeaders() OVERRIDE { return TArray<FString>(); }
-	virtual FString GetContentType() OVERRIDE { return TEXT(""); }
-	virtual int32 GetContentLength() OVERRIDE { return 0; }
+	virtual FString GetURL() override { return TEXT(""); }
+	virtual FString GetURLParameter(const FString& ParameterName) override { return TEXT(""); }
+	virtual FString GetHeader(const FString& HeaderName) override { return TEXT(""); }
+	virtual TArray<FString> GetAllHeaders() override { return TArray<FString>(); }
+	virtual FString GetContentType() override { return TEXT(""); }
+	virtual int32 GetContentLength() override { return 0; }
 	virtual const TArray<uint8>& GetContent() { static TArray<uint8> Temp; return Temp; }
 
 	// IHttpRequest
-	virtual FString GetVerb() OVERRIDE { return TEXT(""); }
-	virtual void SetVerb(const FString& Verb) OVERRIDE {}
-	virtual void SetURL(const FString& URL) OVERRIDE {}
-	virtual void SetContent(const TArray<uint8>& ContentPayload) OVERRIDE {}
-	virtual void SetContentAsString(const FString& ContentString) OVERRIDE {}
-	virtual void SetHeader(const FString& HeaderName, const FString& HeaderValue) OVERRIDE {}
-	virtual bool ProcessRequest() OVERRIDE { return false; }
-	virtual FHttpRequestCompleteDelegate& OnProcessRequestComplete() OVERRIDE { static FHttpRequestCompleteDelegate RequestCompleteDelegate; return RequestCompleteDelegate; }
-	virtual FHttpRequestProgressDelegate& OnRequestProgress() OVERRIDE { static FHttpRequestProgressDelegate RequestProgressDelegate; return RequestProgressDelegate; }
-	virtual void CancelRequest() OVERRIDE {}
-	virtual EHttpRequestStatus::Type GetStatus() OVERRIDE { return EHttpRequestStatus::NotStarted; }
-	virtual void Tick(float DeltaSeconds) OVERRIDE {}
+	virtual FString GetVerb() override { return TEXT(""); }
+	virtual void SetVerb(const FString& Verb) override {}
+	virtual void SetURL(const FString& URL) override {}
+	virtual void SetContent(const TArray<uint8>& ContentPayload) override {}
+	virtual void SetContentAsString(const FString& ContentString) override {}
+	virtual void SetHeader(const FString& HeaderName, const FString& HeaderValue) override {}
+	virtual bool ProcessRequest() override { return false; }
+	virtual FHttpRequestCompleteDelegate& OnProcessRequestComplete() override { static FHttpRequestCompleteDelegate RequestCompleteDelegate; return RequestCompleteDelegate; }
+	virtual FHttpRequestProgressDelegate& OnRequestProgress() override { static FHttpRequestProgressDelegate RequestProgressDelegate; return RequestProgressDelegate; }
+	virtual void CancelRequest() override {}
+	virtual EHttpRequestStatus::Type GetStatus() override { return EHttpRequestStatus::NotStarted; }
+	virtual void Tick(float DeltaSeconds) override {}
 };
 
 
diff --git a/Engine/Source/Runtime/Core/Private/HAL/ConsoleManager.cpp b/Engine/Source/Runtime/Core/Private/HAL/ConsoleManager.cpp
--- a/Engine/Source/Runtime/Core/Private/HAL/ConsoleManager.cpp
+++ b/Engine/Source/Runtime/Core/Private/HAL/ConsoleManager.cpp
@@ -457,12 +457,12 @@ public:
 
 	// interface IConsoleCommand -----------------------------------
 
-	virtual void Release() OVERRIDE
+	virtual void Release() override
 	{
 		delete this; 
 	} 
 
-	virtual bool Execute( const TArray< FString > Args, UWorld* InWorld, FOutputDevice& OutputDevice ) OVERRIDE
+	virtual bool Execute( const TArray< FString > Args, UWorld* InWorld, FOutputDevice& OutputDevice ) override
 	{
 		// NOTE: Args are ignored for FConsoleCommand.  Use FConsoleCommandWithArgs if you need parameters.
 		return Delegate.ExecuteIfBound();
@@ -488,12 +488,12 @@ public:
 
 	// interface IConsoleCommand -----------------------------------
 
-	virtual void Release() OVERRIDE
+	virtual void Release() override
 	{
 		delete this; 
 	} 
 
-	virtual bool Execute( const TArray< FString > Args, UWorld* InWorld, FOutputDevice& OutputDevice ) OVERRIDE
+	virtual bool Execute( const TArray< FString > Args, UWorld* InWorld, FOutputDevice& OutputDevice ) override
 	{
 		return Delegate.ExecuteIfBound( Args );
 	}
@@ -517,12 +517,12 @@ public:
 
 	// interface IConsoleCommand -----------------------------------
 
-	virtual void Release() OVERRIDE
+	virtual void Release() override
 	{
 		delete this; 
 	} 
 
-	virtual bool Execute( const TArray< FString > Args, UWorld* InWorld, FOutputDevice& OutputDevice ) OVERRIDE
+	virtual bool Execute( const TArray< FString > Args, UWorld* InWorld, FOutputDevice& OutputDevice ) override
 	{
 		return Delegate.ExecuteIfBound( InWorld );
 	}
@@ -546,12 +546,12 @@ public:
 
 	// interface IConsoleCommand -----------------------------------
 
-	virtual void Release() OVERRIDE
+	virtual void Release() override
 	{
 		delete this; 
 	} 
 
-	virtual bool Execute( const TArray< FString > Args, UWorld* InWorld, FOutputDevice& OutputDevice ) OVERRIDE
+	virtual bool Execute( const TArray< FString > Args, UWorld* InWorld, FOutputDevice& OutputDevice ) override
 	{
 		return Delegate.ExecuteIfBound( Args, InWorld );
 	}
@@ -575,12 +575,12 @@ public:
 
 	// interface IConsoleCommand -----------------------------------
 
-	virtual void Release() OVERRIDE
+	virtual void Release() override
 	{
 		delete this; 
 	} 
 
-	virtual bool Execute( const TArray< FString > Args, UWorld* InWorld, FOutputDevice& OutputDevice ) OVERRIDE
+	virtual bool Execute( const TArray< FString > Args, UWorld* InWorld, FOutputDevice& OutputDevice ) override
 	{
 		return Delegate.ExecuteIfBound( OutputDevice );
 	}
@@ -603,12 +603,12 @@ public:
 
 	// interface IConsoleCommand -----------------------------------
 
-	virtual void Release() OVERRIDE
+	virtual void Release() override
 	{
 		delete this; 
 	} 
 
-	virtual bool Execute( const TArray< FString > Args, UWorld* InCmdWorld, FOutputDevice& OutputDevice ) OVERRIDE
+	virtual bool Execute( const TArray< FString > Args, UWorld* InCmdWorld, FOutputDevice& OutputDevice ) override
 	{
 		return false;
 	}
diff --git a/Engine/Source/Runtime/Core/Private/HAL/PThreadRunnableThread.h b/Engine/Source/Runtime/Core/Private/HAL/PThreadRunnableThread.h
--- a/Engine/Source/Runtime/Core/Private/HAL/PThreadRunnableThread.h
+++ b/Engine/Source/Runtime/Core/Private/HAL/PThreadRunnableThread.h
@@ -262,7 +262,7 @@ public:
 		ThreadID = 0;
 	}
 
-	virtual void SetThreadPriority(EThreadPriority NewPriority) OVERRIDE
+	virtual void SetThreadPriority(EThreadPriority NewPriority) override
 	{
 		// Don't bother calling the OS if there is no need
 		if (NewPriority != ThreadPriority)
@@ -272,13 +272,13 @@ public:
 		}
 	}
 
-	virtual void Suspend(bool bShouldPause = 1) OVERRIDE
+	virtual void Suspend(bool bShouldPause = 1) override
 	{
 		check(Thread);
 		// Impossible in pthreads!
 	}
 
-	virtual bool Kill(bool bShouldWait = false) OVERRIDE
+	virtual bool Kill(bool bShouldWait = false) override
 	{
 		check(Thread && "Did you forget to call Create()?");
 		bool bDidExitOK = true;
@@ -304,7 +304,7 @@ public:
 		return bDidExitOK;
 	}
 
-	virtual void WaitForCompletion() OVERRIDE
+	virtual void WaitForCompletion() override
 	{
 		// Block until this thread exits
 		while (ThreadIsRunning)
@@ -313,12 +313,12 @@ public:
 		}
 	}
 
-	virtual uint32 GetThreadID() OVERRIDE
+	virtual uint32 GetThreadID() override
 	{
 		return ThreadID;
 	}
 
-	virtual FString GetThreadName() OVERRIDE
+	virtual FString GetThreadName() override
 	{
 		return ThreadName;
 	}
@@ -327,7 +327,7 @@ protected:
 
 	virtual bool CreateInternal(FRunnable* InRunnable, const TCHAR* InThreadName,
 		uint32 InStackSize = 0,
-		EThreadPriority InThreadPri = TPri_Normal, uint64 InThreadAffinityMask = 0) OVERRIDE
+		EThreadPriority InThreadPri = TPri_Normal, uint64 InThreadAffinityMask = 0) override
 	{
 		check(InRunnable);
 		Runnable = InRunnable;
diff --git a/Engine/Source/Runtime/Core/Private/HAL/ThreadingBase.cpp b/Engine/Source/Runtime/Core/Private/HAL/ThreadingBase.cpp
--- a/Engine/Source/Runtime/Core/Private/HAL/ThreadingBase.cpp
+++ b/Engine/Source/Runtime/Core/Private/HAL/ThreadingBase.cpp
@@ -77,12 +77,12 @@ public:
 	}
 
 	// FRunnableThread interface
-	virtual void SetThreadPriority(EThreadPriority NewPriority) OVERRIDE
+	virtual void SetThreadPriority(EThreadPriority NewPriority) override
 	{
 		// Not relevant.
 	}
 
-	virtual void Suspend(bool bShouldPause) OVERRIDE
+	virtual void Suspend(bool bShouldPause) override
 	{
 		bIsSuspended = bShouldPause;
 	}
@@ -93,17 +93,17 @@ public:
 		return true;
 	}
 
-	virtual void WaitForCompletion() OVERRIDE
+	virtual void WaitForCompletion() override
 	{
 		FSingleThreadManager::Get().RemoveThread(this);
 	}
 
-	virtual uint32 GetThreadID() OVERRIDE
+	virtual uint32 GetThreadID() override
 	{
 		return ThreadId;
 	}
 
-	virtual FString GetThreadName() OVERRIDE
+	virtual FString GetThreadName() override
 	{
 		return Name;
 	}
@@ -286,7 +286,7 @@ protected:
 	 * The real thread entry point. It waits for work events to be queued. Once
 	 * an event is queued, it executes it and goes back to waiting.
 	 */
-	virtual uint32 Run() OVERRIDE
+	virtual uint32 Run() override
 	{
 		while (!TimeToDie)
 		{
@@ -436,7 +436,7 @@ public:
 		Destroy();
 	}
 
-	virtual bool Create(uint32 InNumQueuedThreads,uint32 StackSize = (32 * 1024),EThreadPriority ThreadPriority=TPri_Normal) OVERRIDE
+	virtual bool Create(uint32 InNumQueuedThreads,uint32 StackSize = (32 * 1024),EThreadPriority ThreadPriority=TPri_Normal) override
 	{
 		// Make sure we have synch objects
 		bool bWasSuccessful = true;
@@ -471,7 +471,7 @@ public:
 		return bWasSuccessful;
 	}
 
-	virtual void Destroy() OVERRIDE
+	virtual void Destroy() override
 	{
 		if (SynchQueue)
 		{
@@ -516,7 +516,7 @@ public:
 		}
 	}
 
-	void AddQueuedWork(FQueuedWork* InQueuedWork) OVERRIDE
+	void AddQueuedWork(FQueuedWork* InQueuedWork) override
 	{
 		if (TimeToDie)
 		{
@@ -552,7 +552,7 @@ public:
 		}
 	}
 
-	virtual bool RetractQueuedWork(FQueuedWork* InQueuedWork) OVERRIDE
+	virtual bool RetractQueuedWork(FQueuedWork* InQueuedWork) override
 	{
 		if (TimeToDie)
 		{
@@ -564,7 +564,7 @@ public:
 		return !!QueuedWork.RemoveSingle(InQueuedWork);
 	}
 
-	virtual FQueuedWork* ReturnToPoolOrGetNextJob(FQueuedThread* InQueuedThread) OVERRIDE
+	virtual FQueuedWork* ReturnToPoolOrGetNextJob(FQueuedThread* InQueuedThread) override
 	{
 		check(InQueuedThread != NULL);
 		FQueuedWork* Work = NULL;
diff --git a/Engine/Source/Runtime/Core/Private/HTML5/HTML5PlatformFile.cpp b/Engine/Source/Runtime/Core/Private/HTML5/HTML5PlatformFile.cpp
--- a/Engine/Source/Runtime/Core/Private/HTML5/HTML5PlatformFile.cpp
+++ b/Engine/Source/Runtime/Core/Private/HTML5/HTML5PlatformFile.cpp
@@ -71,27 +71,27 @@ public:
 		FileHandle = -1;
 	}
 
-	virtual int64 Tell() OVERRIDE
+	virtual int64 Tell() override
 	{
 		check(IsValid());
 		return lseek(FileHandle, 0, SEEK_CUR);
 	}
 
-	virtual bool Seek(int64 NewPosition) OVERRIDE
+	virtual bool Seek(int64 NewPosition) override
 	{
 		check(IsValid());
 		check(NewPosition >= 0);
 		return lseek(FileHandle, NewPosition, SEEK_SET) != -1;
 	}
 
-	virtual bool SeekFromEnd(int64 NewPositionRelativeToEnd = 0) OVERRIDE
+	virtual bool SeekFromEnd(int64 NewPositionRelativeToEnd = 0) override
 	{
 		check(IsValid());
 		check(NewPositionRelativeToEnd <= 0);
 		return lseek(FileHandle, NewPositionRelativeToEnd, SEEK_END) != -1;
 	}
 
-	virtual bool Read(uint8* Destination, int64 BytesToRead) OVERRIDE
+	virtual bool Read(uint8* Destination, int64 BytesToRead) override
 	{
 		check(IsValid());
 		while (BytesToRead)
@@ -109,7 +109,7 @@ public:
 		return true;
 	}
 
-	virtual bool Write(const uint8* Source, int64 BytesToWrite) OVERRIDE
+	virtual bool Write(const uint8* Source, int64 BytesToWrite) override
 	{
 		check(IsValid());
 		while (BytesToWrite)
@@ -177,7 +177,7 @@ protected:
 	}
 
 public:
-	virtual bool FileExists(const TCHAR* Filename) OVERRIDE
+	virtual bool FileExists(const TCHAR* Filename) override
 	{
 		struct stat FileInfo;
 		if (stat(TCHAR_TO_UTF8(*NormalizeFilename(Filename)), &FileInfo) != -1)
@@ -187,7 +187,7 @@ public:
 		return false;
 	}
 
-	virtual int64 FileSize(const TCHAR* Filename) OVERRIDE
+	virtual int64 FileSize(const TCHAR* Filename) override
 	{
 		struct stat FileInfo;
 		FileInfo.st_size = -1;
@@ -200,12 +200,12 @@ public:
 		return FileInfo.st_size;
 	}
 
-	virtual bool DeleteFile(const TCHAR* Filename) OVERRIDE
+	virtual bool DeleteFile(const TCHAR* Filename) override
 	{
 		return unlink(TCHAR_TO_UTF8(*NormalizeFilename(Filename))) == 0;
 	}
 
-	virtual bool IsReadOnly(const TCHAR* Filename) OVERRIDE
+	virtual bool IsReadOnly(const TCHAR* Filename) override
 	{
 		if (access(TCHAR_TO_UTF8(*NormalizeFilename(Filename)), F_OK) == -1)
 		{
@@ -218,12 +218,12 @@ public:
 		return false;
 	}
 
-	virtual bool MoveFile(const TCHAR* To, const TCHAR* From) OVERRIDE
+	virtual bool MoveFile(const TCHAR* To, const TCHAR* From) override
 	{
 		return rename(TCHAR_TO_UTF8(*NormalizeFilename(From)), TCHAR_TO_UTF8(*NormalizeFilename(To))) != -1;
 	}
 
-	virtual bool SetReadOnly(const TCHAR* Filename, bool bNewReadOnlyValue) OVERRIDE
+	virtual bool SetReadOnly(const TCHAR* Filename, bool bNewReadOnlyValue) override
 	{
 		struct stat FileInfo;
 		if (stat(TCHAR_TO_UTF8(*NormalizeFilename(Filename)), &FileInfo) != -1)
@@ -241,7 +241,7 @@ public:
 		return false;
 	}
 
-	virtual FDateTime GetTimeStamp(const TCHAR* Filename) OVERRIDE
+	virtual FDateTime GetTimeStamp(const TCHAR* Filename) override
 	{
 		// get file times
 		struct stat FileInfo;
@@ -255,7 +255,7 @@ public:
 		return HTML5Epoch + TimeSinceEpoch;
 	}
 
-	virtual void SetTimeStamp(const TCHAR* Filename, FDateTime DateTime) OVERRIDE
+	virtual void SetTimeStamp(const TCHAR* Filename, FDateTime DateTime) override
 	{
 		// get file times
 		struct stat FileInfo;
@@ -271,7 +271,7 @@ public:
 		utime(TCHAR_TO_UTF8(*NormalizeFilename(Filename)), &Times);
 	}
 
-	virtual FDateTime GetAccessTimeStamp(const TCHAR* Filename) OVERRIDE
+	virtual FDateTime GetAccessTimeStamp(const TCHAR* Filename) override
 	{
 		// get file times
 		struct stat FileInfo;
@@ -285,7 +285,7 @@ public:
 		return HTML5Epoch + TimeSinceEpoch;
 	}
 
-	virtual IFileHandle* OpenRead(const TCHAR* Filename) OVERRIDE
+	virtual IFileHandle* OpenRead(const TCHAR* Filename) override
 	{
 		FString fn = NormalizeFilename(Filename);
 		int32 Handle = open(TCHAR_TO_UTF8(*fn), O_RDONLY | O_BINARY);
@@ -300,7 +300,7 @@ public:
 		return NULL;
 	}
 
-	virtual IFileHandle* OpenWrite(const TCHAR* Filename, bool bAppend = false, bool bAllowRead = false) OVERRIDE
+	virtual IFileHandle* OpenWrite(const TCHAR* Filename, bool bAppend = false, bool bAllowRead = false) override
 	{
 		int Flags = O_CREAT;
 		if (bAppend)
@@ -332,7 +332,7 @@ public:
 		return NULL;
 	}
 
-	virtual bool DirectoryExists(const TCHAR* Directory) OVERRIDE
+	virtual bool DirectoryExists(const TCHAR* Directory) override
 	{
 		struct stat FileInfo;
 		if (stat(TCHAR_TO_UTF8(*NormalizeFilename(Directory)), &FileInfo) != -1)
@@ -341,7 +341,7 @@ public:
 		}
 		return false;
 	}
-	virtual bool CreateDirectory(const TCHAR* Directory) OVERRIDE
+	virtual bool CreateDirectory(const TCHAR* Directory) override
 	{
 #if PLATFORM_HTML5_WIN32
 		return _mkdir(TCHAR_TO_ANSI(*NormalizeFilename(Directory))) || (errno == EEXIST);
@@ -351,7 +351,7 @@ public:
 #endif 
 	}
 
-	virtual bool DeleteDirectory(const TCHAR* Directory) OVERRIDE
+	virtual bool DeleteDirectory(const TCHAR* Directory) override
 	{
 #if PLATFORM_HTML5_WIN32
 		return false;
diff --git a/Engine/Source/Runtime/Core/Private/IOS/IOSPlatformFile.h b/Engine/Source/Runtime/Core/Private/IOS/IOSPlatformFile.h
--- a/Engine/Source/Runtime/Core/Private/IOS/IOSPlatformFile.h
+++ b/Engine/Source/Runtime/Core/Private/IOS/IOSPlatformFile.h
@@ -22,11 +22,11 @@ public:
 	FFileHandleIOS(int32 InFileHandle = -1);
 	virtual ~FFileHandleIOS();
 
-	virtual int64 Tell() OVERRIDE;
-	virtual bool Seek(int64 NewPosition) OVERRIDE;
-	virtual bool SeekFromEnd(int64 NewPositionRelativeToEnd = 0) OVERRIDE;
-	virtual bool Read(uint8* Destination, int64 BytesToRead) OVERRIDE;
-	virtual bool Write(const uint8* Source, int64 BytesToWrite) OVERRIDE;
+	virtual int64 Tell() override;
+	virtual bool Seek(int64 NewPosition) override;
+	virtual bool SeekFromEnd(int64 NewPositionRelativeToEnd = 0) override;
+	virtual bool Read(uint8* Destination, int64 BytesToRead) override;
+	virtual bool Write(const uint8* Source, int64 BytesToWrite) override;
 };
 
 /**
@@ -41,23 +41,23 @@ protected:
 public:
 	//virtual bool Initialize(IPlatformFile* Inner, const TCHAR* CommandLineParam) OVERRIDE;
 
-	virtual bool FileExists(const TCHAR* Filename) OVERRIDE;
-	virtual int64 FileSize(const TCHAR* Filename) OVERRIDE;
-	virtual bool DeleteFile(const TCHAR* Filename) OVERRIDE;
-	virtual bool IsReadOnly(const TCHAR* Filename) OVERRIDE;
-	virtual bool MoveFile(const TCHAR* To, const TCHAR* From) OVERRIDE;
-	virtual bool SetReadOnly(const TCHAR* Filename, bool bNewReadOnlyValue) OVERRIDE;
+	virtual bool FileExists(const TCHAR* Filename) override;
+	virtual int64 FileSize(const TCHAR* Filename) override;
+	virtual bool DeleteFile(const TCHAR* Filename) override;
+	virtual bool IsReadOnly(const TCHAR* Filename) override;
+	virtual bool MoveFile(const TCHAR* To, const TCHAR* From) override;
+	virtual bool SetReadOnly(const TCHAR* Filename, bool bNewReadOnlyValue) override;
 
-	virtual FDateTime GetTimeStamp(const TCHAR* Filename) OVERRIDE;
-	virtual void SetTimeStamp(const TCHAR* Filename, const FDateTime DateTime) OVERRIDE;
-	virtual FDateTime GetAccessTimeStamp(const TCHAR* Filename) OVERRIDE;
+	virtual FDateTime GetTimeStamp(const TCHAR* Filename) override;
+	virtual void SetTimeStamp(const TCHAR* Filename, const FDateTime DateTime) override;
+	virtual FDateTime GetAccessTimeStamp(const TCHAR* Filename) override;
 
-	virtual IFileHandle* OpenRead(const TCHAR* Filename) OVERRIDE;
-	virtual IFileHandle* OpenWrite(const TCHAR* Filename, bool bAppend = false, bool bAllowRead = false) OVERRIDE;
+	virtual IFileHandle* OpenRead(const TCHAR* Filename) override;
+	virtual IFileHandle* OpenWrite(const TCHAR* Filename, bool bAppend = false, bool bAllowRead = false) override;
 
-	virtual bool DirectoryExists(const TCHAR* Directory) OVERRIDE;
-	virtual bool CreateDirectory(const TCHAR* Directory) OVERRIDE;
-	virtual bool DeleteDirectory(const TCHAR* Directory) OVERRIDE;
+	virtual bool DirectoryExists(const TCHAR* Directory) override;
+	virtual bool CreateDirectory(const TCHAR* Directory) override;
+	virtual bool DeleteDirectory(const TCHAR* Directory) override;
 	bool IterateDirectory(const TCHAR* Directory, FDirectoryVisitor& Visitor);
 
 private:
diff --git a/Engine/Source/Runtime/Core/Private/IOS/IOSPlatformOutputDevicesPrivate.h b/Engine/Source/Runtime/Core/Private/IOS/IOSPlatformOutputDevicesPrivate.h
--- a/Engine/Source/Runtime/Core/Private/IOS/IOSPlatformOutputDevicesPrivate.h
+++ b/Engine/Source/Runtime/Core/Private/IOS/IOSPlatformOutputDevicesPrivate.h
@@ -22,7 +22,7 @@ public:
 	 * @param	Data	Text to log
 	 * @param	Event	Event name used for suppression purposes
 	 */
-	virtual void Serialize( const TCHAR* Msg, ELogVerbosity::Type Verbosity, const class FName& Category ) OVERRIDE;
+	virtual void Serialize( const TCHAR* Msg, ELogVerbosity::Type Verbosity, const class FName& Category ) override;
 };
 
 
@@ -38,7 +38,7 @@ public:
 	 * @param	Data	Text to log
 	 * @param	Event	Event name used for suppression purposes
 	 */
-	virtual void Serialize( const TCHAR* Msg, ELogVerbosity::Type Verbosity, const class FName& Category ) OVERRIDE;
+	virtual void Serialize( const TCHAR* Msg, ELogVerbosity::Type Verbosity, const class FName& Category ) override;
 
 	/**
 	 * Error handling function that is being called from within the system wide global
diff --git a/Engine/Source/Runtime/Core/Private/IOS/IOSWindow.h b/Engine/Source/Runtime/Core/Private/IOS/IOSWindow.h
--- a/Engine/Source/Runtime/Core/Private/IOS/IOSWindow.h
+++ b/Engine/Source/Runtime/Core/Private/IOS/IOSWindow.h
@@ -24,16 +24,16 @@ public:
 	static FPlatformRect GetScreenRect();
 
 	/** Returns a void pointer to the hWnd (for other apis)*/
-	virtual void* GetOSWindowHandle() const OVERRIDE { return Window; }
+	virtual void* GetOSWindowHandle() const override { return Window; }
 
 	void Initialize( class FIOSApplication* const Application, const TSharedRef< FGenericWindowDefinition >& InDefinition, const TSharedPtr< FIOSWindow >& InParent, const bool bShowImmediately );
 
 	/** Returns the rectangle of the screen the window is associated with */
-	virtual bool GetFullScreenInfo( int32& X, int32& Y, int32& Width, int32& Height ) const OVERRIDE;
+	virtual bool GetFullScreenInfo( int32& X, int32& Y, int32& Width, int32& Height ) const override;
 
 protected:
 	/** @return true if the native window is currently in fullscreen mode, false otherwise */
-	virtual EWindowMode::Type GetWindowMode() const OVERRIDE { return EWindowMode::Fullscreen; }
+	virtual EWindowMode::Type GetWindowMode() const override { return EWindowMode::Fullscreen; }
 
 private:
 	/**
diff --git a/Engine/Source/Runtime/Core/Private/Internationalization/Internationalization.cpp b/Engine/Source/Runtime/Core/Private/Internationalization/Internationalization.cpp
--- a/Engine/Source/Runtime/Core/Private/Internationalization/Internationalization.cpp
+++ b/Engine/Source/Runtime/Core/Private/Internationalization/Internationalization.cpp
@@ -1,4 +1,4 @@
-﻿// Copyright 1998-2014 Epic Games, Inc. All Rights Reserved.
+// Copyright 1998-2014 Epic Games, Inc. All Rights Reserved.
 
 #include "CorePrivate.h"
 #include "InvariantCulture.h"
@@ -411,7 +411,7 @@ void FInternationalization::GetCulturesWithAvailableLocalization(const TArray<FS
 			{
 			}
 
-			virtual bool Visit(const TCHAR* FilenameOrDirectory, bool bIsDirectory) OVERRIDE
+			virtual bool Visit(const TCHAR* FilenameOrDirectory, bool bIsDirectory) override
 			{
 				if(bIsDirectory)
 				{
diff --git a/Engine/Source/Runtime/Core/Private/Internationalization/TextHistory.h b/Engine/Source/Runtime/Core/Private/Internationalization/TextHistory.h
--- a/Engine/Source/Runtime/Core/Private/Internationalization/TextHistory.h
+++ b/Engine/Source/Runtime/Core/Private/Internationalization/TextHistory.h
@@ -49,11 +49,11 @@ public:
 
 
 	// Begin FTextHistory interface
-	virtual FText ToText() const OVERRIDE;
-	virtual void Serialize(FArchive& Ar) OVERRIDE;
-	virtual void SerializeForDisplayString(FArchive& Ar, TSharedRef<FString>& InOutDisplayString) OVERRIDE;
-	virtual bool IsOutOfDate(int32 InRevision) OVERRIDE { return false; }
-	virtual TSharedPtr< FString > GetSourceString() OVERRIDE;
+	virtual FText ToText() const override;
+	virtual void Serialize(FArchive& Ar) override;
+	virtual void SerializeForDisplayString(FArchive& Ar, TSharedRef<FString>& InOutDisplayString) override;
+	virtual bool IsOutOfDate(int32 InRevision) override { return false; }
+	virtual TSharedPtr< FString > GetSourceString() override;
 	// End FTextHistory interface
 
 private:
@@ -69,8 +69,8 @@ public:
 	FTextHistory_NamedFormat(const FText& InSourceText, const FFormatNamedArguments& InArguments);
 
 	// Begin FTextHistory interface
-	virtual FText ToText() const OVERRIDE;
-	virtual void Serialize(FArchive& Ar) OVERRIDE;
+	virtual FText ToText() const override;
+	virtual void Serialize(FArchive& Ar) override;
 	// End FTextHistory interface
 
 private:
@@ -88,8 +88,8 @@ public:
 	FTextHistory_OrderedFormat(const FText& InSourceText, const FFormatOrderedArguments& InArguments);
 
 	// Begin FTextHistory interface
-	virtual FText ToText() const OVERRIDE;
-	virtual void Serialize(FArchive& Ar) OVERRIDE;
+	virtual FText ToText() const override;
+	virtual void Serialize(FArchive& Ar) override;
 	// End FTextHistory interface
 
 private:
@@ -107,8 +107,8 @@ public:
 	FTextHistory_ArgumentDataFormat(const FText& InSourceText, const TArray< struct FFormatArgumentData >& InArguments);
 
 	// Begin FTextHistory interface
-	virtual FText ToText() const OVERRIDE;
-	virtual void Serialize(FArchive& Ar) OVERRIDE;
+	virtual FText ToText() const override;
+	virtual void Serialize(FArchive& Ar) override;
 	// End FTextHistory interface
 
 private:
@@ -128,7 +128,7 @@ public:
 	~FTextHistory_FormatNumber();
 
 	// Begin FTextHistory interface
-	virtual void Serialize(FArchive& Ar) OVERRIDE;
+	virtual void Serialize(FArchive& Ar) override;
 	// End FTextHistory interface
 protected:
 	/** The source value to format from */
@@ -147,8 +147,8 @@ public:
 	FTextHistory_AsNumber(const FFormatArgumentValue& InSourceValue, const FNumberFormattingOptions* const InFormatOptions, const TSharedPtr<FCulture> InTargetCulture);
 
 	// Begin FTextHistory interface
-	virtual FText ToText() const OVERRIDE;
-	virtual void Serialize(FArchive& Ar) OVERRIDE;
+	virtual FText ToText() const override;
+	virtual void Serialize(FArchive& Ar) override;
 	// End FTextHistory interface
 };
 
@@ -160,8 +160,8 @@ public:
 	FTextHistory_AsPercent(const FFormatArgumentValue& InSourceValue, const FNumberFormattingOptions* const InFormatOptions, const TSharedPtr<FCulture> InTargetCulture);
 
 	// Begin FTextHistory interface
-	virtual FText ToText() const OVERRIDE;
-	virtual void Serialize(FArchive& Ar) OVERRIDE;
+	virtual FText ToText() const override;
+	virtual void Serialize(FArchive& Ar) override;
 	// End FTextHistory interface
 };
 
@@ -173,8 +173,8 @@ public:
 	FTextHistory_AsCurrency(const FFormatArgumentValue& InSourceValue, const FNumberFormattingOptions* const InFormatOptions, const TSharedPtr<FCulture> InTargetCulture);
 
 	// Begin FTextHistory interface
-	virtual FText ToText() const OVERRIDE;
-	virtual void Serialize(FArchive& Ar) OVERRIDE;
+	virtual FText ToText() const override;
+	virtual void Serialize(FArchive& Ar) override;
 	// End FTextHistory interface
 };
 
@@ -186,8 +186,8 @@ public:
 	FTextHistory_AsDate(const FDateTime& InSourceDateTime, const EDateTimeStyle::Type InDateStyle, const TSharedPtr<FCulture> InTargetCulture);
 
 	// Begin FTextHistory interface
-	virtual FText ToText() const OVERRIDE;
-	virtual void Serialize(FArchive& Ar) OVERRIDE;
+	virtual FText ToText() const override;
+	virtual void Serialize(FArchive& Ar) override;
 	// End FTextHistory interface
 
 private:
@@ -207,8 +207,8 @@ public:
 	FTextHistory_AsTime(const FDateTime& InSourceDateTime, const EDateTimeStyle::Type InTimeStyle, const FString& InTimeZone, const TSharedPtr<FCulture> InTargetCulture);
 
 	// Begin FTextHistory interface
-	virtual FText ToText() const OVERRIDE;
-	virtual void Serialize(FArchive& Ar) OVERRIDE;
+	virtual FText ToText() const override;
+	virtual void Serialize(FArchive& Ar) override;
 	// End FTextHistory interface
 
 private:
@@ -230,8 +230,8 @@ public:
 	FTextHistory_AsDateTime(const FDateTime& InSourceDateTime, const EDateTimeStyle::Type InDateStyle, const EDateTimeStyle::Type InTimeStyle, const FString& InTimeZone, const TSharedPtr<FCulture> InTargetCulture);
 
 	// Begin FTextHistory interface
-	virtual FText ToText() const OVERRIDE;
-	virtual void Serialize(FArchive& Ar) OVERRIDE;
+	virtual FText ToText() const override;
+	virtual void Serialize(FArchive& Ar) override;
 	// End FTextHistory interfaces
 
 private:
diff --git a/Engine/Source/Runtime/Core/Private/Linux/LinuxPlatformFile.cpp b/Engine/Source/Runtime/Core/Private/Linux/LinuxPlatformFile.cpp
--- a/Engine/Source/Runtime/Core/Private/Linux/LinuxPlatformFile.cpp
+++ b/Engine/Source/Runtime/Core/Private/Linux/LinuxPlatformFile.cpp
@@ -33,27 +33,27 @@ public:
 		FileHandle = -1;
 	}
 
-	virtual int64 Tell() OVERRIDE
+	virtual int64 Tell() override
 	{
 		check(IsValid());
 		return lseek(FileHandle, 0, SEEK_CUR);
 	}
 
-	virtual bool Seek(int64 NewPosition) OVERRIDE
+	virtual bool Seek(int64 NewPosition) override
 	{
 		check(IsValid());
 		check(NewPosition >= 0);
 		return lseek(FileHandle, NewPosition, SEEK_SET) != -1;
 	}
 
-	virtual bool SeekFromEnd(int64 NewPositionRelativeToEnd = 0) OVERRIDE
+	virtual bool SeekFromEnd(int64 NewPositionRelativeToEnd = 0) override
 	{
 		check(IsValid());
 		check(NewPositionRelativeToEnd <= 0);
 		return lseek(FileHandle, NewPositionRelativeToEnd, SEEK_END) != -1;
 	}
 
-	virtual bool Read(uint8* Destination, int64 BytesToRead) OVERRIDE
+	virtual bool Read(uint8* Destination, int64 BytesToRead) override
 	{
 		check(IsValid());
 		while (BytesToRead)
@@ -71,7 +71,7 @@ public:
 		return true;
 	}
 
-	virtual bool Write(const uint8* Source, int64 BytesToWrite) OVERRIDE
+	virtual bool Write(const uint8* Source, int64 BytesToWrite) override
 	{
 		check(IsValid());
 		while (BytesToWrite)
diff --git a/Engine/Source/Runtime/Core/Private/Mac/MacPlatformOutputDevicesPrivate.h b/Engine/Source/Runtime/Core/Private/Mac/MacPlatformOutputDevicesPrivate.h
--- a/Engine/Source/Runtime/Core/Private/Mac/MacPlatformOutputDevicesPrivate.h
+++ b/Engine/Source/Runtime/Core/Private/Mac/MacPlatformOutputDevicesPrivate.h
@@ -18,7 +18,7 @@ public:
 	 * @param	Data	Text to log
 	 * @param	Event	Event name used for suppression purposes
 	 */
-	virtual void Serialize( const TCHAR* Msg, ELogVerbosity::Type Verbosity, const class FName& Category ) OVERRIDE;
+	virtual void Serialize( const TCHAR* Msg, ELogVerbosity::Type Verbosity, const class FName& Category ) override;
 
 	/**
 	 * Error handling function that is being called from within the system wide global
diff --git a/Engine/Source/Runtime/Core/Private/Mac/MacTextInputMethodSystem.cpp b/Engine/Source/Runtime/Core/Private/Mac/MacTextInputMethodSystem.cpp
--- a/Engine/Source/Runtime/Core/Private/Mac/MacTextInputMethodSystem.cpp
+++ b/Engine/Source/Runtime/Core/Private/Mac/MacTextInputMethodSystem.cpp
@@ -359,10 +359,10 @@ namespace
 		void SetContextWindow(TSharedPtr<FGenericWindow> Window);
 		TSharedPtr<FGenericWindow> GetContextWindow(void) const;
 		
-		virtual void NotifyLayoutChanged(const ELayoutChangeType ChangeType) OVERRIDE;
-		virtual void NotifySelectionChanged() OVERRIDE;
-		virtual void NotifyTextChanged(const uint32 BeginIndex, const uint32 OldLength, const uint32 NewLength) OVERRIDE;
-		virtual void CancelComposition() OVERRIDE;
+		virtual void NotifyLayoutChanged(const ELayoutChangeType ChangeType) override;
+		virtual void NotifySelectionChanged() override;
+		virtual void NotifyTextChanged(const uint32 BeginIndex, const uint32 OldLength, const uint32 NewLength) override;
+		virtual void CancelComposition() override;
 		
 	private:
 		TWeakPtr<ITextInputMethodContext> Context;
diff --git a/Engine/Source/Runtime/Core/Private/Misc/MessageLog.cpp b/Engine/Source/Runtime/Core/Private/Misc/MessageLog.cpp
--- a/Engine/Source/Runtime/Core/Private/Misc/MessageLog.cpp
+++ b/Engine/Source/Runtime/Core/Private/Misc/MessageLog.cpp
@@ -19,12 +19,12 @@ public:
 	}
 
 	/** Begin IMessageLog interface */
-	virtual void AddMessage( const TSharedRef<FTokenizedMessage>& NewMessage ) OVERRIDE
+	virtual void AddMessage( const TSharedRef<FTokenizedMessage>& NewMessage ) override
 	{
 		AddMessageInternal(NewMessage);
 	}
 
-	virtual void AddMessages( const TArray< TSharedRef<FTokenizedMessage> >& NewMessages ) OVERRIDE
+	virtual void AddMessages( const TArray< TSharedRef<FTokenizedMessage> >& NewMessages ) override
 	{
 		for(TArray< TSharedRef<FTokenizedMessage> >::TConstIterator It = NewMessages.CreateConstIterator(); It; It++)
 		{
@@ -32,26 +32,26 @@ public:
 		}
 	}
 
-	virtual void NewPage( const FText& Title ) OVERRIDE
+	virtual void NewPage( const FText& Title ) override
 	{
 		FFormatNamedArguments Arguments;
 		Arguments.Add(TEXT("PageTitle"), Title);
 		FMsg::Logf(__FILE__, __LINE__, LogName, ELogVerbosity::Log, *FText::Format(LOCTEXT("BasicMessageLog_NewPage", "New Page: {PageTitle}"), Arguments).ToString());
 	}
 
-	virtual void NotifyIfAnyMessages( const FText& Message, EMessageSeverity::Type SeverityFilter, bool bForce ) OVERRIDE
+	virtual void NotifyIfAnyMessages( const FText& Message, EMessageSeverity::Type SeverityFilter, bool bForce ) override
 	{
 		FFormatNamedArguments Arguments;
 		Arguments.Add(TEXT("Message"), Message);
 		FMsg::Logf(__FILE__, __LINE__, LogName, ELogVerbosity::Log, *FText::Format(LOCTEXT("BasicMessageLog_Notify", "Notify: {Message}"), Arguments).ToString());
 	}
 
-	virtual void Open() OVERRIDE
+	virtual void Open() override
 	{
 		FMsg::Logf(__FILE__, __LINE__, LogName, ELogVerbosity::Log, *LOCTEXT("BasicMessageLog_Open", "Open Log").ToString());
 	}
 
-	virtual int32 NumMessages( EMessageSeverity::Type SeverityFilter ) OVERRIDE
+	virtual int32 NumMessages( EMessageSeverity::Type SeverityFilter ) override
 	{
 		return 0;
 	}
diff --git a/Engine/Source/Runtime/Core/Private/Misc/OutputDevice.cpp b/Engine/Source/Runtime/Core/Private/Misc/OutputDevice.cpp
--- a/Engine/Source/Runtime/Core/Private/Misc/OutputDevice.cpp
+++ b/Engine/Source/Runtime/Core/Private/Misc/OutputDevice.cpp
@@ -456,7 +456,7 @@ static FOutputDeviceRedirector LogRedirector;
 class FThrowOut : public FOutputDevice
 {
 public:
-	void Serialize( const TCHAR* V, ELogVerbosity::Type Verbosity, const class FName& Category ) OVERRIDE
+	void Serialize( const TCHAR* V, ELogVerbosity::Type Verbosity, const class FName& Category ) override
 	{
 #if PLATFORM_EXCEPTIONS_DISABLED
 		FPlatformMisc::DebugBreak();
diff --git a/Engine/Source/Runtime/Core/Private/ProfilingDebugging/Windows/AQtimeExternalProfiler.cpp b/Engine/Source/Runtime/Core/Private/ProfilingDebugging/Windows/AQtimeExternalProfiler.cpp
--- a/Engine/Source/Runtime/Core/Private/ProfilingDebugging/Windows/AQtimeExternalProfiler.cpp
+++ b/Engine/Source/Runtime/Core/Private/ProfilingDebugging/Windows/AQtimeExternalProfiler.cpp
@@ -37,21 +37,21 @@ public:
 
 
 	/** Gets the name of this profiler as a string.  This is used to allow the user to select this profiler in a system configuration file or on the command-line */
-	virtual const TCHAR* GetProfilerName() const OVERRIDE
+	virtual const TCHAR* GetProfilerName() const override
 	{
 		return TEXT( "AQtime" );
 	}
 
 
 	/** Pauses profiling. */
-	virtual void ProfilerPauseFunction() OVERRIDE
+	virtual void ProfilerPauseFunction() override
 	{
 		EnableProfiling( (short)0 );
 	}
 
 
 	/** Resumes profiling. */
-	virtual void ProfilerResumeFunction() OVERRIDE
+	virtual void ProfilerResumeFunction() override
 	{
 		EnableProfiling( (short)- 1 );
 	}
diff --git a/Engine/Source/Runtime/Core/Private/ProfilingDebugging/Windows/VSPerfExternalProfiler.cpp b/Engine/Source/Runtime/Core/Private/ProfilingDebugging/Windows/VSPerfExternalProfiler.cpp
--- a/Engine/Source/Runtime/Core/Private/ProfilingDebugging/Windows/VSPerfExternalProfiler.cpp
+++ b/Engine/Source/Runtime/Core/Private/ProfilingDebugging/Windows/VSPerfExternalProfiler.cpp
@@ -44,14 +44,14 @@ public:
 
 
 	/** Gets the name of this profiler as a string.  This is used to allow the user to select this profiler in a system configuration file or on the command-line */
-	virtual const TCHAR* GetProfilerName() const OVERRIDE
+	virtual const TCHAR* GetProfilerName() const override
 	{
 		return TEXT( "VSPerf" );
 	}
 
 
 	/** Pauses profiling. */
-	virtual void ProfilerPauseFunction() OVERRIDE
+	virtual void ProfilerPauseFunction() override
 	{
 		const int ProcessOrThreadId = PROFILE_CURRENTID;
 		PROFILE_COMMAND_STATUS Result = StopProfileFunction( PROFILE_GLOBALLEVEL, ProcessOrThreadId );
@@ -63,7 +63,7 @@ public:
 
 
 	/** Resumes profiling. */
-	virtual void ProfilerResumeFunction() OVERRIDE
+	virtual void ProfilerResumeFunction() override
 	{
 		const int ProcessOrThreadId = PROFILE_CURRENTID;
 		PROFILE_COMMAND_STATUS Result = StartProfileFunction( PROFILE_GLOBALLEVEL, ProcessOrThreadId );
diff --git a/Engine/Source/Runtime/Core/Private/Serialization/AsyncIOSystemBase.h b/Engine/Source/Runtime/Core/Private/Serialization/AsyncIOSystemBase.h
--- a/Engine/Source/Runtime/Core/Private/Serialization/AsyncIOSystemBase.h
+++ b/Engine/Source/Runtime/Core/Private/Serialization/AsyncIOSystemBase.h
@@ -38,7 +38,7 @@ struct CORE_API FAsyncIOSystemBase : public FIOSystem, FRunnable, FSingleThreadR
 		int64 Size, 
 		void* Dest, 
 		FThreadSafeCounter* Counter,
-		EAsyncIOPriority Priority ) OVERRIDE;
+		EAsyncIOPriority Priority ) override;
 
 	/**
 	 * Requests compressed data to be loaded async. Returns immediately.
@@ -62,7 +62,7 @@ struct CORE_API FAsyncIOSystemBase : public FIOSystem, FRunnable, FSingleThreadR
 		void* Dest, 
 		ECompressionFlags CompressionFlags, 
 		FThreadSafeCounter* Counter,
-		EAsyncIOPriority Priority ) OVERRIDE;
+		EAsyncIOPriority Priority ) override;
 
 	/**
 	 * Removes N outstanding requests from the queue and returns how many were canceled. We can't cancel
@@ -71,17 +71,17 @@ struct CORE_API FAsyncIOSystemBase : public FIOSystem, FRunnable, FSingleThreadR
 	 * @param	RequestIndices	Indices of requests to cancel.
 	 * @return	The number of requests that were canceled
 	 */
-	virtual int32 CancelRequests( uint64* RequestIndices, int32 NumIndices ) OVERRIDE;
+	virtual int32 CancelRequests( uint64* RequestIndices, int32 NumIndices ) override;
 	
 	/**
 	 * Removes all outstanding requests from the queue
 	 */
-	virtual void CancelAllOutstandingRequests() OVERRIDE;
+	virtual void CancelAllOutstandingRequests() override;
 	
 	/**
 	 * Blocks till all requests are finished and also flushes potentially open handles.
 	 */
-	virtual void BlockTillAllRequestsFinishedAndFlushHandles() OVERRIDE;
+	virtual void BlockTillAllRequestsFinishedAndFlushHandles() override;
 
 	// FRunnable interface.
 
@@ -93,12 +93,12 @@ struct CORE_API FAsyncIOSystemBase : public FIOSystem, FRunnable, FSingleThreadR
 	 *
 	 * @return True if initialization was successful, false otherwise
 	 */
-	virtual bool Init() OVERRIDE;
+	virtual bool Init() override;
 
 	/**
 	 * Called in the context of the aggregating thread to perform cleanup.
 	 */
-	virtual void Exit() OVERRIDE;
+	virtual void Exit() override;
 
 	/**
 	 * This is where all the actual loading is done. This is only called
@@ -106,38 +106,38 @@ struct CORE_API FAsyncIOSystemBase : public FIOSystem, FRunnable, FSingleThreadR
 	 *
 	 * @return always 0
 	 */
-	virtual uint32 Run() OVERRIDE;
+	virtual uint32 Run() override;
 	
 	/**
 	 * This is called if a thread is requested to terminate early.
 	 */
-	virtual void Stop() OVERRIDE;
+	virtual void Stop() override;
 	
 	/**
 	 * This is called if a thread is requested to suspend its' IO activity
 	 */
-	virtual void Suspend() OVERRIDE;
+	virtual void Suspend() override;
 
 	/**
 	 * This is called if a thread is requested to resume its' IO activity
 	 */
-	virtual void Resume() OVERRIDE;
+	virtual void Resume() override;
 
 	/**
 	 * Returns a pointer to the single threaded interface when mulithreading is disabled.
 	 */
-	virtual FSingleThreadRunnable* GetSingleThreadInterface() OVERRIDE
+	virtual FSingleThreadRunnable* GetSingleThreadInterface() override
 	{
 		return this;
 	}
 
 	// FSingleThreadRunnable interface
 
-	virtual void Tick() OVERRIDE;
+	virtual void Tick() override;
 
 	// FIOSystem Interface
 
-	virtual void TickSingleThreaded() OVERRIDE;
+	virtual void TickSingleThreaded() override;
 
 
 	/**
@@ -147,21 +147,21 @@ struct CORE_API FAsyncIOSystemBase : public FIOSystem, FRunnable, FSingleThreadR
 	 *
 	 * @param	MinPriority		Min priority of requests to fulfill
 	 */
-	virtual void SetMinPriority( EAsyncIOPriority MinPriority ) OVERRIDE;
+	virtual void SetMinPriority( EAsyncIOPriority MinPriority ) override;
 
 	/**
 	 * Give the IO system a hint that it is done with the file for now
 	 *
 	 * @param Filename File that was being async loaded from, but no longer is
 	 */
-	virtual void HintDoneWithFile(const FString& Filename) OVERRIDE;
+	virtual void HintDoneWithFile(const FString& Filename) override;
 
 	/**
 	 * The minimum read size...used to be DVD_ECC_BLOCK_SIZE
 	 *
 	 * @return Minimum read size
 	 */
-	virtual int64 MinimumReadSize() OVERRIDE;
+	virtual int64 MinimumReadSize() override;
 
 protected:
 
diff --git a/Engine/Source/Runtime/Core/Private/Stats/StatsCommand.cpp b/Engine/Source/Runtime/Core/Private/Stats/StatsCommand.cpp
--- a/Engine/Source/Runtime/Core/Private/Stats/StatsCommand.cpp
+++ b/Engine/Source/Runtime/Core/Private/Stats/StatsCommand.cpp
@@ -1305,7 +1305,7 @@ static class FStatCmd : private FSelfRegisteringExec
 {
 public:
 	/** Console commands, see embeded usage statement **/
-	virtual bool Exec( UWorld*, const TCHAR* Cmd, FOutputDevice& Ar ) OVERRIDE
+	virtual bool Exec( UWorld*, const TCHAR* Cmd, FOutputDevice& Ar ) override
 	{
 		// Block the thread as this affects external stat states now
 		return DirectStatsCommand(Cmd,true,&Ar);
diff --git a/Engine/Source/Runtime/Core/Private/UObject/UnrealNames.cpp b/Engine/Source/Runtime/Core/Private/UObject/UnrealNames.cpp
--- a/Engine/Source/Runtime/Core/Private/UObject/UnrealNames.cpp
+++ b/Engine/Source/Runtime/Core/Private/UObject/UnrealNames.cpp
@@ -1050,7 +1050,7 @@ FNameEntry* AllocateNameEntry( const void* Name, NAME_INDEX Index, FNameEntry* H
 **/
 static class FFNameExec: private FSelfRegisteringExec
 {
-	virtual bool Exec( UWorld* InWorld, const TCHAR* Cmd, FOutputDevice& Ar ) OVERRIDE
+	virtual bool Exec( UWorld* InWorld, const TCHAR* Cmd, FOutputDevice& Ar ) override
 	{
 		// Display information about the name table only rather than using HASH
 		if( FParse::Command( &Cmd, TEXT("NAMEHASH") ) )
diff --git a/Engine/Source/Runtime/Core/Private/WinRT/WinRTEvent.h b/Engine/Source/Runtime/Core/Private/WinRT/WinRTEvent.h
--- a/Engine/Source/Runtime/Core/Private/WinRT/WinRTEvent.h
+++ b/Engine/Source/Runtime/Core/Private/WinRT/WinRTEvent.h
@@ -37,7 +37,7 @@ public:
 
 public:
 
-	virtual bool Create (bool bIsManualReset = false) OVERRIDE
+	virtual bool Create (bool bIsManualReset = false) override
 	{
 		// Create the event and default it to non-signaled
 //		Event = CreateEvent(NULL, bIsManualReset, 0, NULL);
@@ -46,21 +46,21 @@ public:
 		return Event != NULL;
 	}
 
-	virtual void Trigger () OVERRIDE
+	virtual void Trigger () override
 	{
 		check(Event);
 
 		SetEvent(Event);
 	}
 
-	virtual void Reset () OVERRIDE
+	virtual void Reset () override
 	{
 		check(Event);
 
 		ResetEvent(Event);
 	}
 
-	virtual bool Wait (uint32 WaitTime) OVERRIDE;
+	virtual bool Wait (uint32 WaitTime) override;
 
 
 private:
diff --git a/Engine/Source/Runtime/Core/Private/WinRT/WinRTFile.cpp b/Engine/Source/Runtime/Core/Private/WinRT/WinRTFile.cpp
--- a/Engine/Source/Runtime/Core/Private/WinRT/WinRTFile.cpp
+++ b/Engine/Source/Runtime/Core/Private/WinRT/WinRTFile.cpp
@@ -47,27 +47,27 @@ public:
 		FileHandle = NULL;
 	}
 
-	virtual int64 Tell() OVERRIDE
+	virtual int64 Tell() override
 	{
 		check(IsValid());
 		return FileSeek(0, FILE_CURRENT);
 	}
 
-	virtual bool Seek(int64 NewPosition) OVERRIDE
+	virtual bool Seek(int64 NewPosition) override
 	{
 		check(IsValid());
 		check(NewPosition >= 0);
 		return (FileSeek(NewPosition, FILE_BEGIN) != -1);
 	}
 
-	virtual bool SeekFromEnd(int64 NewPositionRelativeToEnd = 0) OVERRIDE
+	virtual bool SeekFromEnd(int64 NewPositionRelativeToEnd = 0) override
 	{
 		check(IsValid());
 		check(NewPositionRelativeToEnd <= 0);
 		return (FileSeek(NewPositionRelativeToEnd, FILE_END) != -1);
 	}
 
-	virtual bool Read(uint8* Destination, int64 BytesToRead) OVERRIDE
+	virtual bool Read(uint8* Destination, int64 BytesToRead) override
 	{
 		check(IsValid());
 		while (BytesToRead)
@@ -86,7 +86,7 @@ public:
 		return true;
 	}
 
-	virtual bool Write(const uint8* Source, int64 BytesToWrite) OVERRIDE
+	virtual bool Write(const uint8* Source, int64 BytesToWrite) override
 	{
 		check(IsValid());
 		while (BytesToWrite)
@@ -189,7 +189,7 @@ public:
 	 * @param Inner Platform file to wrap by this file.
 	 * @param CommandLineParam Optional parameter passed in the commandline.
 	 * @return true if the initialization was successful, false otherise. */
-	virtual bool Initialize(IPlatformFile* Inner, const TCHAR* CommandLineParam) OVERRIDE
+	virtual bool Initialize(IPlatformFile* Inner, const TCHAR* CommandLineParam) override
 	{
 		// Physical platform file should never wrap anything.
 		check(Inner == NULL);
@@ -197,18 +197,18 @@ public:
 	}
 
 	/** Gets the platform file wrapped by this file. */
-	virtual IPlatformFile* GetLowerLevel() OVERRIDE
+	virtual IPlatformFile* GetLowerLevel() override
 	{
 		return NULL;
 	}
 
 	/** Gets this platform file type name. */
-	virtual const TCHAR* GetName() const OVERRIDE
+	virtual const TCHAR* GetName() const override
 	{
 		return IPlatformFile::GetPhysicalTypeName();
 	}
 
-	virtual bool FileExists(const TCHAR* Filename) OVERRIDE
+	virtual bool FileExists(const TCHAR* Filename) override
 	{
 		WIN32_FILE_ATTRIBUTE_DATA AttribData;
 		if (GetFileAttributesEx(*NormalizeFilename(Filename), GetFileExInfoStandard, (void*)&AttribData) == TRUE)
@@ -221,7 +221,7 @@ public:
 		return false;
 	}
 	
-	virtual int64 FileSize(const TCHAR* Filename) OVERRIDE
+	virtual int64 FileSize(const TCHAR* Filename) override
 	{
 		WIN32_FILE_ATTRIBUTE_DATA Info;
 		if (GetFileAttributesExW(*NormalizeFilename(Filename), GetFileExInfoStandard, &Info) == TRUE)
@@ -237,12 +237,12 @@ public:
 		return -1;
 	}
 	
-	virtual bool DeleteFile(const TCHAR* Filename) OVERRIDE
+	virtual bool DeleteFile(const TCHAR* Filename) override
 	{
 		return !!DeleteFileW(*NormalizeFilename(Filename));
 	}
 	
-	virtual bool IsReadOnly(const TCHAR* Filename) OVERRIDE
+	virtual bool IsReadOnly(const TCHAR* Filename) override
 	{
 		WIN32_FILE_ATTRIBUTE_DATA Info;
 		if (!GetFileAttributesExW(*NormalizeFilename(Filename), GetFileExInfoStandard, &Info))
@@ -253,17 +253,17 @@ public:
 		return ((Info.dwFileAttributes & FILE_ATTRIBUTE_READONLY) == 0);
 	}
 	
-	virtual bool MoveFile(const TCHAR* To, const TCHAR* From) OVERRIDE
+	virtual bool MoveFile(const TCHAR* To, const TCHAR* From) override
 	{
 		return !!MoveFileEx(*NormalizeFilename(From), *NormalizeFilename(To), MOVEFILE_REPLACE_EXISTING | MOVEFILE_WRITE_THROUGH);
 	}
 	
-	virtual bool SetReadOnly(const TCHAR* Filename, bool bNewReadOnlyValue) OVERRIDE
+	virtual bool SetReadOnly(const TCHAR* Filename, bool bNewReadOnlyValue) override
 	{
 		return !!SetFileAttributesW(*NormalizeFilename(Filename), bNewReadOnlyValue ? FILE_ATTRIBUTE_READONLY : FILE_ATTRIBUTE_NORMAL);
 	}
 	
-	virtual FDateTime GetTimeStamp(const TCHAR* Filename) OVERRIDE
+	virtual FDateTime GetTimeStamp(const TCHAR* Filename) override
 	{
 		// get file times
 		struct _stat FileInfo;
@@ -278,7 +278,7 @@ public:
 		return WinRTEpoch + TimeSinceEpoch;
 	}
 
-	virtual void SetTimeStamp(const TCHAR* Filename, FDateTime DateTime) OVERRIDE
+	virtual void SetTimeStamp(const TCHAR* Filename, FDateTime DateTime) override
 	{
 		// get file times
 		struct _stat FileInfo;
@@ -338,7 +338,7 @@ public:
 		return WinRTEpoch + TimeSinceEpoch;
 	}
 
-	virtual IFileHandle* OpenRead(const TCHAR* Filename) OVERRIDE
+	virtual IFileHandle* OpenRead(const TCHAR* Filename) override
 	{
 		DWORD  Access    = GENERIC_READ;
 		DWORD  WinFlags  = FILE_SHARE_READ;
@@ -352,7 +352,7 @@ public:
 		return NULL;
 	}
 
-	virtual IFileHandle* OpenWrite(const TCHAR* Filename, bool bAppend = 0, bool bAllowRead = 0) OVERRIDE
+	virtual IFileHandle* OpenWrite(const TCHAR* Filename, bool bAppend = 0, bool bAllowRead = 0) override
 	{
 		DWORD  Access    = GENERIC_WRITE;
 		DWORD  WinFlags  = bAllowRead ? FILE_SHARE_READ : 0;
@@ -365,7 +365,7 @@ public:
 		return NULL;
 	}
 	
-	virtual bool DirectoryExists(const TCHAR* Directory) OVERRIDE
+	virtual bool DirectoryExists(const TCHAR* Directory) override
 	{
 		WIN32_FILE_ATTRIBUTE_DATA AttribData;
 		if (GetFileAttributesEx(*NormalizeDirectory(Directory), GetFileExInfoStandard, (void*)&AttribData) == TRUE)
@@ -378,12 +378,12 @@ public:
 		return false;
 	}
 	
-	virtual bool CreateDirectory(const TCHAR* Directory) OVERRIDE
+	virtual bool CreateDirectory(const TCHAR* Directory) override
 	{
 		return CreateDirectoryW(*NormalizeDirectory(Directory), NULL) || GetLastError() == ERROR_ALREADY_EXISTS;
 	}
 	
-	virtual bool DeleteDirectory(const TCHAR* Directory) OVERRIDE
+	virtual bool DeleteDirectory(const TCHAR* Directory) override
 	{
 		RemoveDirectoryW(*NormalizeDirectory(Directory));
 		return !DirectoryExists(Directory);
diff --git a/Engine/Source/Runtime/Core/Private/WinRT/WinRTOutputDevicesPrivate.h b/Engine/Source/Runtime/Core/Private/WinRT/WinRTOutputDevicesPrivate.h
--- a/Engine/Source/Runtime/Core/Private/WinRT/WinRTOutputDevicesPrivate.h
+++ b/Engine/Source/Runtime/Core/Private/WinRT/WinRTOutputDevicesPrivate.h
@@ -56,7 +56,7 @@ public:
 		TearDown();
 	}
 
-	virtual void Serialize(const TCHAR* Buffer, ELogVerbosity::Type Verbosity, const class FName& Category) OVERRIDE
+	virtual void Serialize(const TCHAR* Buffer, ELogVerbosity::Type Verbosity, const class FName& Category) override
 	{
 		if (EventLog != NULL)
 		{
@@ -124,7 +124,7 @@ public:
 	 * @param	Data	Text to log
 	 * @param	Event	Event name used for suppression purposes
 	 */
-	virtual void Serialize( const TCHAR* Msg, ELogVerbosity::Type Verbosity, const class FName& Category ) OVERRIDE;
+	virtual void Serialize( const TCHAR* Msg, ELogVerbosity::Type Verbosity, const class FName& Category ) override;
 
 	/**
 	 * Error handling function that is being called from within the system wide global
diff --git a/Engine/Source/Runtime/Core/Private/WinRT/WinRTRunnableThread.h b/Engine/Source/Runtime/Core/Private/WinRT/WinRTRunnableThread.h
--- a/Engine/Source/Runtime/Core/Private/WinRT/WinRTRunnableThread.h
+++ b/Engine/Source/Runtime/Core/Private/WinRT/WinRTRunnableThread.h
@@ -144,7 +144,7 @@ public:
 		FRunnableThread::GetThreadRegistry().Remove(ThreadID);
 	}
 	
-	virtual void SetThreadPriority(EThreadPriority NewPriority) OVERRIDE
+	virtual void SetThreadPriority(EThreadPriority NewPriority) override
 	{
 		// Don't bother calling the OS if there is no need
 		if (NewPriority != ThreadPriority)
@@ -158,7 +158,7 @@ public:
 		}
 	}
 
-	virtual void Suspend(bool bShouldPause = 1) OVERRIDE
+	virtual void Suspend(bool bShouldPause = 1) override
 	{
 		check(Thread);
 		if (bShouldPause == true)
@@ -172,7 +172,7 @@ public:
 		}
 	}
 
-	virtual bool Kill(bool bShouldWait = false) OVERRIDE
+	virtual bool Kill(bool bShouldWait = false) override
 	{
 		check(Thread && "Did you forget to call Create()?");
 		bool bDidExitOK = true;
@@ -200,19 +200,19 @@ public:
 		return bDidExitOK;
 	}
 
-	virtual void WaitForCompletion() OVERRIDE
+	virtual void WaitForCompletion() override
 	{
 		// Block until this thread exits
 		//WaitForSingleObject(Thread,INFINITE);
 		WaitForSingleObjectEx(Thread, INFINITE, FALSE);
 	}
 
-	virtual uint32 GetThreadID() OVERRIDE
+	virtual uint32 GetThreadID() override
 	{
 		return ThreadID;
 	}
 
-	virtual FString GetThreadName() OVERRIDE
+	virtual FString GetThreadName() override
 	{
 		return ThreadName;
 	}
@@ -220,7 +220,7 @@ public:
 protected:
 	virtual bool CreateInternal(FRunnable* InRunnable, const TCHAR* InThreadName,
 		uint32 InStackSize = 0,
-		EThreadPriority InThreadPri = TPri_Normal, uint64 InThreadAffinityMask = 0) OVERRIDE
+		EThreadPriority InThreadPri = TPri_Normal, uint64 InThreadAffinityMask = 0) override
 	{
 		check(InRunnable);
 		Runnable = InRunnable;
