Diffed assets no longer appear in the asset picker

#ttp 337409 - Editor: Temporary diff assets show up in asset picker

#branch UE4

#change Diffed packages now get loaded with the ForDiff flag, which allows us to filter them out later on.

diff --git a/Engine/Source/Developer/AssetTools/Private/AssetTools.cpp b/Engine/Source/Developer/AssetTools/Private/AssetTools.cpp
--- a/Engine/Source/Developer/AssetTools/Private/AssetTools.cpp
+++ b/Engine/Source/Developer/AssetTools/Private/AssetTools.cpp
@@ -893,7 +893,7 @@ void FAssetTools::DiffAgainstDepot( UObject* InObject, const FString& InPackageP
 					TGuardValue<bool> DisableCompileOnLoad(GForceDisableBlueprintCompileOnLoad, true);
 
 					// Try and load that package
-					UPackage* TempPackage = LoadPackage(NULL, *TempFileName, LOAD_None);
+					UPackage* TempPackage = LoadPackage(NULL, *TempFileName, LOAD_ForDiff);
 					if(TempPackage != NULL)
 					{
 						// Grab the old asset from that old package
diff --git a/Engine/Source/Runtime/AssetRegistry/Private/AssetRegistry.cpp b/Engine/Source/Runtime/AssetRegistry/Private/AssetRegistry.cpp
--- a/Engine/Source/Runtime/AssetRegistry/Private/AssetRegistry.cpp
+++ b/Engine/Source/Runtime/AssetRegistry/Private/AssetRegistry.cpp
@@ -307,7 +307,6 @@ bool FAssetRegistry::GetAssets(const FARFilter& Filter, TArray<FAssetData>& OutA
 			if ( ObjIt->IsAsset() )
 			{
 				UPackage* InMemoryPackage = ObjIt->GetOutermost();
-				const FName ObjectPath = FName(*ObjIt->GetPathName());
 
 				static const bool bUsingWorldAssets = FParse::Param(FCommandLine::Get(), TEXT("WorldAssets"));
 				// Skip assets in map packages... unless we are showing world assets
@@ -316,7 +315,14 @@ bool FAssetRegistry::GetAssets(const FARFilter& Filter, TArray<FAssetData>& OutA
 					continue;
 				}
 
+				// Skip assets that were loaded for diffing
+				if ( InMemoryPackage->PackageFlags & PKG_ForDiffing )
+				{
+					continue;
+				}
+
 				// add it to in-memory object list for later merge
+				const FName ObjectPath = FName(*ObjIt->GetPathName());
 				InMemoryObjectPaths.Add(ObjectPath);
 			
 				// Package name
