Added support for all BlueprintCore-derived Blueprint classes in AssetRegistry

diff --git a/Engine/Source/Runtime/AssetRegistry/Private/AssetRegistry.cpp b/Engine/Source/Runtime/AssetRegistry/Private/AssetRegistry.cpp
--- a/Engine/Source/Runtime/AssetRegistry/Private/AssetRegistry.cpp
+++ b/Engine/Source/Runtime/AssetRegistry/Private/AssetRegistry.cpp
@@ -92,6 +92,25 @@ FAssetRegistry::FAssetRegistry()
 	// Listen for new content paths being added at runtime.  These are usually plugin-specific asset paths that
 	// will be loaded a bit later on.
 	FPackageName::OnContentPathMounted().AddRaw( this, &FAssetRegistry::OnContentPathMounted );
+
+	// Now collect all code generator classes (currently BlueprintCore-derived ones)
+	CollectCodeGeneratorClasses();
+}
+
+void FAssetRegistry::CollectCodeGeneratorClasses()
+{
+	// Work around the fact we don't reference Engine module directly
+	UClass* BlueprintCoreClass = Cast<UClass>(StaticFindObject(UClass::StaticClass(), ANY_PACKAGE, TEXT("BlueprintCore")));
+	if (BlueprintCoreClass)
+	{
+		for (TObjectIterator<UClass> It; It; ++It)
+		{
+			if (It->IsChildOf(BlueprintCoreClass))
+			{
+				ClassGeneratorNames.Add(It->GetFName());
+			}
+		}
+	}
 }
 
 FAssetRegistry::~FAssetRegistry()
@@ -1548,7 +1567,7 @@ void FAssetRegistry::AddAssetData(FAssetData* AssetData)
 	AssetAddedEvent.Broadcast(*AssetData);
 
 	// Populate the class map if adding blueprint
-	if ( AssetData->AssetClass == FName("Blueprint") )
+	if (ClassGeneratorNames.Contains(AssetData->AssetClass))
 	{
 		FString* GeneratedClassPtr = AssetData->TagsAndValues.Find("GeneratedClass");
 		FString* ParentClassPtr = AssetData->TagsAndValues.Find("ParentClass");
@@ -1637,7 +1656,7 @@ void FAssetRegistry::UpdateAssetData(FAssetData* AssetData, const FAssetData& Ne
 	}
 
 	// Update the class map if updating a blueprint
-	if ( AssetData->AssetClass == FName("Blueprint") )
+	if (ClassGeneratorNames.Contains(AssetData->AssetClass))
 	{
 		FString* OldGeneratedClassPtr = AssetData->TagsAndValues.Find("GeneratedClass");
 
@@ -1671,7 +1690,7 @@ bool FAssetRegistry::RemoveAssetData(FAssetData* AssetData)
 		AssetRemovedEvent.Broadcast(*AssetData);
 
 		// Remove from the class map if removing a blueprint
-		if ( AssetData->AssetClass == FName("Blueprint") )
+		if (ClassGeneratorNames.Contains(AssetData->AssetClass))
 		{
 			FString* OldGeneratedClassPtr = AssetData->TagsAndValues.Find("GeneratedClass");
 
diff --git a/Engine/Source/Runtime/AssetRegistry/Private/AssetRegistry.h b/Engine/Source/Runtime/AssetRegistry/Private/AssetRegistry.h
--- a/Engine/Source/Runtime/AssetRegistry/Private/AssetRegistry.h
+++ b/Engine/Source/Runtime/AssetRegistry/Private/AssetRegistry.h
@@ -148,6 +148,9 @@ private:
 	void GetSubClasses(const TArray<FName>& InClassNames, const TSet<FName>& ExcludedClassNames, TSet<FName>& SubClassNames) const;
 	void GetSubClasses_Recursive(FName InClassName, TSet<FName>& SubClassNames, const TMap<FName, TSet<FName>>& ReverseInheritanceMap, const TSet<FName>& ExcludedClassNames) const;	
 
+	/** Finds all class names of classes capable of generating new UClasses */
+	void CollectCodeGeneratorClasses();
+
 private:
 	/** The map of ObjectPath names to asset data for assets saved to disk */
 	TMap<FName, FAssetData*> CachedAssetsByObjectPath;
@@ -234,4 +237,7 @@ private:
 
 	/** A set used to ignore repeated requests to synchronously scan the same folder multiple times */
 	TSet<FString> SynchronouslyScannedPaths;
+
+	/** List of all class names derived from Blueprint (including Blueprint itself) */
+	TSet<FName> ClassGeneratorNames;
 };
