Fixed issue where empty packages with metadata were persisting until shutdown, incorrectly prompting the user on exit that they need to be saved.
#ttp 340045 - Editor: After deleting a newly-created asset that had not been saved, I was prompted to save the deleted asset on exit
#branch UE4
#proj Runtime.AssetRegistry
#add In FAssetRegistry::AssetDeleted, when deleting an asset results in its package becoming empty, clear the RF_Standalone flag on any metadata object it may have, so that this can be GCd next time and stop the package from persisting.

diff --git a/Engine/Source/Runtime/AssetRegistry/Private/AssetRegistry.cpp b/Engine/Source/Runtime/AssetRegistry/Private/AssetRegistry.cpp
--- a/Engine/Source/Runtime/AssetRegistry/Private/AssetRegistry.cpp
+++ b/Engine/Source/Runtime/AssetRegistry/Private/AssetRegistry.cpp
@@ -1152,6 +1152,12 @@ void FAssetRegistry::AssetDeleted(UObject* DeletedAsset)
 			if ( UPackage::IsEmptyPackage( DeletedObjectPackage, DeletedAsset ) )
 			{
 				AddEmptyPackage( DeletedObjectPackage->GetFName() );
+
+				// If there is a package metadata object, clear the standalone flag so the package can be truly emptied upon GC
+				if ( DeletedObjectPackage->MetaData != nullptr )
+				{
+					DeletedObjectPackage->MetaData->ClearFlags( RF_Standalone );
+				}
 			}
 		}
 
