UE4: Fixed a bug where not all assets would be removed from the disk cache in the asset registry when deleting a package containing multiple assets. Also, worlds no longer create redirectors when using save-as. TTP#345100

diff --git a/Engine/Source/Editor/UnrealEd/Private/FileHelpers.cpp b/Engine/Source/Editor/UnrealEd/Private/FileHelpers.cpp
--- a/Engine/Source/Editor/UnrealEd/Private/FileHelpers.cpp
+++ b/Engine/Source/Editor/UnrealEd/Private/FileHelpers.cpp
@@ -490,12 +490,12 @@ static bool SaveWorld(UWorld* World,
 		{
 			if ( bPackageNeedsRename )
 			{
-				Package->Rename(*NewPackageName, NULL, REN_NonTransactional);
+				Package->Rename(*NewPackageName, NULL, REN_NonTransactional | REN_DontCreateRedirectors);
 			}
 
 			if ( bWorldNeedsRename )
 			{
-				World->Rename(*NewWorldAssetName, NULL, REN_NonTransactional);
+				World->Rename(*NewWorldAssetName, NULL, REN_NonTransactional | REN_DontCreateRedirectors);
 			}
 		}
 
diff --git a/Engine/Source/Runtime/AssetRegistry/Private/AssetRegistry.cpp b/Engine/Source/Runtime/AssetRegistry/Private/AssetRegistry.cpp
--- a/Engine/Source/Runtime/AssetRegistry/Private/AssetRegistry.cpp
+++ b/Engine/Source/Runtime/AssetRegistry/Private/AssetRegistry.cpp
@@ -1834,13 +1834,14 @@ void FAssetRegistry::OnDirectoryChanged (const TArray<FFileChangeData>& FileChan
 					{
 						// This file was deleted. Remove all assets in the package from the registry.
 						const FName PackageName = FName(*LongPackageName);
-						TArray<FAssetData*>* PackageAssets = CachedAssetsByPackageName.Find(PackageName);
+						TArray<FAssetData*>* PackageAssetsPtr = CachedAssetsByPackageName.Find(PackageName);
 
-						if (PackageAssets)
+						if (PackageAssetsPtr)
 						{
-							for ( auto AssetIt = (*PackageAssets).CreateConstIterator(); AssetIt; ++AssetIt )
+							TArray<FAssetData*>& PackageAssets = *PackageAssetsPtr;
+							for ( int32 AssetIdx = PackageAssets.Num() - 1; AssetIdx >= 0; --AssetIdx )
 							{
-								RemoveAssetData(*AssetIt);
+								RemoveAssetData(PackageAssets[AssetIdx]);
 							}
 						}
 
