UE4: Added support for removing dynamic mount points. Added more information to the OnContentPathMounted delegate as well.

diff --git a/Engine/Source/Editor/ContentBrowser/Private/SPathView.cpp b/Engine/Source/Editor/ContentBrowser/Private/SPathView.cpp
--- a/Engine/Source/Editor/ContentBrowser/Private/SPathView.cpp
+++ b/Engine/Source/Editor/ContentBrowser/Private/SPathView.cpp
@@ -14,6 +14,7 @@ SPathView::~SPathView()
 {
 	// Unsubscribe from content path events
 	FPackageName::OnContentPathMounted().RemoveAll( this );
+	FPackageName::OnContentPathDismounted().RemoveAll( this );
 
 	// Load the asset registry module to stop listening for updates
 	FAssetRegistryModule& AssetRegistryModule = FModuleManager::LoadModuleChecked<FAssetRegistryModule>(TEXT("AssetRegistry"));
@@ -43,8 +44,9 @@ void SPathView::Construct( const FArguments& InArgs )
 	SearchBoxFolderFilter = MakeShareable( new FolderTextFilter( FolderTextFilter::FItemToStringArray::CreateSP( this, &SPathView::PopulateFolderSearchStrings ) ) );
 	SearchBoxFolderFilter->OnChanged().AddSP( this, &SPathView::FilterUpdated );
 
-	// Listen to find out when new game content paths are mounted, so that we can refresh our root set of paths
-	FPackageName::OnContentPathMounted().AddSP( this, &SPathView::OnContentPathMounted );
+	// Listen to find out when new game content paths are mounted or dismounted, so that we can refresh our root set of paths
+	FPackageName::OnContentPathMounted().AddSP( this, &SPathView::OnContentPathMountedOrDismounted );
+	FPackageName::OnContentPathDismounted().AddSP( this, &SPathView::OnContentPathMountedOrDismounted );
 
 	ClassesRootName = TEXT("Classes");
 	GameRootName = TEXT("Game");
@@ -1324,7 +1326,7 @@ void SPathView::OnAssetRegistrySearchCompleted()
 	PendingInitialPaths.Empty();
 }
 
-void SPathView::OnContentPathMounted( const FString& ContentPath )
+void SPathView::OnContentPathMountedOrDismounted( const FString& AssetPath, const FString& FilesystemPath )
 {
 	// A new content path has appeared, so we should refresh out root set of paths
 	bNeedsRepopulate = true;
diff --git a/Engine/Source/Editor/ContentBrowser/Private/SPathView.h b/Engine/Source/Editor/ContentBrowser/Private/SPathView.h
--- a/Engine/Source/Editor/ContentBrowser/Private/SPathView.h
+++ b/Engine/Source/Editor/ContentBrowser/Private/SPathView.h
@@ -219,8 +219,8 @@ private:
 	/** Notification for when the Asset Registry has completed it's initial search */
 	void OnAssetRegistrySearchCompleted();
 
-	/** Called from an engine core event when a new content path has been added, so that we can refresh our root set of paths */
-	void OnContentPathMounted( const FString& ContentPath );
+	/** Called from an engine core event when a new content path has been added or removed, so that we can refresh our root set of paths */
+	void OnContentPathMountedOrDismounted( const FString& AssetPath, const FString& FileSystemPath );
 
 	/** Delegate called when an editor setting is changed */
 	void HandleSettingChanged(FName PropertyName);
diff --git a/Engine/Source/Runtime/AssetRegistry/Private/AssetRegistry.cpp b/Engine/Source/Runtime/AssetRegistry/Private/AssetRegistry.cpp
--- a/Engine/Source/Runtime/AssetRegistry/Private/AssetRegistry.cpp
+++ b/Engine/Source/Runtime/AssetRegistry/Private/AssetRegistry.cpp
@@ -89,9 +89,10 @@ FAssetRegistry::FAssetRegistry()
 	}
 #endif // WITH_EDITOR
 
-	// Listen for new content paths being added at runtime.  These are usually plugin-specific asset paths that
+	// Listen for new content paths being added or removed at runtime.  These are usually plugin-specific asset paths that
 	// will be loaded a bit later on.
 	FPackageName::OnContentPathMounted().AddRaw( this, &FAssetRegistry::OnContentPathMounted );
+	FPackageName::OnContentPathDismounted().AddRaw( this, &FAssetRegistry::OnContentPathDismounted );
 
 	// Now collect all code generator classes (currently BlueprintCore-derived ones)
 	CollectCodeGeneratorClasses();
@@ -168,6 +169,7 @@ FAssetRegistry::~FAssetRegistry()
 
 	// Stop listening for content mount point events
 	FPackageName::OnContentPathMounted().RemoveAll( this );
+	FPackageName::OnContentPathDismounted().RemoveAll( this );
 
 	// Commandlets dont listen for directory changes
 #if WITH_EDITOR
@@ -1538,15 +1540,18 @@ bool FAssetRegistry::AddAssetPath(const FString& PathToAdd)
 	return false;
 }
 
-bool FAssetRegistry::RemoveAssetPath(const FString& PathToRemove)
+bool FAssetRegistry::RemoveAssetPath(const FString& PathToRemove, bool bEvenIfAssetsStillExist)
 {
-	// Check if there were assets in the specified folder. You can not remove paths that still contain assets
-	TArray<FAssetData> AssetsInPath;
-	GetAssetsByPath(FName(*PathToRemove), AssetsInPath, true);
-	if ( AssetsInPath.Num() > 0 )
+	if ( !bEvenIfAssetsStillExist )
 	{
-		// At least one asset still exists in the path. Fail the remove.
-		return false;
+		// Check if there were assets in the specified folder. You can not remove paths that still contain assets
+		TArray<FAssetData> AssetsInPath;
+		GetAssetsByPath(FName(*PathToRemove), AssetsInPath, true);
+		if ( AssetsInPath.Num() > 0 )
+		{
+			// At least one asset still exists in the path. Fail the remove.
+			return false;
+		}
 	}
 
 	if ( PathTreeRoot.RemoveFolder(PathToRemove) )
@@ -1861,8 +1866,16 @@ void FAssetRegistry::OnDirectoryChanged (const TArray<FFileChangeData>& FileChan
 #endif // WITH_EDITOR
 
 
-void FAssetRegistry::OnContentPathMounted( const FString& AssetPath )
+void FAssetRegistry::OnContentPathMounted( const FString& InAssetPath, const FString& FileSystemPath )
 {
+	// Sanitize
+	FString AssetPath = InAssetPath;
+	if (AssetPath.EndsWith(TEXT("/")) == false)
+	{
+		// We actually want a trailing slash here so the path can be properly converted while searching for assets
+		AssetPath = AssetPath + TEXT("/");
+	}
+
 	// Add this to our list of root paths to process
 	AddPathToSearch( AssetPath );
 
@@ -1875,8 +1888,63 @@ void FAssetRegistry::OnContentPathMounted( const FString& AssetPath )
 		IDirectoryWatcher* DirectoryWatcher = DirectoryWatcherModule.Get();
 		if (DirectoryWatcher)
 		{
-			const FString& ContentFolder = FPackageName::LongPackageNameToFilename( AssetPath );
-			DirectoryWatcher->RegisterDirectoryChangedCallback( ContentFolder, IDirectoryWatcher::FDirectoryChanged::CreateRaw(this, &FAssetRegistry::OnDirectoryChanged));
+			// If the path doesn't exist on disk, make it so the watcher will work.
+			IFileManager::Get().MakeDirectory(*FileSystemPath);
+			DirectoryWatcher->RegisterDirectoryChangedCallback( FileSystemPath, IDirectoryWatcher::FDirectoryChanged::CreateRaw(this, &FAssetRegistry::OnDirectoryChanged));
+		}
+	}
+#endif // WITH_EDITOR
+
+}
+
+void FAssetRegistry::OnContentPathDismounted(const FString& InAssetPath, const FString& FileSystemPath)
+{
+	// Sanitize
+	FString AssetPath = InAssetPath;
+	if ( AssetPath.EndsWith(TEXT("/")) )
+	{
+		// We don't want a trailing slash here as it could interfere with RemoveAssetPath
+		AssetPath = AssetPath.LeftChop(1);
+	}
+
+	// Remove all cached assets found at this location
+	{
+		TArray<FAssetData*> AllAssetDataToRemove;
+		TArray<FString> PathList;
+		const bool bRecurse = true;
+		GetSubPaths(AssetPath, PathList, bRecurse);
+		PathList.Add(AssetPath);
+		for ( const FString& Path : PathList )
+		{
+			TArray<FAssetData*>* AssetsInPath = CachedAssetsByPath.Find(FName(*Path));
+			if ( AssetsInPath )
+			{
+				AllAssetDataToRemove.Append(*AssetsInPath);
+			}
+		}
+
+		for ( FAssetData* AssetData : AllAssetDataToRemove )
+		{
+			RemoveAssetData(AssetData);
+		}
+	}
+
+	// Remove the root path
+	{
+		const bool bEvenIfAssetsStillExist = true;
+		RemoveAssetPath(AssetPath, bEvenIfAssetsStillExist);
+	}
+
+	// Stop listening for directory changes in this content path
+#if WITH_EDITOR
+	// Commandlets and in-game don't listen for directory changes
+	if (!IsRunningCommandlet() && GIsEditor)
+	{
+		FDirectoryWatcherModule& DirectoryWatcherModule = FModuleManager::LoadModuleChecked<FDirectoryWatcherModule>(TEXT("DirectoryWatcher"));
+		IDirectoryWatcher* DirectoryWatcher = DirectoryWatcherModule.Get();
+		if (DirectoryWatcher)
+		{
+			DirectoryWatcher->UnregisterDirectoryChangedCallback(FileSystemPath, IDirectoryWatcher::FDirectoryChanged::CreateRaw(this, &FAssetRegistry::OnDirectoryChanged));
 		}
 	}
 #endif // WITH_EDITOR
diff --git a/Engine/Source/Runtime/AssetRegistry/Private/AssetRegistry.h b/Engine/Source/Runtime/AssetRegistry/Private/AssetRegistry.h
--- a/Engine/Source/Runtime/AssetRegistry/Private/AssetRegistry.h
+++ b/Engine/Source/Runtime/AssetRegistry/Private/AssetRegistry.h
@@ -107,7 +107,7 @@ private:
 	bool AddAssetPath(const FString& PathToAdd);
 
 	/** Removes a path to the cached paths tree. Returns true if successful. */
-	bool RemoveAssetPath(const FString& PathToAdd);
+	bool RemoveAssetPath(const FString& PathToAdd, bool bEvenIfAssetsStillExist = false);
 
 	/** Helper function to return the name of an object, given the objects export text path */
 	FString ExportTextPathToObjectName(const FString& InExportTextPath) const;
@@ -140,9 +140,19 @@ private:
 	 * Called by the engine core when a new content path is added dynamically at runtime.  This is wired to 
 	 * FPackageName's static delegate.
 	 *
-	 * @param	AssetPath	The new content root asset path that was added (e.g. "/MyPlugin/")
+	 * @param	AssetPath		The new content root asset path that was added (e.g. "/MyPlugin/")
+	 * @param	FileSystemPath	The filesystem path that the AssetPath is mapped to
 	 */
-	void OnContentPathMounted( const FString& AssetPath );
+	void OnContentPathMounted( const FString& AssetPath, const FString& FileSystemPath );
+
+	/**
+	 * Called by the engine core when a content path is removed dynamically at runtime.  This is wired to 
+	 * FPackageName's static delegate.
+	 *
+	 * @param	AssetPath		The new content root asset path that was added (e.g. "/MyPlugin/")
+	 * @param	FileSystemPath	The filesystem path that the AssetPath is mapped to
+	 */
+	void OnContentPathDismounted( const FString& AssetPath, const FString& FileSystemPath );
 
 	/** Checks a filter to make sure there are no illegal entries */
 	bool IsFilterValid(const FARFilter& Filter) const;
diff --git a/Engine/Source/Runtime/CoreUObject/Private/Misc/PackageName.cpp b/Engine/Source/Runtime/CoreUObject/Private/Misc/PackageName.cpp
--- a/Engine/Source/Runtime/CoreUObject/Private/Misc/PackageName.cpp
+++ b/Engine/Source/Runtime/CoreUObject/Private/Misc/PackageName.cpp
@@ -15,6 +15,9 @@ FString FPackageName::MapPackageExtension = TEXT(".umap");
 /** Event that is triggered when a new content path is mounted */
 FPackageName::FOnContentPathMountedEvent FPackageName::OnContentPathMountedEvent;
 
+/** Event that is triggered when a content path is dismounted */
+FPackageName::FOnContentPathDismountedEvent FPackageName::OnContentPathDismountedEvent;
+
 namespace PackageNameConstants
 {
 	// Minimum theoretical package name length ("/A/B") is 4
@@ -78,6 +81,11 @@ struct FPathPair
 	// The physical relative path (e.g., "../../../Engine/Content/")
 	const FString ContentPath;
 
+	bool operator ==(const FPathPair& Other) const
+	{
+		return RootPath == Other.RootPath && ContentPath == Other.ContentPath;
+	}
+
 	// Construct a path pair
 	FPathPair(const FString& InRootPath, const FString& InContentPath)
 		: RootPath(InRootPath)
@@ -147,8 +155,33 @@ struct FLongPackagePathsSingleton
 		MountPointRootPaths.Add( RootPath );
 
 		// Let subscribers know that a new content path was mounted
-		FPackageName::OnContentPathMounted().Broadcast( RootPath );
+		FPackageName::OnContentPathMounted().Broadcast( RootPath, ContentPath );
 	}
+
+	// This will remove a previously inserted mount point
+	void RemoveMountPoint(const FString& RootPath, const FString& ContentPath)
+	{
+		// Make sure the content path is stored as a relative path, consistent with the other paths we have
+		FString RelativeContentPath = IFileManager::Get().ConvertToRelativePath(*ContentPath);
+
+		// Make sure the path ends in a trailing path separator.  We are expecting that in the InternalFilenameToLongPackageName code.
+		if (!RelativeContentPath.EndsWith(TEXT("/")))
+		{
+			RelativeContentPath += TEXT("/");
+		}
+
+		if ( MountPointRootPaths.Remove(RootPath) > 0 )
+		{
+			FPathPair Pair(RootPath, RelativeContentPath);
+			ContentRootToPath.Remove(Pair);
+			ContentPathToRoot.Remove(Pair);
+			MountPointRootPaths.Remove(RootPath);
+
+			// Let subscribers know that a new content path was mounted
+			FPackageName::OnContentPathDismounted().Broadcast( RootPath, ContentPath );
+		}
+	}
+
 private:
 	FLongPackagePathsSingleton()
 	{
@@ -471,6 +504,12 @@ void FPackageName::RegisterMountPoint(const FString& RootPath, const FString& Co
 	Paths.InsertMountPoint(RootPath, ContentPath);
 }
 
+void FPackageName::UnRegisterMountPoint(const FString& RootPath, const FString& ContentPath)
+{
+	FLongPackagePathsSingleton& Paths = FLongPackagePathsSingleton::Get();
+	Paths.RemoveMountPoint(RootPath, ContentPath);
+}
+
 FName FPackageName::GetPackageMountPoint(const FString& InPackagePath)
 {
 	FLongPackagePathsSingleton& Paths = FLongPackagePathsSingleton::Get();
diff --git a/Engine/Source/Runtime/CoreUObject/Public/Misc/PackageName.h b/Engine/Source/Runtime/CoreUObject/Public/Misc/PackageName.h
--- a/Engine/Source/Runtime/CoreUObject/Public/Misc/PackageName.h
+++ b/Engine/Source/Runtime/CoreUObject/Public/Misc/PackageName.h
@@ -161,6 +161,14 @@ public:
 	 */
 	static void RegisterMountPoint(const FString& RootPath, const FString& ContentPath);
 
+	/**
+	 * This will remove a previously inserted mount point.
+	 *
+	 * @param RootPath Root Path.
+	 * @param ContentPath Content Path.
+	 */
+	static void UnRegisterMountPoint(const FString& RootPath, const FString& ContentPath);
+
 	/**
 	 * Get the mount point for a given package path
 	 * 
@@ -246,12 +254,19 @@ public:
 	static bool FindPackagesInDirectory(TArray<FString>& OutPackages, const FString& RootDir);
 
 	/** Event that is triggered when a new content path is mounted */
-	DECLARE_MULTICAST_DELEGATE_OneParam( FOnContentPathMountedEvent, const FString& /* Asset path */ );
+	DECLARE_MULTICAST_DELEGATE_TwoParams( FOnContentPathMountedEvent, const FString& /* Asset path */, const FString& /* ContentPath */ );
 	static FOnContentPathMountedEvent& OnContentPathMounted()
 	{
 		return OnContentPathMountedEvent;
 	}
 
+	/** Event that is triggered when a new content path is removed */
+	DECLARE_MULTICAST_DELEGATE_TwoParams(FOnContentPathDismountedEvent, const FString& /* Asset path */, const FString& /* ContentPath */ );
+	static FOnContentPathDismountedEvent& OnContentPathDismounted()
+	{
+		return OnContentPathDismountedEvent;
+	}
+
 	/**
 	 * Queries all of the root content paths, like "/Game/", "/Engine/", and any dynamically added paths
 	 *
@@ -324,6 +339,9 @@ private:
 
 	/** Event that is triggered when a new content path is mounted */
 	static FOnContentPathMountedEvent OnContentPathMountedEvent;
+
+	/** Event that is triggered when a new content path is removed */
+	static FOnContentPathDismountedEvent OnContentPathDismountedEvent;
 };
 
 #endif	// __PACKAGENAME_H__
