UE4: Clicking on a folder in any SPathView in the editor will now prioritize discovering assets in that folder. Also fixed a performance problem when discovering paths in the project initially.

diff --git a/Engine/Source/Editor/ContentBrowser/Private/SPathView.cpp b/Engine/Source/Editor/ContentBrowser/Private/SPathView.cpp
--- a/Engine/Source/Editor/ContentBrowser/Private/SPathView.cpp
+++ b/Engine/Source/Editor/ContentBrowser/Private/SPathView.cpp
@@ -826,6 +826,13 @@ void SPathView::TreeSelectionChanged( TSharedPtr< FTreeItem > TreeItem, ESelectI
 			}
 		}
 	}
+
+	if (TreeItem.IsValid())
+	{
+		// Prioritize the asset registry scan for the selected path
+		FAssetRegistryModule& AssetRegistryModule = FModuleManager::LoadModuleChecked<FAssetRegistryModule>(TEXT("AssetRegistry"));
+		AssetRegistryModule.Get().PrioritizeSearchPath(TreeItem->FolderPath);
+	}
 }
 
 void SPathView::TreeExpansionChanged( TSharedPtr< FTreeItem > TreeItem, bool bIsExpanded )
diff --git a/Engine/Source/Runtime/AssetRegistry/Private/AssetDataGatherer.cpp b/Engine/Source/Runtime/AssetRegistry/Private/AssetDataGatherer.cpp
--- a/Engine/Source/Runtime/AssetRegistry/Private/AssetDataGatherer.cpp
+++ b/Engine/Source/Runtime/AssetRegistry/Private/AssetDataGatherer.cpp
@@ -311,12 +311,36 @@ void FAssetDataGatherer::AddFilesToSearch(const TArray<FString>& Files)
 	}
 }
 
+void FAssetDataGatherer::PrioritizeSearchPath(const FString& PathToPrioritize)
+{
+	const bool bIncludeReadOnlyRoots = true;
+	if ( FPackageName::IsValidLongPackageName(PathToPrioritize, bIncludeReadOnlyRoots) )
+	{
+		const FString FilenamePathToPrioritize = FPackageName::LongPackageNameToFilename(PathToPrioritize);
+
+		// Critical section. This code needs to be as fast as possible since it is in a critical section!
+		// Swap all priority files to the top of the list
+		{
+			FScopeLock CritSectionLock(&WorkerThreadCriticalSection);
+			int32 LowestNonPriorityFileIdx = 0;
+			for ( int32 FilenameIdx = 0; FilenameIdx < FilesToSearch.Num(); ++FilenameIdx )
+			{
+				if ( FilesToSearch[FilenameIdx].StartsWith(FilenamePathToPrioritize) )
+				{
+					FilesToSearch.Swap(FilenameIdx, LowestNonPriorityFileIdx);
+					LowestNonPriorityFileIdx++;
+				}
+			}
+		}
+	}
+}
+
 void FAssetDataGatherer::DiscoverFilesToSearch()
 {
 	if( PathsToSearch.Num() > 0 )
 	{
 		TArray<FString> DiscoveredFilesToSearch;
-		TArray<FString> LocalDiscoveredPaths;
+		TSet<FString> LocalDiscoveredPathsSet;
 
 		TArray<FString> CopyOfPathsToSearch;
 		{
@@ -348,17 +372,20 @@ void FAssetDataGatherer::DiscoverFilesToSearch()
 				{
 					// Add the path to this asset into the list of discovered paths
 					const FString LongPackageName = FPackageName::FilenameToLongPackageName(Filename);
-					LocalDiscoveredPaths.AddUnique( FPackageName::GetLongPackagePath(LongPackageName) );
+					LocalDiscoveredPathsSet.Add( FPackageName::GetLongPackagePath(LongPackageName) );
 					DiscoveredFilesToSearch.Add(Filename);
 				}
 			}
 		}
 
+		// Turn the set into an array here before the critical section below
+		TArray<FString> LocalDiscoveredPathsArray = LocalDiscoveredPathsSet.Array();
+
 		{
 			// Place all the discovered files into the files to search list
 			FScopeLock CritSectionLock(&WorkerThreadCriticalSection);
 			FilesToSearch.Append(DiscoveredFilesToSearch);
-			DiscoveredPaths.Append(LocalDiscoveredPaths);
+			DiscoveredPaths.Append(LocalDiscoveredPathsArray);
 		}
 	}
 }
diff --git a/Engine/Source/Runtime/AssetRegistry/Private/AssetDataGatherer.h b/Engine/Source/Runtime/AssetRegistry/Private/AssetDataGatherer.h
--- a/Engine/Source/Runtime/AssetRegistry/Private/AssetDataGatherer.h
+++ b/Engine/Source/Runtime/AssetRegistry/Private/AssetDataGatherer.h
@@ -30,6 +30,9 @@ public:
 	/** Adds specific files to the search queue. Only works when searching asynchronously */
 	void AddFilesToSearch(const TArray<FString>& Files);
 
+	/** If assets are currently being asynchronously scanned in the specified path, this will cause them to be scanned before other assets. */
+	void PrioritizeSearchPath(const FString& PathToPrioritize);
+
 private:
 	
 	/** This function is run on the gathering thread, unless synchronous */
diff --git a/Engine/Source/Runtime/AssetRegistry/Private/AssetRegistry.cpp b/Engine/Source/Runtime/AssetRegistry/Private/AssetRegistry.cpp
--- a/Engine/Source/Runtime/AssetRegistry/Private/AssetRegistry.cpp
+++ b/Engine/Source/Runtime/AssetRegistry/Private/AssetRegistry.cpp
@@ -1070,6 +1070,14 @@ void FAssetRegistry::ScanPathsSynchronous(const TArray<FString>& InPaths, bool b
 	ScanPathsSynchronous_Internal(InPaths, bForceRescan, bUseCache);
 }
 
+void FAssetRegistry::PrioritizeSearchPath(const FString& PathToPrioritize)
+{
+	if (BackgroundAssetSearch.IsValid())
+	{
+		BackgroundAssetSearch->PrioritizeSearchPath(PathToPrioritize);
+	}
+}
+
 void FAssetRegistry::AssetCreated(UObject* NewAsset)
 {
 	if ( ensure(NewAsset) && NewAsset->IsAsset() )
diff --git a/Engine/Source/Runtime/AssetRegistry/Private/AssetRegistry.h b/Engine/Source/Runtime/AssetRegistry/Private/AssetRegistry.h
--- a/Engine/Source/Runtime/AssetRegistry/Private/AssetRegistry.h
+++ b/Engine/Source/Runtime/AssetRegistry/Private/AssetRegistry.h
@@ -36,6 +36,7 @@ public:
 	virtual bool RemovePath(const FString& PathToRemove) override;
 	virtual void SearchAllAssets(bool bSynchronousSearch) override;
 	virtual void ScanPathsSynchronous(const TArray<FString>& InPaths, bool bForceRescan = false) override;
+	virtual void PrioritizeSearchPath(const FString& PathToPrioritize) override;
 	virtual void Serialize(FArchive& Ar) override;
 	virtual void SaveRegistryData(FArchive& Ar, TMap<FName, FAssetData*>& Data, int32 AssetCount) override;
 
diff --git a/Engine/Source/Runtime/AssetRegistry/Public/IAssetRegistry.h b/Engine/Source/Runtime/AssetRegistry/Public/IAssetRegistry.h
--- a/Engine/Source/Runtime/AssetRegistry/Public/IAssetRegistry.h
+++ b/Engine/Source/Runtime/AssetRegistry/Public/IAssetRegistry.h
@@ -165,6 +165,9 @@ public:
 	/** Look for all assets on disk (can be async or synchronous) */
 	virtual void SearchAllAssets(bool bSynchronousSearch) = 0;
 
+	/** If assets are currently being asynchronously scanned in the specified path, this will cause them to be scanned before other assets. */
+	virtual void PrioritizeSearchPath(const FString& PathToPrioritize) = 0;
+
 	/** Event for when paths are added to the registry */
 	DECLARE_EVENT_OneParam( IAssetRegistry, FPathAddedEvent, const FString& /*Path*/ );
 	virtual FPathAddedEvent& OnPathAdded() = 0;
