Added support for asset registry to be loaded async on editor startup from the cooker.

diff --git a/Engine/Source/Editor/UnrealEd/Private/CookOnTheFlyServer.cpp b/Engine/Source/Editor/UnrealEd/Private/CookOnTheFlyServer.cpp
--- a/Engine/Source/Editor/UnrealEd/Private/CookOnTheFlyServer.cpp
+++ b/Engine/Source/Editor/UnrealEd/Private/CookOnTheFlyServer.cpp
@@ -808,6 +808,15 @@ uint32 UCookOnTheFlyServer::TickCookOnTheSide( const float TimeSlice, uint32 &Co
 	FCookerTimer Timer(TimeSlice, IsRealtimeMode());
 
 	uint32 Result = 0;
+
+	FAssetRegistryModule& AssetRegistryModule = FModuleManager::LoadModuleChecked<FAssetRegistryModule>(TEXT("AssetRegistry"));
+	IAssetRegistry& AssetRegistry = AssetRegistryModule.Get();
+	if ( AssetRegistry.IsLoadingCookingAssets() )
+	{
+		// early out
+		return Result;
+	}
+
 	// This is all the target platforms which we needed to process requests for this iteration
 	// we use this in the unsolicited packages processing below
 	TArray<FName> AllTargetPlatformNames;
@@ -2060,15 +2069,27 @@ void UCookOnTheFlyServer::GenerateAssetRegistry(const TArray<ITargetPlatform*>&
 		// Perform a synchronous search of any .ini based asset paths (note that the per-game delegate may
 		// have already scanned paths on its own)
 		// We want the registry to be fully initialized when generating streaming manifests too.
+		bool bSyncronous = !IsRealtimeMode();
 		TArray<FString> ScanPaths;
 		if (GConfig->GetArray(TEXT("AssetRegistry"), TEXT("PathsToScanForCook"), ScanPaths, GEngineIni) > 0)
 		{
-			AssetRegistry.ScanPathsSynchronous(ScanPaths);
+			if ( bSyncronous )
+			{
+				AssetRegistry.ScanPathsSynchronous(ScanPaths);
+			}
+			else
+			{
+				for ( const auto& Path : ScanPaths )
+				{
+					AssetRegistry.AddPath(Path);
+				}
+			}
 		}
 		else
 		{
-			AssetRegistry.SearchAllAssets(true);
+			AssetRegistry.SearchAllAssets(bSyncronous);
 		}
+		AssetRegistry.StartLoadingCookingAssets();
 		
 		// When not cooking on the fly the registry will be saved after the cooker has finished
 		if (CurrentCookMode == ECookMode::CookOnTheFly)
diff --git a/Engine/Source/Runtime/AssetRegistry/Private/AssetDataGatherer.cpp b/Engine/Source/Runtime/AssetRegistry/Private/AssetDataGatherer.cpp
--- a/Engine/Source/Runtime/AssetRegistry/Private/AssetDataGatherer.cpp
+++ b/Engine/Source/Runtime/AssetRegistry/Private/AssetDataGatherer.cpp
@@ -13,6 +13,7 @@ FAssetDataGatherer::FAssetDataGatherer(const TArray<FString>& InPaths, bool bInI
 	, bLoadAndSaveCache( bInLoadAndSaveCache )
 	, bSavedCacheAfterInitialDiscovery( false )
 	, DiskCachedAssetDataBuffer( NULL )
+	, Thread(NULL)
 {
 	const FString AllIllegalCharacters = INVALID_LONGPACKAGE_CHARACTERS;
 	for ( int32 CharIdx = 0; CharIdx < AllIllegalCharacters.Len() ; ++CharIdx )
diff --git a/Engine/Source/Runtime/AssetRegistry/Private/AssetRegistry.cpp b/Engine/Source/Runtime/AssetRegistry/Private/AssetRegistry.cpp
--- a/Engine/Source/Runtime/AssetRegistry/Private/AssetRegistry.cpp
+++ b/Engine/Source/Runtime/AssetRegistry/Private/AssetRegistry.cpp
@@ -39,9 +39,12 @@ FAssetRegistry::FAssetRegistry()
 
 	NumAssets = 0;
 	NumDependsNodes = 0;
+	bSearchAllStarted = false;
 	bInitialSearchCompleted = true;
+	bCookSearchCompleted = true;
 	AmortizeStartTime = 0;
 	TotalAmortizeTime = 0;
+	
 
 	MaxSecondsPerFrame = 0.015;
 
@@ -208,6 +211,11 @@ FAssetRegistry::~FAssetRegistry()
 
 void FAssetRegistry::SearchAllAssets(bool bSynchronousSearch)
 {
+	if ( !bSynchronousSearch && bSearchAllStarted )
+		// already started the search and it's async so don't need to do it again
+		return;
+
+	bSearchAllStarted = true;
 	// Mark the time before the first search started
 	FullSearchStartTime = FPlatformTime::Seconds();
 
@@ -226,7 +234,17 @@ void FAssetRegistry::SearchAllAssets(bool bSynchronousSearch)
 	}
 	else
 	{
-		BackgroundAssetSearch = MakeShareable( new FAssetDataGatherer(PathsToSearch, bSynchronousSearch, bLoadAndSaveCache) );
+		if ( BackgroundAssetSearch.IsValid() )
+		{
+			for ( const auto& Path : PathsToSearch )
+			{
+				BackgroundAssetSearch->AddPathToSearch( Path );
+			}
+		}
+		else
+		{
+			BackgroundAssetSearch = MakeShareable( new FAssetDataGatherer(PathsToSearch, bSynchronousSearch, bLoadAndSaveCache) );
+		}
 	}
 }
 
@@ -1174,6 +1192,16 @@ bool FAssetRegistry::IsLoadingAssets() const
 	return !bInitialSearchCompleted;
 }
 
+void FAssetRegistry::StartLoadingCookingAssets()
+{
+	bCookSearchCompleted = false;
+}
+
+bool FAssetRegistry::IsLoadingCookingAssets() const
+{
+	return !bCookSearchCompleted;
+}
+
 void FAssetRegistry::Tick(float DeltaTime)
 {
 	double TickStartTime = FPlatformTime::Seconds();
@@ -1231,17 +1259,26 @@ void FAssetRegistry::Tick(float DeltaTime)
 	}
 
 	// If completing an initial search, refresh the content browser
-	if ( !bInitialSearchCompleted && NumFilesToSearch == 0 && NumPathsToSearch == 0 && !bIsSearching && BackgroundPathResults.Num() == 0 && BackgroundAssetResults.Num() == 0 && BackgroundDependencyResults.Num() == 0 )
+	if ( NumFilesToSearch == 0 && NumPathsToSearch == 0 && !bIsSearching && BackgroundPathResults.Num() == 0 && BackgroundAssetResults.Num() == 0 && BackgroundDependencyResults.Num() == 0 )
 	{
-		UE_LOG(LogAssetRegistry, Verbose, TEXT("### Time spent amortizing search results: %0.4f seconds"), TotalAmortizeTime);
-		UE_LOG(LogAssetRegistry, Log, TEXT("Asset discovery search completed in %0.4f seconds"), FPlatformTime::Seconds() - FullSearchStartTime);
+		if ( !bInitialSearchCompleted )
+		{
+			UE_LOG(LogAssetRegistry, Verbose, TEXT("### Time spent amortizing search results: %0.4f seconds"), TotalAmortizeTime);
+			UE_LOG(LogAssetRegistry, Log, TEXT("Asset discovery search completed in %0.4f seconds"), FPlatformTime::Seconds() - FullSearchStartTime);
 
-		bInitialSearchCompleted = true;
+			bInitialSearchCompleted = true;
 
-		FileLoadedEvent.Broadcast();
+			FileLoadedEvent.Broadcast();
+		}
+		
+		if ( !bCookSearchCompleted )
+		{
+			bCookSearchCompleted = true;
+		}
 	}
 }
 
+
 bool FAssetRegistry::IsUsingWorldAssets()
 {
 	return !FParse::Param(FCommandLine::Get(), TEXT("DisableWorldAssets"));
diff --git a/Engine/Source/Runtime/AssetRegistry/Private/AssetRegistry.h b/Engine/Source/Runtime/AssetRegistry/Private/AssetRegistry.h
--- a/Engine/Source/Runtime/AssetRegistry/Private/AssetRegistry.h
+++ b/Engine/Source/Runtime/AssetRegistry/Private/AssetRegistry.h
@@ -71,6 +71,10 @@ public:
 	DECLARE_DERIVED_EVENT( FAssetRegistry, IAssetRegistry::FFileLoadProgressUpdatedEvent, FFileLoadProgressUpdatedEvent );
 	virtual FFileLoadProgressUpdatedEvent& OnFileLoadProgressUpdated() override { return FileLoadProgressUpdatedEvent; }
 
+	/** Async management of cook asset directories */
+	virtual void StartLoadingCookingAssets() override;
+	virtual bool IsLoadingCookingAssets() const override;
+
 	virtual bool IsLoadingAssets() const override;
 
 	virtual void Tick (float DeltaTime) override;
@@ -246,6 +250,12 @@ private:
 	/** Flag to indicate if the initial background search has completed */
 	bool bInitialSearchCompleted;
 
+	/** Flag to indicate the if SearchAllAssets has already been run */
+	bool bSearchAllStarted;
+
+	/** Flag to indicate the cooked directories have been searched */
+	bool bCookSearchCompleted;
+
 	/** When loading a registry from disk, we can allocate all the FAssetData objects in one chunk, to save on 10s of thousands of heap allocations */
 	FAssetData* PreallocatedAssetDataBuffer;
 
diff --git a/Engine/Source/Runtime/AssetRegistry/Public/IAssetRegistry.h b/Engine/Source/Runtime/AssetRegistry/Public/IAssetRegistry.h
--- a/Engine/Source/Runtime/AssetRegistry/Public/IAssetRegistry.h
+++ b/Engine/Source/Runtime/AssetRegistry/Public/IAssetRegistry.h
@@ -213,6 +213,12 @@ public:
 	DECLARE_EVENT_TwoParams( IAssetRegistry, FFileLoadProgressUpdatedEvent, int32 /*NumAssetsDiscovered*/, int32 /*TotalAssets*/ );
 	virtual FFileLoadProgressUpdatedEvent& OnFileLoadProgressUpdated() = 0;
 
+	/** Indicate that we have started loading cooking asset */
+	virtual void StartLoadingCookingAssets() = 0;
+
+	/** Returns false when asset registry finishes loading cooking assets (after StartLoadingCookingAssets has been called) */
+	virtual bool IsLoadingCookingAssets() const = 0;
+
 	/** Returns true if the asset registry is currently loading files and does not yet know about all assets */
 	virtual bool IsLoadingAssets() const = 0;
 
