Cleaned up asset registry loading for the cooker.
Cooked paths are no longer needed so removed them.

diff --git a/Engine/Source/Editor/UnrealEd/Private/CookOnTheFlyServer.cpp b/Engine/Source/Editor/UnrealEd/Private/CookOnTheFlyServer.cpp
--- a/Engine/Source/Editor/UnrealEd/Private/CookOnTheFlyServer.cpp
+++ b/Engine/Source/Editor/UnrealEd/Private/CookOnTheFlyServer.cpp
@@ -811,7 +811,7 @@ uint32 UCookOnTheFlyServer::TickCookOnTheSide( const float TimeSlice, uint32 &Co
 
 	FAssetRegistryModule& AssetRegistryModule = FModuleManager::LoadModuleChecked<FAssetRegistryModule>(TEXT("AssetRegistry"));
 	IAssetRegistry& AssetRegistry = AssetRegistryModule.Get();
-	if ( AssetRegistry.IsLoadingCookingAssets() )
+	if ( AssetRegistry.IsLoadingAssets() )
 	{
 		// early out
 		return Result;
@@ -2069,28 +2069,22 @@ void UCookOnTheFlyServer::GenerateAssetRegistry(const TArray<ITargetPlatform*>&
 		// Perform a synchronous search of any .ini based asset paths (note that the per-game delegate may
 		// have already scanned paths on its own)
 		// We want the registry to be fully initialized when generating streaming manifests too.
-		bool bSyncronous = !IsRealtimeMode();
-		TArray<FString> ScanPaths;
-		if (GConfig->GetArray(TEXT("AssetRegistry"), TEXT("PathsToScanForCook"), ScanPaths, GEngineIni) > 0)
+		bool bEditor = IsRealtimeMode();
+
+		// editor will scan asset registry automagically 
+		if ( !bEditor )
 		{
-			if ( bSyncronous )
+			TArray<FString> ScanPaths;
+			if (GConfig->GetArray(TEXT("AssetRegistry"), TEXT("PathsToScanForCook"), ScanPaths, GEngineIni) > 0)
 			{
 				AssetRegistry.ScanPathsSynchronous(ScanPaths);
 			}
 			else
 			{
-				for ( const auto& Path : ScanPaths )
-				{
-					AssetRegistry.AddPath(Path);
-				}
+				AssetRegistry.SearchAllAssets(true);
 			}
 		}
-		else
-		{
-			AssetRegistry.SearchAllAssets(bSyncronous);
-		}
-		AssetRegistry.StartLoadingCookingAssets();
-		
+
 		// When not cooking on the fly the registry will be saved after the cooker has finished
 		if (CurrentCookMode == ECookMode::CookOnTheFly)
 		{
diff --git a/Engine/Source/Runtime/AssetRegistry/Private/AssetRegistry.cpp b/Engine/Source/Runtime/AssetRegistry/Private/AssetRegistry.cpp
--- a/Engine/Source/Runtime/AssetRegistry/Private/AssetRegistry.cpp
+++ b/Engine/Source/Runtime/AssetRegistry/Private/AssetRegistry.cpp
@@ -39,9 +39,7 @@ FAssetRegistry::FAssetRegistry()
 
 	NumAssets = 0;
 	NumDependsNodes = 0;
-	bSearchAllStarted = false;
 	bInitialSearchCompleted = true;
-	bCookSearchCompleted = true;
 	AmortizeStartTime = 0;
 	TotalAmortizeTime = 0;
 	
@@ -211,11 +209,6 @@ FAssetRegistry::~FAssetRegistry()
 
 void FAssetRegistry::SearchAllAssets(bool bSynchronousSearch)
 {
-	if ( !bSynchronousSearch && bSearchAllStarted )
-		// already started the search and it's async so don't need to do it again
-		return;
-
-	bSearchAllStarted = true;
 	// Mark the time before the first search started
 	FullSearchStartTime = FPlatformTime::Seconds();
 
@@ -232,19 +225,10 @@ void FAssetRegistry::SearchAllAssets(bool bSynchronousSearch)
 		const bool bForceRescan = false;
 		ScanPathsSynchronous_Internal(PathsToSearch, bForceRescan, bLoadAndSaveCache);
 	}
-	else
+	else if ( !BackgroundAssetSearch.IsValid() )
 	{
-		if ( BackgroundAssetSearch.IsValid() )
-		{
-			for ( const auto& Path : PathsToSearch )
-			{
-				BackgroundAssetSearch->AddPathToSearch( Path );
-			}
-		}
-		else
-		{
-			BackgroundAssetSearch = MakeShareable( new FAssetDataGatherer(PathsToSearch, bSynchronousSearch, bLoadAndSaveCache) );
-		}
+		// if the BackgroundAssetSearch is already valid then we have already called it before
+		BackgroundAssetSearch = MakeShareable( new FAssetDataGatherer(PathsToSearch, bSynchronousSearch, bLoadAndSaveCache) );
 	}
 }
 
@@ -1192,16 +1176,6 @@ bool FAssetRegistry::IsLoadingAssets() const
 	return !bInitialSearchCompleted;
 }
 
-void FAssetRegistry::StartLoadingCookingAssets()
-{
-	bCookSearchCompleted = false;
-}
-
-bool FAssetRegistry::IsLoadingCookingAssets() const
-{
-	return !bCookSearchCompleted;
-}
-
 void FAssetRegistry::Tick(float DeltaTime)
 {
 	double TickStartTime = FPlatformTime::Seconds();
@@ -1270,11 +1244,6 @@ void FAssetRegistry::Tick(float DeltaTime)
 
 			FileLoadedEvent.Broadcast();
 		}
-		
-		if ( !bCookSearchCompleted )
-		{
-			bCookSearchCompleted = true;
-		}
 	}
 }
 
diff --git a/Engine/Source/Runtime/AssetRegistry/Private/AssetRegistry.h b/Engine/Source/Runtime/AssetRegistry/Private/AssetRegistry.h
--- a/Engine/Source/Runtime/AssetRegistry/Private/AssetRegistry.h
+++ b/Engine/Source/Runtime/AssetRegistry/Private/AssetRegistry.h
@@ -71,10 +71,6 @@ public:
 	DECLARE_DERIVED_EVENT( FAssetRegistry, IAssetRegistry::FFileLoadProgressUpdatedEvent, FFileLoadProgressUpdatedEvent );
 	virtual FFileLoadProgressUpdatedEvent& OnFileLoadProgressUpdated() override { return FileLoadProgressUpdatedEvent; }
 
-	/** Async management of cook asset directories */
-	virtual void StartLoadingCookingAssets() override;
-	virtual bool IsLoadingCookingAssets() const override;
-
 	virtual bool IsLoadingAssets() const override;
 
 	virtual void Tick (float DeltaTime) override;
@@ -250,12 +246,6 @@ private:
 	/** Flag to indicate if the initial background search has completed */
 	bool bInitialSearchCompleted;
 
-	/** Flag to indicate the if SearchAllAssets has already been run */
-	bool bSearchAllStarted;
-
-	/** Flag to indicate the cooked directories have been searched */
-	bool bCookSearchCompleted;
-
 	/** When loading a registry from disk, we can allocate all the FAssetData objects in one chunk, to save on 10s of thousands of heap allocations */
 	FAssetData* PreallocatedAssetDataBuffer;
 
diff --git a/Engine/Source/Runtime/AssetRegistry/Public/IAssetRegistry.h b/Engine/Source/Runtime/AssetRegistry/Public/IAssetRegistry.h
--- a/Engine/Source/Runtime/AssetRegistry/Public/IAssetRegistry.h
+++ b/Engine/Source/Runtime/AssetRegistry/Public/IAssetRegistry.h
@@ -213,12 +213,6 @@ public:
 	DECLARE_EVENT_TwoParams( IAssetRegistry, FFileLoadProgressUpdatedEvent, int32 /*NumAssetsDiscovered*/, int32 /*TotalAssets*/ );
 	virtual FFileLoadProgressUpdatedEvent& OnFileLoadProgressUpdated() = 0;
 
-	/** Indicate that we have started loading cooking asset */
-	virtual void StartLoadingCookingAssets() = 0;
-
-	/** Returns false when asset registry finishes loading cooking assets (after StartLoadingCookingAssets has been called) */
-	virtual bool IsLoadingCookingAssets() const = 0;
-
 	/** Returns true if the asset registry is currently loading files and does not yet know about all assets */
 	virtual bool IsLoadingAssets() const = 0;
 
