optimized asset registry queries when filtering by UClass name(s)

diff --git a/Engine/Source/Runtime/AssetRegistry/Private/AssetRegistry.cpp b/Engine/Source/Runtime/AssetRegistry/Private/AssetRegistry.cpp
--- a/Engine/Source/Runtime/AssetRegistry/Private/AssetRegistry.cpp
+++ b/Engine/Source/Runtime/AssetRegistry/Private/AssetRegistry.cpp
@@ -325,60 +325,52 @@ bool FAssetRegistry::GetAssets(const FARFilter& Filter, TArray<FAssetData>& OutA
 			FilterObjectPaths.Add(Filter.ObjectPaths[ObjectPathIdx]);
 		}
 
-		// Iterate over all in-memory assets to find the ones that pass the filter components
-		for (FObjectIterator ObjIt; ObjIt; ++ObjIt)
+		auto FilterInMemoryObjectLambda = [&](const UObject* Obj)
 		{
-			if ( ObjIt->IsAsset() )
+			if ( Obj->IsAsset() )
 			{
-				UPackage* InMemoryPackage = ObjIt->GetOutermost();
+				UPackage* InMemoryPackage = Obj->GetOutermost();
 
 				static const bool bUsingWorldAssets = FAssetRegistry::IsUsingWorldAssets();
 				// Skip assets in map packages... unless we are showing world assets
 				if ( InMemoryPackage->ContainsMap() && !bUsingWorldAssets )
 				{
-					continue;
+					return;
 				}
 
 				// Skip assets that were loaded for diffing
 				if ( InMemoryPackage->PackageFlags & PKG_ForDiffing )
 				{
-					continue;
+					return;
 				}
 
 				// add it to in-memory object list for later merge
-				const FName ObjectPath = FName(*ObjIt->GetPathName());
+				const FName ObjectPath = FName(*Obj->GetPathName());
 				InMemoryObjectPaths.Add(ObjectPath);
-			
+
 				// Package name
 				const FName PackageName = InMemoryPackage->GetFName();
 				if ( NumFilterPackageNames && !FilterPackageNames.Contains(PackageName) )
 				{
-					continue;
+					return;
 				}
 
 				// Object Path
 				if ( NumFilterObjectPaths && !FilterObjectPaths.Contains(ObjectPath) )
 				{
-					continue;
+					return;
 				}
 
 				// Package path
 				const FName PackagePath = FName(*FPackageName::GetLongPackagePath(InMemoryPackage->GetName()));
 				if ( NumFilterPackagePaths && !FilterPackagePaths.Contains(PackagePath) )
 				{
-					continue;
-				}
-
-				// Asset class
-				const FName AssetClassName = ObjIt->GetClass()->GetFName();
-				if ( NumFilterClasses && !FilterClassNames.Contains(AssetClassName) )
-				{
-					continue;
+					return;
 				}
 
 				// Tags and values
 				TArray<UObject::FAssetRegistryTag> ObjectTags;
-				ObjIt->GetAssetRegistryTags(ObjectTags);
+				Obj->GetAssetRegistryTags(ObjectTags);
 				if ( Filter.TagsAndValues.Num() )
 				{
 					bool bMatch = false;
@@ -395,7 +387,7 @@ bool FAssetRegistry::GetAssets(const FARFilter& Filter, TArray<FAssetData>& OutA
 								{
 									bMatch = true;
 								}
-							
+
 								break;
 							}
 						}
@@ -408,14 +400,14 @@ bool FAssetRegistry::GetAssets(const FARFilter& Filter, TArray<FAssetData>& OutA
 
 					if (!bMatch)
 					{
-						continue;
+						return;
 					}
 				}
 
 				// Find the group names
 				FString GroupNamesStr;
 				FString AssetNameStr;
-				ObjIt->GetPathName(InMemoryPackage).Split(TEXT("."), &GroupNamesStr, &AssetNameStr, ESearchCase::CaseSensitive, ESearchDir::FromEnd);
+				Obj->GetPathName(InMemoryPackage).Split(TEXT("."), &GroupNamesStr, &AssetNameStr, ESearchCase::CaseSensitive, ESearchDir::FromEnd);
 
 				TMap<FName, FString> TagMap;
 				for ( auto TagIt = ObjectTags.CreateConstIterator(); TagIt; ++TagIt )
@@ -424,12 +416,36 @@ bool FAssetRegistry::GetAssets(const FARFilter& Filter, TArray<FAssetData>& OutA
 				}
 
 				// This asset is in memory and passes all filters
-				FAssetData* AssetData = new (OutAssetData) FAssetData(PackageName, PackagePath, FName(*GroupNamesStr), ObjIt->GetFName(), AssetClassName, TagMap, InMemoryPackage->GetChunkIDs());
+				FAssetData* AssetData = new (OutAssetData) FAssetData(PackageName, PackagePath, FName(*GroupNamesStr), Obj->GetFName(), Obj->GetClass()->GetFName(), TagMap, InMemoryPackage->GetChunkIDs());
 			}
-		}
-	}
+		};
 
+		// Iterate over all in-memory assets to find the ones that pass the filter components
+		if(NumFilterClasses)
+		{
+			TArray<UObject*> InMemoryObjects;
+			for (auto ClassNameIt = FilterClassNames.CreateConstIterator(); ClassNameIt; ++ClassNameIt)
+			{
+				UClass* Class = FindObjectFast<UClass>(nullptr, *ClassNameIt->ToString(), false, true, RF_NoFlags);
+				if(Class != nullptr)
+				{
+					GetObjectsOfClass(Class, InMemoryObjects, false, RF_NoFlags);
+				}
+			}
 
+			for (auto ObjIt = InMemoryObjects.CreateConstIterator(); ObjIt; ++ObjIt)
+			{
+				FilterInMemoryObjectLambda(*ObjIt);
+			}
+		}
+		else
+		{
+			for (FObjectIterator ObjIt; ObjIt; ++ObjIt)
+			{
+				FilterInMemoryObjectLambda(*ObjIt);
+			}
+		}
+	}
 
 	// Now add cached (unloaded) assets
 	// Form a set of assets matched by each filter
