Call CollectCodeGeneratorClasses() at the top of the FAssetRegistry constructor. This fixes an issue where cooked builds would not correctly add all blueprint classes into the cached inheritence map, since the asset registry was serialized in prior to CollectCodeGeneratrClasses() being called.

diff --git a/Engine/Source/Runtime/AssetRegistry/Private/AssetRegistry.cpp b/Engine/Source/Runtime/AssetRegistry/Private/AssetRegistry.cpp
--- a/Engine/Source/Runtime/AssetRegistry/Private/AssetRegistry.cpp
+++ b/Engine/Source/Runtime/AssetRegistry/Private/AssetRegistry.cpp
@@ -46,6 +46,9 @@ FAssetRegistry::FAssetRegistry()
 
 	MaxSecondsPerFrame = 0.015;
 
+	// Collect all code generator classes (currently BlueprintCore-derived ones)
+	CollectCodeGeneratorClasses();
+
 	// If in the editor, we scan all content right now
 	// If in the game, we expect user to make explicit sync queries using ScanPathsSynchronous
 	// If in a commandlet, we expect the commandlet to decide when to perform a synchronous scan
@@ -96,9 +99,6 @@ FAssetRegistry::FAssetRegistry()
 	// will be loaded a bit later on.
 	FPackageName::OnContentPathMounted().AddRaw( this, &FAssetRegistry::OnContentPathMounted );
 	FPackageName::OnContentPathDismounted().AddRaw( this, &FAssetRegistry::OnContentPathDismounted );
-
-	// Now collect all code generator classes (currently BlueprintCore-derived ones)
-	CollectCodeGeneratorClasses();
 }
 
 void FAssetRegistry::CollectCodeGeneratorClasses()
