[INTEGRATE] Change 2412806 by [REDACTED_PERSON][REDACTED_HANDLE].7 on 2015/01/20 16:20:05

	Changed importable asset types to explicitly add source file paths to asset registry tags

	Previously this was harvested from the object's properties, but required recursing into sub objects which was not feasible.

diff --git a/Engine/Plugins/2D/Paper2D/Source/Paper2D/Classes/PaperTileMap.h b/Engine/Plugins/2D/Paper2D/Source/Paper2D/Classes/PaperTileMap.h
--- a/Engine/Plugins/2D/Paper2D/Source/Paper2D/Classes/PaperTileMap.h
+++ b/Engine/Plugins/2D/Paper2D/Source/Paper2D/Classes/PaperTileMap.h
@@ -92,7 +92,7 @@ class PAPER2D_API UPaperTileMap : public UObject
 public:
 #if WITH_EDITORONLY_DATA
 	/** Importing data and options used for this tile map */
-	UPROPERTY(Category=ImportSettings, VisibleAnywhere, Instanced, AssetRegistrySearchable)
+	UPROPERTY(Category=ImportSettings, VisibleAnywhere, Instanced)
 	class UAssetImportData* AssetImportData;
 
 	/** The currently selected layer index */
@@ -114,6 +114,9 @@ public:
 	virtual void PostEditChangeProperty(FPropertyChangedEvent& PropertyChangedEvent) override;
 	virtual void PostLoad() override;
 	void ValidateSelectedLayerIndex();
+#endif
+#if WITH_EDITORONLY_DATA
+	virtual void GetAssetRegistryTags(TArray<FAssetRegistryTag>& OutTags) const;
 #endif
 	// End of UObject interface
 
diff --git a/Engine/Plugins/2D/Paper2D/Source/Paper2D/Private/PaperTileMap.cpp b/Engine/Plugins/2D/Paper2D/Source/Paper2D/Private/PaperTileMap.cpp
--- a/Engine/Plugins/2D/Paper2D/Source/Paper2D/Private/PaperTileMap.cpp
+++ b/Engine/Plugins/2D/Paper2D/Source/Paper2D/Private/PaperTileMap.cpp
@@ -126,6 +126,17 @@ void UPaperTileMap::ValidateSelectedLayerIndex()
 
 #endif
 
+#if WITH_EDITORONLY_DATA
+void UPaperTileMap::GetAssetRegistryTags(TArray<FAssetRegistryTag>& OutTags) const
+{
+	if (AssetImportData)
+	{
+		OutTags.Add( FAssetRegistryTag(SourceFileTagName(), AssetImportData->SourceFilePath, FAssetRegistryTag::TT_Hidden) );
+	}
+
+	Super::GetAssetRegistryTags(OutTags);
+}
+#endif
 
 void UPaperTileMap::UpdateBodySetup()
 {
diff --git a/Engine/Plugins/2D/Paper2D/Source/PaperJsonImporter/Private/PaperSpriteSheet.cpp b/Engine/Plugins/2D/Paper2D/Source/PaperJsonImporter/Private/PaperSpriteSheet.cpp
--- a/Engine/Plugins/2D/Paper2D/Source/PaperJsonImporter/Private/PaperSpriteSheet.cpp
+++ b/Engine/Plugins/2D/Paper2D/Source/PaperJsonImporter/Private/PaperSpriteSheet.cpp
@@ -8,3 +8,14 @@ UPaperSpriteSheet::UPaperSpriteSheet(const FObjectInitializer& ObjectInitializer
 {
 }
 
+#if WITH_EDITORONLY_DATA
+void UPaperSpriteSheet::GetAssetRegistryTags(TArray<FAssetRegistryTag>& OutTags) const
+{
+	if (AssetImportData)
+	{
+		OutTags.Add( FAssetRegistryTag(SourceFileTagName(), AssetImportData->SourceFilePath, FAssetRegistryTag::TT_Hidden) );
+	}
+
+	Super::GetAssetRegistryTags(OutTags);
+}
+#endif
\ No newline at end of file
diff --git a/Engine/Plugins/2D/Paper2D/Source/PaperJsonImporter/Private/PaperSpriteSheet.h b/Engine/Plugins/2D/Paper2D/Source/PaperJsonImporter/Private/PaperSpriteSheet.h
--- a/Engine/Plugins/2D/Paper2D/Source/PaperJsonImporter/Private/PaperSpriteSheet.h
+++ b/Engine/Plugins/2D/Paper2D/Source/PaperJsonImporter/Private/PaperSpriteSheet.h
@@ -26,7 +26,10 @@ public:
 
 #if WITH_EDITORONLY_DATA
 	// Import data for this 
-	UPROPERTY(VisibleAnywhere, Instanced, Category=ImportSettings, AssetRegistrySearchable)
+	UPROPERTY(VisibleAnywhere, Instanced, Category=ImportSettings)
 	class UAssetImportData* AssetImportData;
+
+	/** Override to ensure we write out the asset import data */
+	virtual void GetAssetRegistryTags(TArray<FAssetRegistryTag>& OutTags) const;
 #endif
 };
diff --git a/Engine/Plugins/ScriptPlugin/Source/ScriptPlugin/Classes/ScriptBlueprint.h b/Engine/Plugins/ScriptPlugin/Source/ScriptPlugin/Classes/ScriptBlueprint.h
--- a/Engine/Plugins/ScriptPlugin/Source/ScriptPlugin/Classes/ScriptBlueprint.h
+++ b/Engine/Plugins/ScriptPlugin/Source/ScriptPlugin/Classes/ScriptBlueprint.h
@@ -17,7 +17,7 @@ class SCRIPTPLUGIN_API UScriptBlueprint : public UBlueprint
 public:
 
 	/** Path to the script text used to construct this asset. Relative to the object's package, BaseDir() or absolute */
-	UPROPERTY(Category = Script, VisibleAnywhere, BlueprintReadWrite, AssetRegistrySearchable, meta=(ReimportPath))
+	UPROPERTY(Category = Script, VisibleAnywhere, BlueprintReadWrite)
 	FString SourceFilePath;
 
 	/** Date/Time-stamp of the file from the last import. */
@@ -31,7 +31,12 @@ public:
 	/** Script source code. @todo: this should be editor-only */
 	UPROPERTY()
 	FString SourceCode;
-	 
+	
+#if WITH_EDITORONLY_DATA
+	/** Override to ensure we write out the asset import data */
+	virtual void GetAssetRegistryTags(TArray<FAssetRegistryTag>& OutTags) const;
+#endif
+
 #if WITH_EDITOR
 	virtual UClass* GetBlueprintClass() const override;
 
diff --git a/Engine/Plugins/ScriptPlugin/Source/ScriptPlugin/Private/ScriptBlueprint.cpp b/Engine/Plugins/ScriptPlugin/Source/ScriptPlugin/Private/ScriptBlueprint.cpp
--- a/Engine/Plugins/ScriptPlugin/Source/ScriptPlugin/Private/ScriptBlueprint.cpp
+++ b/Engine/Plugins/ScriptPlugin/Source/ScriptPlugin/Private/ScriptBlueprint.cpp
@@ -13,6 +13,15 @@ UScriptBlueprint::UScriptBlueprint(const FObjectInitializer& ObjectInitializer)
 {
 }
 
+#if WITH_EDITORONLY_DATA
+void UScriptBlueprint::GetAssetRegistryTags(TArray<FAssetRegistryTag>& OutTags) const
+{
+	OutTags.Add( FAssetRegistryTag(SourceFileTagName(), SourceFilePath, FAssetRegistryTag::TT_Hidden) );
+
+	Super::GetAssetRegistryTags(OutTags);
+}
+#endif
+
 #if WITH_EDITOR
 UClass* UScriptBlueprint::GetBlueprintClass() const
 {
diff --git a/Engine/Source/Editor/UnrealEd/Private/AutoReimportManager.cpp b/Engine/Source/Editor/UnrealEd/Private/AutoReimportManager.cpp
--- a/Engine/Source/Editor/UnrealEd/Private/AutoReimportManager.cpp
+++ b/Engine/Source/Editor/UnrealEd/Private/AutoReimportManager.cpp
@@ -19,7 +19,6 @@ DEFINE_LOG_CATEGORY_STATIC(LogAutoReimportManager, Log, All);
 
 namespace
 {
-	static const FName ReimportPathName("ReimportPath");
 
 	/** Template helper function to count the number of elements in the specified array that pass a predicate */
 	template<typename T, typename P>
@@ -91,6 +90,7 @@ namespace
 	{
 		TArray<FAssetData> Assets, Result;
 		const FString LeafName = FPaths::GetCleanFilename(AbsoluteFilename);
+		const FName TagName = UObject::SourceFileTagName();
 
 		FARFilter Filter;
 		Filter.SourceFilenames.Add(*LeafName);
@@ -103,7 +103,7 @@ namespace
 			{
 				// We don't compare numbers on FNames for this check because the tag "ReimportPath" may exist multiple times
 				const bool bCompareNumber = false;
-				if (Pair.Key.IsEqual(ReimportPathName, ENameCase::IgnoreCase, bCompareNumber))
+				if (Pair.Key.IsEqual(TagName, ENameCase::IgnoreCase, bCompareNumber))
 				{
 					// Either relative to the asset itself, or relative to the base path, or absolute.
 					// Would ideally use FReimportManager::ResolveImportFilename here, but we don't want to ask the file system whether the file exists.
@@ -437,10 +437,12 @@ void UAutoReimportManager::OnAssetDeleted(UObject* Object)
 
 	TArray<FString> AbsoluteSourceFileNames;
 
+	const FName TagName = UObject::SourceFileTagName();
+
 	for (const auto& Tag : Tags)
 	{
 		// Don't compare the FName number suffix
-		if (Tag.Name.IsEqual(ReimportPathName, ENameCase::IgnoreCase, false))
+		if (Tag.Name.IsEqual(TagName, ENameCase::IgnoreCase, false))
 		{
 			const FString SourceFile = FReimportManager::ResolveImportFilename(Tag.Value, Object);
 			// Check there aren't any other assets referencing this thing
diff --git a/Engine/Source/Runtime/AssetRegistry/Private/AssetRegistry.cpp b/Engine/Source/Runtime/AssetRegistry/Private/AssetRegistry.cpp
--- a/Engine/Source/Runtime/AssetRegistry/Private/AssetRegistry.cpp
+++ b/Engine/Source/Runtime/AssetRegistry/Private/AssetRegistry.cpp
@@ -11,9 +11,6 @@
 
 DEFINE_LOG_CATEGORY(LogAssetRegistry);
 
-
-static const FName ReimportPathName("ReimportPath");
-
 /** Returns the appropriate ChunkProgressReportingType for the given Asset enum */
 EChunkProgressReportingType::Type GetChunkAvailabilityProgressType(EAssetAvailabilityProgressReportingType::Type ReportType)
 {
@@ -1650,7 +1647,7 @@ void FAssetRegistry::AddAssetData(FAssetData* AssetData)
 
 #if WITH_EDITORONLY_DATA
 		// Don't check FName number so we still accumulate them even if there are multiple
-		if (GIsEditor && TagIt.Key().IsEqual(ReimportPathName, ENameCase::IgnoreCase, false))
+		if (GIsEditor && TagIt.Key().IsEqual(UObject::SourceFileTagName(), ENameCase::IgnoreCase, false))
 		{
 			auto& SourceFiles = CachedAssetsBySourceFileName.FindOrAdd(*FPaths::GetCleanFilename(TagIt.Value()));
 			SourceFiles.Add(AssetData);
@@ -1733,6 +1730,8 @@ void FAssetRegistry::UpdateAssetData(FAssetData* AssetData, const FAssetData& Ne
 	// Update Tags
 	if (bTagsChanged)
 	{
+		const FName& SourceFileTagName = UObject::SourceFileTagName();
+
 		for (TMap<FName, FString>::TConstIterator TagIt(AssetData->TagsAndValues); TagIt; ++TagIt)
 		{
 			const FName FNameKey = TagIt.Key();
@@ -1742,7 +1741,7 @@ void FAssetRegistry::UpdateAssetData(FAssetData* AssetData, const FAssetData& Ne
 
 #if WITH_EDITORONLY_DATA
 			// Don't check FName number so we still remove them even if there are multiple
-			if (GIsEditor && TagIt.Key().IsEqual(ReimportPathName, ENameCase::IgnoreCase, false))
+			if (GIsEditor && TagIt.Key().IsEqual(SourceFileTagName, ENameCase::IgnoreCase, false))
 			{
 				auto* SourceFiles = CachedAssetsBySourceFileName.Find(*FPaths::GetCleanFilename(TagIt.Value()));
 				SourceFiles->Remove(AssetData);
@@ -1759,7 +1758,7 @@ void FAssetRegistry::UpdateAssetData(FAssetData* AssetData, const FAssetData& Ne
 
 #if WITH_EDITORONLY_DATA
 			// Don't check FName number so we still add them even if there are multiple
-			if (GIsEditor && TagIt.Key().IsEqual(ReimportPathName, ENameCase::IgnoreCase, false))
+			if (GIsEditor && TagIt.Key().IsEqual(SourceFileTagName, ENameCase::IgnoreCase, false))
 			{
 				auto& SourceFiles = CachedAssetsBySourceFileName.FindOrAdd(*FPaths::GetCleanFilename(TagIt.Value()));
 				SourceFiles.Add(AssetData);
@@ -1823,6 +1822,7 @@ bool FAssetRegistry::RemoveAssetData(FAssetData* AssetData)
 		OldPathAssets->Remove(AssetData);
 		OldClassAssets->Remove(AssetData);
 
+		const FName& SourceFileTagName = UObject::SourceFileTagName();
 		for (TMap<FName, FString>::TConstIterator TagIt(AssetData->TagsAndValues); TagIt; ++TagIt)
 		{
 			auto OldTagAssets = CachedAssetsByTag.Find(TagIt.Key());
@@ -1831,7 +1831,7 @@ bool FAssetRegistry::RemoveAssetData(FAssetData* AssetData)
 
 #if WITH_EDITORONLY_DATA
 			// Don't check FName number so we still remove them even if there are multiple
-			if (GIsEditor && TagIt.Key().IsEqual(ReimportPathName, ENameCase::IgnoreCase, false))
+			if (GIsEditor && TagIt.Key().IsEqual(SourceFileTagName, ENameCase::IgnoreCase, false))
 			{
 				auto* SourceFiles = CachedAssetsBySourceFileName.Find(*FPaths::GetCleanFilename(TagIt.Value()));
 				SourceFiles->Remove(AssetData);
diff --git a/Engine/Source/Runtime/CoreUObject/Private/UObject/Obj.cpp b/Engine/Source/Runtime/CoreUObject/Private/UObject/Obj.cpp
--- a/Engine/Source/Runtime/CoreUObject/Private/UObject/Obj.cpp
+++ b/Engine/Source/Runtime/CoreUObject/Private/UObject/Obj.cpp
@@ -1181,6 +1181,12 @@ void UObject::GetAssetRegistryTags(TArray<FAssetRegistryTag>& OutTags) const
 	FAssetRegistryTag::GetAssetRegistryTagsFromSearchableProperties(this, OutTags);
 }
 
+const FName& UObject::SourceFileTagName()
+{
+	static const FName SourceFilePathName("SourceFile");
+	return SourceFilePathName;
+}
+
 #if WITH_EDITOR
 void UObject::GetAssetRegistryTagMetadata(TMap<FName, FAssetRegistryTagMetadata>& OutMetadata) const
 {
diff --git a/Engine/Source/Runtime/CoreUObject/Public/UObject/UObject.h b/Engine/Source/Runtime/CoreUObject/Public/UObject/UObject.h
--- a/Engine/Source/Runtime/CoreUObject/Public/UObject/UObject.h
+++ b/Engine/Source/Runtime/CoreUObject/Public/UObject/UObject.h
@@ -519,6 +519,9 @@ public:
 	};
 	virtual void GetAssetRegistryTags(TArray<FAssetRegistryTag>& OutTags) const;
 
+	/** Get the common tag name used for all asset source file import paths */
+	static const FName& SourceFileTagName();
+
 #if WITH_EDITOR
 	/**
 	 * Additional data pertaining to asset registry tags used by the editor
diff --git a/Engine/Source/Runtime/Engine/Classes/Animation/AnimSequence.h b/Engine/Source/Runtime/Engine/Classes/Animation/AnimSequence.h
--- a/Engine/Source/Runtime/Engine/Classes/Animation/AnimSequence.h
+++ b/Engine/Source/Runtime/Engine/Classes/Animation/AnimSequence.h
@@ -525,7 +525,7 @@ class UAnimSequence : public UAnimSequenceBase
 	uint32 bWasCompressedWithoutTranslations:1;
 
 	/** Importing data and options used for this mesh */
-	UPROPERTY(EditAnywhere, Instanced, Category=Reimport, AssetRegistrySearchable)
+	UPROPERTY(EditAnywhere, Instanced, Category=Reimport)
 	class UAssetImportData* AssetImportData;
 
 	/***  for Reimport **/
@@ -553,6 +553,7 @@ public:
 #endif // WITH_EDITOR
 	ENGINE_API virtual void BeginDestroy() override;
 	ENGINE_API virtual SIZE_T GetResourceSize(EResourceSizeMode::Type Mode) override;
+	virtual void GetAssetRegistryTags(TArray<FAssetRegistryTag>& OutTags) const;
 	// End of UObject interface
 
 	// Begin UAnimationAsset interface
diff --git a/Engine/Source/Runtime/Engine/Classes/Curves/CurveBase.h b/Engine/Source/Runtime/Engine/Classes/Curves/CurveBase.h
--- a/Engine/Source/Runtime/Engine/Classes/Curves/CurveBase.h
+++ b/Engine/Source/Runtime/Engine/Classes/Curves/CurveBase.h
@@ -446,7 +446,7 @@ class ENGINE_API UCurveBase : public UObject, public FCurveOwnerInterface
 	GENERATED_UCLASS_BODY()
 
 	/** The filename imported to create this object. Relative to this object's package, BaseDir() or absolute */
-	UPROPERTY(AssetRegistrySearchable, meta=(ReimportPath))
+	UPROPERTY()
 	FString ImportPath;
 
 	/** Get the time range across all curves */
@@ -494,6 +494,12 @@ public:
 
 	// End UCurveBase interface
 
+	// Begin UObject interface
+#if WITH_EDITORONLY_DATA
+	/** Override to ensure we write out the asset import data */
+	virtual void GetAssetRegistryTags(TArray<FAssetRegistryTag>& OutTags) const;
+#endif
+	// End UObject interface
 };
 
 //////////////////////////////////////////////////////////////////////////
diff --git a/Engine/Source/Runtime/Engine/Classes/EditorFramework/AssetImportData.h b/Engine/Source/Runtime/Engine/Classes/EditorFramework/AssetImportData.h
--- a/Engine/Source/Runtime/Engine/Classes/EditorFramework/AssetImportData.h
+++ b/Engine/Source/Runtime/Engine/Classes/EditorFramework/AssetImportData.h
@@ -12,7 +12,7 @@ class UAssetImportData : public UObject
 	GENERATED_UCLASS_BODY()
 
 	/** Path to the resource used to construct this static mesh. Relative to the object's package, BaseDir() or absolute */
-	UPROPERTY(EditAnywhere, Category=ImportSettings, AssetRegistrySearchable, meta=(ReimportPath))
+	UPROPERTY(EditAnywhere, Category=ImportSettings)
 	FString SourceFilePath;
 
 	/** Date/Time-stamp of the file from the last import */
diff --git a/Engine/Source/Runtime/Engine/Classes/Engine/CurveTable.h b/Engine/Source/Runtime/Engine/Classes/Engine/CurveTable.h
--- a/Engine/Source/Runtime/Engine/Classes/Engine/CurveTable.h
+++ b/Engine/Source/Runtime/Engine/Classes/Engine/CurveTable.h
@@ -25,7 +25,7 @@ class UCurveTable
 	GENERATED_UCLASS_BODY()
 
 	/** The filename that was used to create this object. Relative to the object's package, BaseDir() or absolute */
-	UPROPERTY(AssetRegistrySearchable, meta=(ReimportPath))
+	UPROPERTY()
 	FString ImportPath;
 
 
@@ -35,6 +35,9 @@ class UCurveTable
 	// Begin UObject interface.
 	virtual void FinishDestroy() override;
 	virtual void Serialize( FArchive& Ar ) override;
+#if WITH_EDITORONLY_DATA
+	virtual void GetAssetRegistryTags(TArray<FAssetRegistryTag>& OutTags) const;
+#endif
 	// End  UObject interface
 
 	// Begin UCurveTable interface
diff --git a/Engine/Source/Runtime/Engine/Classes/Engine/DataTable.h b/Engine/Source/Runtime/Engine/Classes/Engine/DataTable.h
--- a/Engine/Source/Runtime/Engine/Classes/Engine/DataTable.h
+++ b/Engine/Source/Runtime/Engine/Classes/Engine/DataTable.h
@@ -41,7 +41,7 @@ class UDataTable
 	UScriptStruct*			RowStruct;
 
 	/** The filename that was used to create this object. Relative to the object's package, BaseDir() or absolute */
-	UPROPERTY(AssetRegistrySearchable, meta=(ReimportPath))
+	UPROPERTY()
 	FString ImportPath;
 
 	/** Map of name of row to row data structure. */
@@ -51,6 +51,9 @@ class UDataTable
 	ENGINE_API virtual void FinishDestroy() override;
 	ENGINE_API virtual void Serialize(FArchive& Ar) override;
 	ENGINE_API static void AddReferencedObjects(UObject* InThis, FReferenceCollector& Collector);
+#if WITH_EDITORONLY_DATA
+	virtual void GetAssetRegistryTags(TArray<FAssetRegistryTag>& OutTags) const;
+#endif
 	// End  UObject interface
 
 	// Begin UDataTable interface
diff --git a/Engine/Source/Runtime/Engine/Classes/Engine/Font.h b/Engine/Source/Runtime/Engine/Classes/Engine/Font.h
--- a/Engine/Source/Runtime/Engine/Classes/Engine/Font.h
+++ b/Engine/Source/Runtime/Engine/Classes/Engine/Font.h
@@ -168,7 +168,7 @@ class UFont : public UObject, public IFontProviderInterface
 	FName LegacyFontName;
 
 	/** Embedded composite font data */
-	UPROPERTY(AssetRegistrySearchable)
+	UPROPERTY()
 	FCompositeFont CompositeFont;
 
 public:
@@ -242,6 +242,9 @@ public:
 	virtual void Serialize( FArchive& Ar ) override;
 	virtual void PostLoad() override;
 	virtual bool IsLocalizedResource() override;
+#if WITH_EDITORONLY_DATA
+	virtual void GetAssetRegistryTags(TArray<FAssetRegistryTag>& OutTags) const;
+#endif
 	// End UObject interface
 
 	/**
diff --git a/Engine/Source/Runtime/Engine/Classes/Engine/SkeletalMesh.h b/Engine/Source/Runtime/Engine/Classes/Engine/SkeletalMesh.h
--- a/Engine/Source/Runtime/Engine/Classes/Engine/SkeletalMesh.h
+++ b/Engine/Source/Runtime/Engine/Classes/Engine/SkeletalMesh.h
@@ -600,7 +600,7 @@ public:
 #if WITH_EDITORONLY_DATA
 
 	/** Importing data and options used for this mesh */
-	UPROPERTY(EditAnywhere, Instanced, Category = Reimport, AssetRegistrySearchable)
+	UPROPERTY(EditAnywhere, Instanced, Category = Reimport)
 	class UAssetImportData* AssetImportData;
 
 	/** Path to the resource used to construct this skeletal mesh */
diff --git a/Engine/Source/Runtime/Engine/Classes/Engine/StaticMesh.h b/Engine/Source/Runtime/Engine/Classes/Engine/StaticMesh.h
--- a/Engine/Source/Runtime/Engine/Classes/Engine/StaticMesh.h
+++ b/Engine/Source/Runtime/Engine/Classes/Engine/StaticMesh.h
@@ -373,7 +373,7 @@ class UStaticMesh : public UObject, public IInterface_CollisionDataProvider, pub
 
 #if WITH_EDITORONLY_DATA
 	/** Importing data and options used for this mesh */
-	UPROPERTY(VisibleAnywhere, Instanced, Category=ImportSettings, AssetRegistrySearchable)
+	UPROPERTY(VisibleAnywhere, Instanced, Category=ImportSettings)
 	class UAssetImportData* AssetImportData;
 
 	/** Path to the resource used to construct this static mesh */
diff --git a/Engine/Source/Runtime/Engine/Classes/Engine/Texture.h b/Engine/Source/Runtime/Engine/Classes/Engine/Texture.h
--- a/Engine/Source/Runtime/Engine/Classes/Engine/Texture.h
+++ b/Engine/Source/Runtime/Engine/Classes/Engine/Texture.h
@@ -343,7 +343,7 @@ private:
 
 public:
 	/** Path to the resource used to construct this texture. Relative to the object's package, BaseDir() or absolute */
-	UPROPERTY(Category=Texture, VisibleAnywhere, BlueprintReadWrite, AssetRegistrySearchable, meta=(ReimportPath))
+	UPROPERTY(Category=Texture, VisibleAnywhere, BlueprintReadWrite)
 	FString SourceFilePath;
 
 	/** Date/Time-stamp of the file from the last import. */
@@ -653,6 +653,9 @@ public:
 	ENGINE_API virtual void BeginDestroy() override;
 	ENGINE_API virtual bool IsReadyForFinishDestroy() override;
 	ENGINE_API virtual void FinishDestroy() override;
+#if WITH_EDITORONLY_DATA
+	ENGINE_API virtual void GetAssetRegistryTags(TArray<FAssetRegistryTag>& OutTags) const;
+#endif
 	// End UObject interface.
 
 	/**
diff --git a/Engine/Source/Runtime/Engine/Classes/Sound/SoundWave.h b/Engine/Source/Runtime/Engine/Classes/Sound/SoundWave.h
--- a/Engine/Source/Runtime/Engine/Classes/Sound/SoundWave.h
+++ b/Engine/Source/Runtime/Engine/Classes/Sound/SoundWave.h
@@ -211,7 +211,7 @@ class ENGINE_API USoundWave : public USoundBase
 
 #if WITH_EDITORONLY_DATA
 	/** Path to the resource used to construct this sound node wave. Relative to the object's package, BaseDir() or absolute. */
-	UPROPERTY(Category=Info, VisibleAnywhere, AssetRegistrySearchable, meta=(ReimportPath))
+	UPROPERTY(Category=Info, VisibleAnywhere)
 	FString SourceFilePath;
 
 	/** Date/Time-stamp of the file from the last import */
diff --git a/Engine/Source/Runtime/Engine/Classes/VectorField/VectorFieldStatic.h b/Engine/Source/Runtime/Engine/Classes/VectorField/VectorFieldStatic.h
--- a/Engine/Source/Runtime/Engine/Classes/VectorField/VectorFieldStatic.h
+++ b/Engine/Source/Runtime/Engine/Classes/VectorField/VectorFieldStatic.h
@@ -35,7 +35,7 @@ public:
 
 #if WITH_EDITORONLY_DATA
 	/** Path to the resource used to construct this vector field. Relative to the object's package, BaseDir() or absolute. */
-	UPROPERTY(Category=SourceAsset, VisibleAnywhere, AssetRegistrySearchable, meta=(ReimportPath))
+	UPROPERTY(Category=SourceAsset, VisibleAnywhere)
 	FString SourceFilePath;
 
 	/** Date/Time-stamp of the file from the last import */
@@ -50,6 +50,9 @@ public:
 	virtual void PostEditChangeProperty(FPropertyChangedEvent& PropertyChangedEvent) override;
 #endif // WITH_EDITOR
 	virtual void Serialize(FArchive& Ar) override;
+#if WITH_EDITORONLY_DATA
+	virtual void GetAssetRegistryTags(TArray<FAssetRegistryTag>& OutTags) const;
+#endif
 	// End UObject interface.
 
 	// Begin UVectorField Interface
diff --git a/Engine/Source/Runtime/Engine/Private/Animation/AnimSequence.cpp b/Engine/Source/Runtime/Engine/Private/Animation/AnimSequence.cpp
--- a/Engine/Source/Runtime/Engine/Private/Animation/AnimSequence.cpp
+++ b/Engine/Source/Runtime/Engine/Private/Animation/AnimSequence.cpp
@@ -122,6 +122,15 @@ SIZE_T UAnimSequence::GetResourceSize(EResourceSizeMode::Type Mode)
 	return ResourceSize;
 }
 
+void UAnimSequence::GetAssetRegistryTags(TArray<FAssetRegistryTag>& OutTags) const
+{
+#if WITH_EDITORONLY_DATA
+	OutTags.Add( FAssetRegistryTag(SourceFileTagName(), AssetImportData->SourceFilePath, FAssetRegistryTag::TT_Hidden) );
+#endif
+
+	Super::GetAssetRegistryTags(OutTags);
+}
+
 int32 UAnimSequence::GetApproxRawSize() const
 {
 	int32 Total = sizeof(FRawAnimSequenceTrack) * RawAnimationData.Num();
diff --git a/Engine/Source/Runtime/Engine/Private/CurveBase.cpp b/Engine/Source/Runtime/Engine/Private/CurveBase.cpp
--- a/Engine/Source/Runtime/Engine/Private/CurveBase.cpp
+++ b/Engine/Source/Runtime/Engine/Private/CurveBase.cpp
@@ -774,6 +774,15 @@ UCurveBase::UCurveBase(const FObjectInitializer& ObjectInitializer)
 {
 }
 
+#if WITH_EDITORONLY_DATA
+void UCurveBase::GetAssetRegistryTags(TArray<FAssetRegistryTag>& OutTags) const
+{
+	OutTags.Add( FAssetRegistryTag(SourceFileTagName(), ImportPath, FAssetRegistryTag::TT_Hidden) );
+
+	Super::GetAssetRegistryTags(OutTags);
+}
+#endif
+
 void UCurveBase::GetTimeRange(float& MinTime, float& MaxTime) const
 {
 	TArray<FRichCurveEditInfoConst> Curves = GetCurves();
diff --git a/Engine/Source/Runtime/Engine/Private/CurveTable.cpp b/Engine/Source/Runtime/Engine/Private/CurveTable.cpp
--- a/Engine/Source/Runtime/Engine/Private/CurveTable.cpp
+++ b/Engine/Source/Runtime/Engine/Private/CurveTable.cpp
@@ -91,6 +91,15 @@ void UCurveTable::FinishDestroy()
 	EmptyTable(); // Free memory when UObject goes away
 }
 
+#if WITH_EDITORONLY_DATA
+void UCurveTable::GetAssetRegistryTags(TArray<FAssetRegistryTag>& OutTags) const
+{
+	OutTags.Add( FAssetRegistryTag(SourceFileTagName(), ImportPath, FAssetRegistryTag::TT_Hidden) );
+
+	Super::GetAssetRegistryTags(OutTags);
+}
+#endif
+
 FString UCurveTable::GetTableAsString()
 {
 	FString Result;
diff --git a/Engine/Source/Runtime/Engine/Private/DataTable.cpp b/Engine/Source/Runtime/Engine/Private/DataTable.cpp
--- a/Engine/Source/Runtime/Engine/Private/DataTable.cpp
+++ b/Engine/Source/Runtime/Engine/Private/DataTable.cpp
@@ -207,6 +207,15 @@ void UDataTable::FinishDestroy()
 	}
 }
 
+#if WITH_EDITORONLY_DATA
+void UDataTable::GetAssetRegistryTags(TArray<FAssetRegistryTag>& OutTags) const
+{
+	OutTags.Add( FAssetRegistryTag(SourceFileTagName(), ImportPath, FAssetRegistryTag::TT_Hidden) );
+
+	Super::GetAssetRegistryTags(OutTags);
+}
+#endif
+
 void UDataTable::EmptyTable()
 {
 	UScriptStruct* LoadUsingStruct = RowStruct;
diff --git a/Engine/Source/Runtime/Engine/Private/Font.cpp b/Engine/Source/Runtime/Engine/Private/Font.cpp
--- a/Engine/Source/Runtime/Engine/Private/Font.cpp
+++ b/Engine/Source/Runtime/Engine/Private/Font.cpp
@@ -92,6 +92,27 @@ void UFont::PostLoad()
 	}
 }
 
+#if WITH_EDITORONLY_DATA
+void UFont::GetAssetRegistryTags(TArray<FAssetRegistryTag>& OutTags) const
+{
+	// Add all the font filenames
+	for( const FTypefaceEntry& TypefaceEntry : CompositeFont.DefaultTypeface.Fonts )
+	{
+		OutTags.Add( FAssetRegistryTag(SourceFileTagName(), TypefaceEntry.Font.FontFilename, FAssetRegistryTag::TT_Hidden) );
+	}
+
+	for( const FCompositeSubFont& SubFont : CompositeFont.SubTypefaces )
+	{
+		for( const FTypefaceEntry& TypefaceEntry : SubFont.Typeface.Fonts )
+		{
+			OutTags.Add( FAssetRegistryTag(SourceFileTagName(), TypefaceEntry.Font.FontFilename, FAssetRegistryTag::TT_Hidden) );
+		}
+	}
+
+	Super::GetAssetRegistryTags(OutTags);
+}
+#endif
+
 void UFont::CacheCharacterCountAndMaxCharHeight()
 {
 	// Cache the number of characters in the font.  Obviously this is pretty simple, but note that it will be
diff --git a/Engine/Source/Runtime/Engine/Private/SkeletalMesh.cpp b/Engine/Source/Runtime/Engine/Private/SkeletalMesh.cpp
--- a/Engine/Source/Runtime/Engine/Private/SkeletalMesh.cpp
+++ b/Engine/Source/Runtime/Engine/Private/SkeletalMesh.cpp
@@ -2705,6 +2705,13 @@ void USkeletalMesh::GetAssetRegistryTags(TArray<FAssetRegistryTag>& OutTags) con
 	OutTags.Add( FAssetRegistryTag("Triangles", FString::FromInt(NumTriangles), FAssetRegistryTag::TT_Numerical) );
 	OutTags.Add( FAssetRegistryTag("Bones", FString::FromInt(RefSkeleton.GetNum()), FAssetRegistryTag::TT_Numerical) );
 
+#if WITH_EDITORONLY_DATA
+	if (AssetImportData)
+	{
+		OutTags.Add( FAssetRegistryTag(SourceFileTagName(), AssetImportData->SourceFilePath, FAssetRegistryTag::TT_Hidden) );
+	}
+#endif
+	
 	Super::GetAssetRegistryTags(OutTags);
 }
 
diff --git a/Engine/Source/Runtime/Engine/Private/SoundWave.cpp b/Engine/Source/Runtime/Engine/Private/SoundWave.cpp
--- a/Engine/Source/Runtime/Engine/Private/SoundWave.cpp
+++ b/Engine/Source/Runtime/Engine/Private/SoundWave.cpp
@@ -127,7 +127,10 @@ FString USoundWave::GetDesc()
 void USoundWave::GetAssetRegistryTags(TArray<FAssetRegistryTag>& OutTags) const
 {
 	Super::GetAssetRegistryTags(OutTags);
-
+	
+#if WITH_EDITORONLY_DATA
+	OutTags.Add( FAssetRegistryTag(SourceFileTagName(), SourceFilePath, FAssetRegistryTag::TT_Hidden) );
+#endif
 	// GetCompressedDataSize could technically modify this->CompressedFormatData therefore it is not const, however this information
 	// is very useful in the asset registry so we will allow GetCompressedDataSize to be modified if the formats do not exist
 	USoundWave* MutableThis = const_cast<USoundWave*>(this);
diff --git a/Engine/Source/Runtime/Engine/Private/StaticMesh.cpp b/Engine/Source/Runtime/Engine/Private/StaticMesh.cpp
--- a/Engine/Source/Runtime/Engine/Private/StaticMesh.cpp
+++ b/Engine/Source/Runtime/Engine/Private/StaticMesh.cpp
@@ -1503,6 +1503,13 @@ void UStaticMesh::GetAssetRegistryTags(TArray<FAssetRegistryTag>& OutTags) const
 	OutTags.Add( FAssetRegistryTag("ApproxSize", ApproxSizeStr, FAssetRegistryTag::TT_Dimensional) );
 	OutTags.Add( FAssetRegistryTag("CollisionPrims", FString::FromInt(NumCollisionPrims), FAssetRegistryTag::TT_Numerical) );
 
+#if WITH_EDITORONLY_DATA
+	if (AssetImportData)
+	{
+		OutTags.Add( FAssetRegistryTag(SourceFileTagName(), AssetImportData->SourceFilePath, FAssetRegistryTag::TT_Hidden) );
+	}
+#endif
+
 	Super::GetAssetRegistryTags(OutTags);
 }
 
diff --git a/Engine/Source/Runtime/Engine/Private/Texture.cpp b/Engine/Source/Runtime/Engine/Private/Texture.cpp
--- a/Engine/Source/Runtime/Engine/Private/Texture.cpp
+++ b/Engine/Source/Runtime/Engine/Private/Texture.cpp
@@ -353,6 +353,14 @@ void UTexture::PreSave()
 #endif // #if WITH_EDITOR
 }
 
+#if WITH_EDITORONLY_DATA
+void UTexture::GetAssetRegistryTags(TArray<FAssetRegistryTag>& OutTags) const
+{
+	OutTags.Add( FAssetRegistryTag(SourceFileTagName(), SourceFilePath, FAssetRegistryTag::TT_Hidden) );
+
+	Super::GetAssetRegistryTags(OutTags);
+}
+#endif
 
 float UTexture::GetAverageBrightness(bool bIgnoreTrueBlack, bool bUseGrayscale)
 {
diff --git a/Engine/Source/Runtime/Engine/Private/VectorField.cpp b/Engine/Source/Runtime/Engine/Private/VectorField.cpp
--- a/Engine/Source/Runtime/Engine/Private/VectorField.cpp
+++ b/Engine/Source/Runtime/Engine/Private/VectorField.cpp
@@ -326,6 +326,15 @@ void UVectorFieldStatic::PostEditChangeProperty(FPropertyChangedEvent& PropertyC
 }
 #endif // WITH_EDITOR
 
+#if WITH_EDITORONLY_DATA
+void UVectorFieldStatic::GetAssetRegistryTags(TArray<FAssetRegistryTag>& OutTags) const
+{
+	OutTags.Add( FAssetRegistryTag(SourceFileTagName(), SourceFilePath, FAssetRegistryTag::TT_Hidden) );
+
+	Super::GetAssetRegistryTags(OutTags);
+}
+#endif
+
 class FVectorFieldCollectorResources : public FOneFrameResource
 {
 public:
diff --git a/Engine/Source/Runtime/SlateCore/Public/Fonts/CompositeFont.h b/Engine/Source/Runtime/SlateCore/Public/Fonts/CompositeFont.h
--- a/Engine/Source/Runtime/SlateCore/Public/Fonts/CompositeFont.h
+++ b/Engine/Source/Runtime/SlateCore/Public/Fonts/CompositeFont.h
@@ -42,7 +42,7 @@ struct SLATECORE_API FFontData
 	bool operator!=(const FFontData& Other) const;
 
 	/** The filename of the font to use - this may not always exist on disk, as we may have previously loaded and cached the font data inside an asset */
-	UPROPERTY(EditAnywhere, BlueprintReadWrite, Category=Font, AssetRegistrySearchable, meta=(ReimportPath))
+	UPROPERTY(EditAnywhere, BlueprintReadWrite, Category=Font)
 	FString FontFilename;
 
 	/** The data associated with the font */
