[INTEGRATE] Change 2449467 by [REDACTED_PERSON][REDACTED_HANDLE] on 2015/02/17 19:26:53

	Asset Registry: Prevent duplicate path delimiters from being created in FPathTreeNode::GetSubPaths_Recursive, which can cause a knock-on crash trying to use the resulting path from the Asset Save dialog
	[UE-9984]

diff --git a/Engine/Source/Runtime/AssetRegistry/Private/PathTreeNode.cpp b/Engine/Source/Runtime/AssetRegistry/Private/PathTreeNode.cpp
--- a/Engine/Source/Runtime/AssetRegistry/Private/PathTreeNode.cpp
+++ b/Engine/Source/Runtime/AssetRegistry/Private/PathTreeNode.cpp
@@ -163,7 +163,7 @@ const FPathTreeNode* FPathTreeNode::FindNode_Recursive(TArray<FString>& PathElem
 
 void FPathTreeNode::GetSubPaths_Recursive(const FString& CurrentPath, TSet<FName>& OutPaths, bool bRecurse) const
 {
-	FString NewPath = CurrentPath + TEXT("/") + FolderName;
+	FString NewPath = CurrentPath / FolderName;
 	OutPaths.Add(FName(*NewPath));
 
 	if(bRecurse)
