#UE4 Added experimental -PathTreeFNames mode that uses FNames to keep track of the path tree in the asset registry. This has possible performance benefits in certain projects.

diff --git a/Engine/Source/Runtime/AssetRegistry/Private/AssetRegistry.cpp b/Engine/Source/Runtime/AssetRegistry/Private/AssetRegistry.cpp
--- a/Engine/Source/Runtime/AssetRegistry/Private/AssetRegistry.cpp
+++ b/Engine/Source/Runtime/AssetRegistry/Private/AssetRegistry.cpp
@@ -32,7 +32,7 @@ EChunkProgressReportingType::Type GetChunkAvailabilityProgressType(EAssetAvailab
 }
 
 FAssetRegistry::FAssetRegistry()
-	: PathTreeRoot(TEXT(""))
+	: PathTreeRoot(FString())
 	, PreallocatedAssetDataBuffer(NULL)
 {
 	const double StartupStartTime = FPlatformTime::Seconds();
diff --git a/Engine/Source/Runtime/AssetRegistry/Private/PathTreeNode.cpp b/Engine/Source/Runtime/AssetRegistry/Private/PathTreeNode.cpp
--- a/Engine/Source/Runtime/AssetRegistry/Private/PathTreeNode.cpp
+++ b/Engine/Source/Runtime/AssetRegistry/Private/PathTreeNode.cpp
@@ -2,8 +2,19 @@
 
 #include "AssetRegistryPCH.h"
 
-FPathTreeNode::FPathTreeNode(FString InFolderName)
+bool IsUsingExperimentalFNameTreeNodes()
+{
+	static bool bUsingExperimentalFNameTreeNodes = FParse::Param(FCommandLine::Get(), TEXT("PathTreeFNames"));
+	return bUsingExperimentalFNameTreeNodes;
+}
+
+FPathTreeNode::FPathTreeNode(const FString& InFolderName)
 	: FolderName(InFolderName)
+	, FolderFName(NAME_None)
+{}
+
+FPathTreeNode::FPathTreeNode(FName InFolderName)
+	: FolderFName(InFolderName)
 {}
 
 FPathTreeNode::~FPathTreeNode()
@@ -19,41 +30,108 @@ FPathTreeNode::~FPathTreeNode()
 
 bool FPathTreeNode::CachePath(const FString& Path)
 {
-	TArray<FString> PathElements;
-	Path.ParseIntoArray(&PathElements, TEXT("/"), /*InCullEmpty=*/true);
+	if (IsUsingExperimentalFNameTreeNodes())
+	{
+		TArray<FString> PathElementStrings;
+		Path.ParseIntoArray(&PathElementStrings, TEXT("/"), /*InCullEmpty=*/true);
+
+		TArray<FName> PathElements;
+		PathElements.Reserve(PathElementStrings.Num());
+		for (const auto& Element : PathElementStrings)
+		{
+			PathElements.Add(FName(*Element));
+		}
 
-	return CachePath_Recursive(PathElements);
+		return CachePath_Recursive(PathElements);
+	}
+	else
+	{
+		TArray<FString> PathElements;
+		Path.ParseIntoArray(&PathElements, TEXT("/"), /*InCullEmpty=*/true);
+
+		return CachePath_Recursive(PathElements);
+	}
 }
 
 bool FPathTreeNode::RemoveFolder(const FString& Path)
 {
-	TArray<FString> PathElements;
-	Path.ParseIntoArray(&PathElements, TEXT("/"), /*InCullEmpty=*/true);
+	if (IsUsingExperimentalFNameTreeNodes())
+	{
+		TArray<FString> PathElementStrings;
+		Path.ParseIntoArray(&PathElementStrings, TEXT("/"), /*InCullEmpty=*/true);
+
+		TArray<FName> PathElements;
+		PathElements.Reserve(PathElementStrings.Num());
+		for (const auto& Element : PathElementStrings)
+		{
+			PathElements.Add(FName(*Element));
+		}
 
-	return RemoveFolder_Recursive(PathElements);
+		return RemoveFolder_Recursive(PathElements);
+	}
+	else
+	{
+		TArray<FString> PathElements;
+		Path.ParseIntoArray(&PathElements, TEXT("/"), /*InCullEmpty=*/true);
+
+		return RemoveFolder_Recursive(PathElements);
+	}
 }
 
 bool FPathTreeNode::GetSubPaths(const FString& BasePath, TSet<FName>& OutPaths, bool bRecurse) const
 {
-	TArray<FString> PathElements;
-	BasePath.ParseIntoArray(&PathElements, TEXT("/"), /*InCullEmpty=*/true);
+	if (IsUsingExperimentalFNameTreeNodes())
+	{
+		TArray<FString> PathElementStrings;
+		BasePath.ParseIntoArray(&PathElementStrings, TEXT("/"), /*InCullEmpty=*/true);
 
-	const FPathTreeNode* BasePathNode = FindNode_Recursive(PathElements);
+		TArray<FName> PathElements;
+		PathElements.Reserve(PathElementStrings.Num());
+		for (const auto& Element : PathElementStrings)
+		{
+			PathElements.Add(FName(*Element));
+		}
 
-	if ( BasePathNode )
-	{
-		// Found the base path, get the paths of all its children
-		for (int32 ChildIdx = 0; ChildIdx < BasePathNode->Children.Num(); ++ChildIdx)
+		const FPathTreeNode* BasePathNode = FindNode_Recursive(PathElements);
+
+		if (BasePathNode)
+		{
+			// Found the base path, get the paths of all its children
+			for (int32 ChildIdx = 0; ChildIdx < BasePathNode->Children.Num(); ++ChildIdx)
+			{
+				BasePathNode->Children[ChildIdx]->GetSubPaths_Recursive(BasePath, OutPaths, bRecurse);
+			}
+
+			return true;
+		}
+		else
 		{
-			BasePathNode->Children[ChildIdx]->GetSubPaths_Recursive(BasePath, OutPaths, bRecurse);
+			// Failed to find the base path
+			return false;
 		}
-		
-		return true;
 	}
 	else
 	{
-		// Failed to find the base path
-		return false;
+		TArray<FString> PathElements;
+		BasePath.ParseIntoArray(&PathElements, TEXT("/"), /*InCullEmpty=*/true);
+
+		const FPathTreeNode* BasePathNode = FindNode_Recursive(PathElements);
+
+		if ( BasePathNode )
+		{
+			// Found the base path, get the paths of all its children
+			for (int32 ChildIdx = 0; ChildIdx < BasePathNode->Children.Num(); ++ChildIdx)
+			{
+				BasePathNode->Children[ChildIdx]->GetSubPaths_Recursive(BasePath, OutPaths, bRecurse);
+			}
+		
+			return true;
+		}
+		else
+		{
+			// Failed to find the base path
+			return false;
+		}
 	}
 }
 
@@ -161,9 +239,114 @@ const FPathTreeNode* FPathTreeNode::FindNode_Recursive(TArray<FString>& PathElem
 	}
 }
 
+bool FPathTreeNode::CachePath_Recursive(TArray<FName>& PathElements)
+{
+	if ( PathElements.Num() > 0 )
+	{
+		// Pop the bottom element
+		FName ChildFolderName = PathElements[0];
+		PathElements.RemoveAt(0);
+
+		// Try to find a child which uses this folder name
+		FPathTreeNode* Child = NULL;
+		for (int32 ChildIdx = 0; ChildIdx < Children.Num(); ++ChildIdx)
+		{
+			if ( Children[ChildIdx]->FolderFName == ChildFolderName )
+			{
+				Child = Children[ChildIdx];
+				break;
+			}
+		}
+
+		// If one was not found, create it
+		if ( !Child )
+		{
+			int32 ChildIdx = Children.Add(new FPathTreeNode(ChildFolderName));
+			Child = Children[ChildIdx];
+
+			if ( PathElements.Num() == 0 )
+			{
+				// We added final element to the tree, return true;
+				return true;
+			}
+		}
+
+		if ( ensure(Child) )
+		{
+			return Child->CachePath_Recursive(PathElements);
+		}
+	}
+
+	return false;
+}
+
+bool FPathTreeNode::RemoveFolder_Recursive(TArray<FName>& PathElements)
+{
+	if ( PathElements.Num() > 0 )
+	{
+		// Pop the bottom element
+		FName ChildFolderName = PathElements[0];
+		PathElements.RemoveAt(0);
+
+		// Try to find a child which uses this folder name
+		for (int32 ChildIdx = 0; ChildIdx < Children.Num(); ++ChildIdx)
+		{
+			if ( Children[ChildIdx]->FolderFName == ChildFolderName )
+			{
+				FPathTreeNode* Child = Children[ChildIdx];
+
+				if ( PathElements.Num() == 0 )
+				{
+					// This is the final child in the path, thus it is the folder to remove.
+					// Remove it now and return true;
+					Children.RemoveAt(ChildIdx);
+					delete Child;
+					return true;
+				}
+				else
+				{
+					// This is not the last child, so recurse below
+					return Child->RemoveFolder_Recursive(PathElements);
+				}
+			}
+		}
+	}
+
+	// The child was not found at some point while recursing, fail the remove.
+	return false;
+}
+
+const FPathTreeNode* FPathTreeNode::FindNode_Recursive(TArray<FName>& PathElements) const
+{
+	if (PathElements.Num() == 0)
+	{
+		// Found the path node
+		return this;
+	}
+	else
+	{
+		// Must drill down further
+		// Pop the bottom element
+		FName ChildFolderName = PathElements[0];
+		PathElements.RemoveAt(0);
+
+		for (int32 ChildIdx = 0; ChildIdx < Children.Num(); ++ChildIdx)
+		{
+			if ( Children[ChildIdx]->FolderFName == ChildFolderName )
+			{
+				return Children[ChildIdx]->FindNode_Recursive(PathElements);
+			}
+		}
+
+		// Could not find the path
+		return NULL;
+	}
+}
+
+
 void FPathTreeNode::GetSubPaths_Recursive(const FString& CurrentPath, TSet<FName>& OutPaths, bool bRecurse) const
 {
-	FString NewPath = CurrentPath / FolderName;
+	FString NewPath = IsUsingExperimentalFNameTreeNodes() ? (CurrentPath / FolderFName.ToString()) : (CurrentPath / FolderName);
 	OutPaths.Add(FName(*NewPath));
 
 	if(bRecurse)
diff --git a/Engine/Source/Runtime/AssetRegistry/Private/PathTreeNode.h b/Engine/Source/Runtime/AssetRegistry/Private/PathTreeNode.h
--- a/Engine/Source/Runtime/AssetRegistry/Private/PathTreeNode.h
+++ b/Engine/Source/Runtime/AssetRegistry/Private/PathTreeNode.h
@@ -5,7 +5,8 @@
 class FPathTreeNode
 {
 public:
-	FPathTreeNode(FString InFolderName);
+	FPathTreeNode(const FString& InFolderName);
+	FPathTreeNode(FName InFolderName);
 	~FPathTreeNode();
 
 	/** Adds the path to the tree relative to this node, creating nodes as needed. Returns true if the specified path was actually added (as opposed to already existed) */
@@ -30,9 +31,17 @@ private:
 	/** Helper function for recursive all subpath gathering. Currentpath is the string version of the parent node. */
 	void GetSubPaths_Recursive(const FString& CurrentPath, TSet<FName>& OutPaths, bool bRecurse = true) const;
 
+private:
+	// Experimental FName version of PathTreeNode functions. Only used when launching with -PathTreeFNames
+	bool CachePath_Recursive(TArray<FName>& PathElements);
+	bool RemoveFolder_Recursive(TArray<FName>& PathElements);
+	const FPathTreeNode* FindNode_Recursive(TArray<FName>& PathElements) const;
+
 private:
 	/** The short folder name in the tree (no path) */
 	FString FolderName;
+	/** Experimental FName folder */
+	FName FolderFName;
 	/** The child node list for this node */
 	TArray<FPathTreeNode*> Children;
 };
\ No newline at end of file
