#UE4 Now prioritizing background search results in addition to background search files when selecting a path in an asset view.

diff --git a/Engine/Source/Runtime/AssetRegistry/Private/AssetRegistry.cpp b/Engine/Source/Runtime/AssetRegistry/Private/AssetRegistry.cpp
--- a/Engine/Source/Runtime/AssetRegistry/Private/AssetRegistry.cpp
+++ b/Engine/Source/Runtime/AssetRegistry/Private/AssetRegistry.cpp
@@ -1111,10 +1111,25 @@ void FAssetRegistry::ScanPathsSynchronous(const TArray<FString>& InPaths, bool b
 
 void FAssetRegistry::PrioritizeSearchPath(const FString& PathToPrioritize)
 {
+	// Prioritize the background search
 	if (BackgroundAssetSearch.IsValid())
 	{
 		BackgroundAssetSearch->PrioritizeSearchPath(PathToPrioritize);
 	}
+
+	// Also prioritize the queue of background search results
+	{
+		// Swap all priority files to the top of the list
+		int32 LowestNonPriorityFileIdx = 0;
+		for (int32 ResultIdx = 0; ResultIdx < BackgroundAssetResults.Num(); ++ResultIdx)
+		{
+			if (BackgroundAssetResults[ResultIdx]->PackagePath.StartsWith(PathToPrioritize))
+			{
+				BackgroundAssetResults.Swap(ResultIdx, LowestNonPriorityFileIdx);
+				LowestNonPriorityFileIdx++;
+			}
+		}
+	}
 }
 
 void FAssetRegistry::AssetCreated(UObject* NewAsset)
