Fixed cached source file paths not being updated when they had changed.

Previously the asset registry skipped the code to update cached tags is the keys were the same, but we need this code to be updated if the values of some tags are changed. The new code should be negligibly slower in the case where tags haven't changed, and potentially faster when they have.

Previously submitted in 2483050, but the change went awry.

diff --git a/Engine/Source/Runtime/AssetRegistry/Private/AssetRegistry.cpp b/Engine/Source/Runtime/AssetRegistry/Private/AssetRegistry.cpp
--- a/Engine/Source/Runtime/AssetRegistry/Private/AssetRegistry.cpp
+++ b/Engine/Source/Runtime/AssetRegistry/Private/AssetRegistry.cpp
@@ -1668,8 +1668,8 @@ void FAssetRegistry::AddAssetData(FAssetData* AssetData)
 		// Don't check FName number so we still accumulate them even if there are multiple
 		if (GIsEditor && TagIt.Key().IsEqual(UObject::SourceFileTagName(), ENameCase::IgnoreCase, false))
 		{
-			auto& SourceFiles = CachedAssetsBySourceFileName.FindOrAdd(*FPaths::GetCleanFilename(TagIt.Value()));
-			SourceFiles.Add(AssetData);
+			auto& Assets = CachedAssetsBySourceFileName.FindOrAdd(*FPaths::GetCleanFilename(TagIt.Value()));
+			Assets.Add(AssetData);
 		}
 #endif
 	}
@@ -1693,22 +1693,6 @@ void FAssetRegistry::AddAssetData(FAssetData* AssetData)
 
 void FAssetRegistry::UpdateAssetData(FAssetData* AssetData, const FAssetData& NewAssetData)
 {
-	// Determine if tags need to be remapped
-	bool bTagsChanged = AssetData->TagsAndValues.Num() != NewAssetData.TagsAndValues.Num();
-	
-	// If the old and new asset data has the same number of tags, see if any are different (its ok if values are different)
-	if (!bTagsChanged)
-	{
-		for (TMap<FName, FString>::TConstIterator TagIt(AssetData->TagsAndValues); TagIt; ++TagIt)
-		{
-			if ( !NewAssetData.TagsAndValues.Contains(TagIt.Key()) )
-			{
-				bTagsChanged = true;
-				break;
-			}
-		}
-	}
-
 	// Update ObjectPath
 	if ( AssetData->PackageName != NewAssetData.PackageName || AssetData->AssetName != NewAssetData.AssetName)
 	{
@@ -1747,43 +1731,71 @@ void FAssetRegistry::UpdateAssetData(FAssetData* AssetData, const FAssetData& Ne
 	}
 
 	// Update Tags
-	if (bTagsChanged)
 	{
 		const FName& SourceFileTagName = UObject::SourceFileTagName();
+		TArray<FString> OldSourcePaths, NewSourcePaths;
 
-		for (TMap<FName, FString>::TConstIterator TagIt(AssetData->TagsAndValues); TagIt; ++TagIt)
+		// Remove any tags that no longer exist
+		for (const auto& Pair : AssetData->TagsAndValues)
 		{
-			const FName FNameKey = TagIt.Key();
-			auto OldTagAssets = CachedAssetsByTag.Find(FNameKey);
-
-			OldTagAssets->Remove(AssetData);
-
 #if WITH_EDITORONLY_DATA
-			// Don't check FName number so we still remove them even if there are multiple
-			if (GIsEditor && TagIt.Key().IsEqual(SourceFileTagName, ENameCase::IgnoreCase, false))
+			// Accumulate a list of old source filenames
+			if (GIsEditor && Pair.Key.IsEqual(SourceFileTagName, ENameCase::IgnoreCase, false))
 			{
-				auto* SourceFiles = CachedAssetsBySourceFileName.Find(*FPaths::GetCleanFilename(TagIt.Value()));
-				SourceFiles->Remove(AssetData);
+				OldSourcePaths.Add(Pair.Value);
 			}
 #endif
+			if (!NewAssetData.TagsAndValues.Contains(Pair.Key))
+			{
+				auto OldTagAssets = CachedAssetsByTag.Find(Pair.Key);
+				OldTagAssets->Remove(AssetData);
+			}
 		}
 
-		for (TMap<FName, FString>::TConstIterator TagIt(NewAssetData.TagsAndValues); TagIt; ++TagIt)
+		// Add new tags
+		for (const auto& Pair : NewAssetData.TagsAndValues)
 		{
-			const FName FNameKey = TagIt.Key();
-			auto& NewTagAssets = CachedAssetsByTag.FindOrAdd(FNameKey);
-
-			NewTagAssets.Add(AssetData);
-
 #if WITH_EDITORONLY_DATA
-			// Don't check FName number so we still add them even if there are multiple
-			if (GIsEditor && TagIt.Key().IsEqual(SourceFileTagName, ENameCase::IgnoreCase, false))
+			// Accumulate a list of old source filenames
+			if (GIsEditor && Pair.Key.IsEqual(SourceFileTagName, ENameCase::IgnoreCase, false))
 			{
-				auto& SourceFiles = CachedAssetsBySourceFileName.FindOrAdd(*FPaths::GetCleanFilename(TagIt.Value()));
-				SourceFiles.Add(AssetData);
+				NewSourcePaths.Add(Pair.Value);
 			}
 #endif
+			if (!AssetData->TagsAndValues.Contains(Pair.Key))
+			{
+				auto& NewTagAssets = CachedAssetsByTag.FindOrAdd(Pair.Key);
+				NewTagAssets.Add(AssetData);
+			}
+		}
+
+#if WITH_EDITORONLY_DATA
+		// Fixup source file -> asset mappings
+		for (const auto& SourcePath : OldSourcePaths)
+		{
+			if (!NewSourcePaths.Contains(SourcePath))
+			{
+				FName CleanFilename = *FPaths::GetCleanFilename(SourcePath);
+				auto* Assets = CachedAssetsBySourceFileName.Find(CleanFilename);
+				if (Assets)
+				{
+					Assets->Remove(AssetData);
+					if (Assets->Num() == 0)
+					{
+						CachedAssetsBySourceFileName.Remove(CleanFilename);
+					}
+				}
+			}
+		}
+		for (const auto& SourcePath : NewSourcePaths)
+		{
+			if (!OldSourcePaths.Contains(SourcePath))
+			{
+				auto& Assets = CachedAssetsBySourceFileName.FindOrAdd(*FPaths::GetCleanFilename(SourcePath));
+				Assets.Add(AssetData);
+			}
 		}
+#endif
 	}
 
 	// Update the class map if updating a blueprint
@@ -1852,8 +1864,16 @@ bool FAssetRegistry::RemoveAssetData(FAssetData* AssetData)
 			// Don't check FName number so we still remove them even if there are multiple
 			if (GIsEditor && TagIt.Key().IsEqual(SourceFileTagName, ENameCase::IgnoreCase, false))
 			{
-				auto* SourceFiles = CachedAssetsBySourceFileName.Find(*FPaths::GetCleanFilename(TagIt.Value()));
-				SourceFiles->Remove(AssetData);
+				FName CleanFilename = *FPaths::GetCleanFilename(TagIt.Value());
+				auto* Assets = CachedAssetsBySourceFileName.Find(CleanFilename);
+				if (Assets)
+				{
+					Assets->Remove(AssetData);
+					if (Assets->Num() == 0)
+					{
+						CachedAssetsBySourceFileName.Remove(CleanFilename);
+					}
+				}
 			}
 #endif
 		}
