Fixing "Shadowed Variable" compile failure in Main / Mac build.

#UE4

diff --git a/Engine/Source/Runtime/AssetRegistry/Public/AssetData.h b/Engine/Source/Runtime/AssetRegistry/Public/AssetData.h
--- a/Engine/Source/Runtime/AssetRegistry/Public/AssetData.h
+++ b/Engine/Source/Runtime/AssetRegistry/Public/AssetData.h
@@ -277,9 +277,9 @@ public:
 			static FName FiBName = TEXT("FiB");
 			if (Ar.IsSaving())
 			{
-				TMap<FName, FString> TagsAndValues = AssetData.TagsAndValues;
-				TagsAndValues.Remove(FiBName);
-				Ar << TagsAndValues;
+				TMap<FName, FString> LocalTagsAndValues = AssetData.TagsAndValues;
+				LocalTagsAndValues.Remove(FiBName);
+				Ar << LocalTagsAndValues;
 			}
 			else if (Ar.IsLoading())
 			{
