#UE4 Bumping asset registry cache version to re-read find in blueprint data which is now stored in the asset registry

diff --git a/Engine/Source/Runtime/AssetRegistry/Private/AssetDataGatherer.cpp b/Engine/Source/Runtime/AssetRegistry/Private/AssetDataGatherer.cpp
--- a/Engine/Source/Runtime/AssetRegistry/Private/AssetDataGatherer.cpp
+++ b/Engine/Source/Runtime/AssetRegistry/Private/AssetDataGatherer.cpp
@@ -3,7 +3,7 @@
 #include "AssetRegistryPCH.h"
 
 #define MAX_FILES_TO_PROCESS_BEFORE_FLUSH 250
-#define CACHE_SERIALIZATION_VERSION 2
+#define CACHE_SERIALIZATION_VERSION 3
 
 FAssetDataGatherer::FAssetDataGatherer(const TArray<FString>& InPaths, bool bInIsSynchronous, bool bInLoadAndSaveCache)
 	: StopTaskCounter( 0 )
