Pull request, 990.
Fixed issue with AssetRegistry loading fnames.

diff --git a/Engine/Source/Runtime/AssetRegistry/Public/AssetData.h b/Engine/Source/Runtime/AssetRegistry/Public/AssetData.h
--- a/Engine/Source/Runtime/AssetRegistry/Public/AssetData.h
+++ b/Engine/Source/Runtime/AssetRegistry/Public/AssetData.h
@@ -258,8 +258,10 @@ public:
 	}
 
 	/** Operator for serialization */
-	friend FArchive& operator<<(FArchive& Ar, FAssetData& AssetData)
+	friend FArchive& operator<<(FArchive& InAr, FAssetData& AssetData)
 	{
+		FNameAsStringProxyArchive Ar(InAr);
+
 		// serialize out the asset info
 		Ar << AssetData.ObjectPath;
 		Ar << AssetData.PackagePath;
