Git hub pull has some issues reverting until better fix found.

diff --git a/Engine/Source/Runtime/AssetRegistry/Public/AssetData.h b/Engine/Source/Runtime/AssetRegistry/Public/AssetData.h
--- a/Engine/Source/Runtime/AssetRegistry/Public/AssetData.h
+++ b/Engine/Source/Runtime/AssetRegistry/Public/AssetData.h
@@ -258,9 +258,10 @@ public:
 	}
 
 	/** Operator for serialization */
-	friend FArchive& operator<<(FArchive& InAr, FAssetData& AssetData)
+	friend FArchive& operator<<(FArchive& Ar, FAssetData& AssetData)
 	{
-		FNameAsStringProxyArchive Ar(InAr);
+
+		//FNameAsStringProxyArchive Ar(InAr);
 
 		// serialize out the asset info
 		Ar << AssetData.ObjectPath;
