Adding missing changes for previous commit

diff --git a/Engine/Source/Editor/UnrealEd/Private/UnrealEdSrv.cpp b/Engine/Source/Editor/UnrealEd/Private/UnrealEdSrv.cpp
--- a/Engine/Source/Editor/UnrealEd/Private/UnrealEdSrv.cpp
+++ b/Engine/Source/Editor/UnrealEd/Private/UnrealEdSrv.cpp
@@ -1102,20 +1102,27 @@ bool UUnrealEdEngine::Exec( UWorld* InWorld, const TCHAR* Stream, FOutputDevice&
 		{
 			static bool RemoveSourcePath( UAssetImportData* Data, const TArray<FString>& SearchTerms )
 			{
-				const FString& SourceFilePath = Data->SourceFilePath;
-				if( !SourceFilePath.IsEmpty() )
+				FAssetImportInfo AssetImportInfo;
+
+				bool bModified = false;
+				for (const auto& File : Data->GetSourceFileData())
 				{
-					for (const FString& SearchTerm : SearchTerms)
+					if( !File.RelativeFilename.IsEmpty() && !SearchTerms.ContainsByPredicate([&](const FString& SearchTerm){ return File.RelativeFilename.Contains(SearchTerm); }) )
 					{
-						if (SourceFilePath.Contains(SearchTerm))
-						{
-							Data->Modify();
-							UE_LOG(LogUnrealEdSrv, Log, TEXT("Removing Path: %s"), *SourceFilePath);
-							Data->SourceFilePath.Empty();
-							Data->SourceFileTimestamp.Empty();
-							return true;
-						}
+						AssetImportInfo.Insert(File);
 					}
+					else
+					{
+						UE_LOG(LogUnrealEdSrv, Log, TEXT("Removing Path: %s"), *File.RelativeFilename);
+						bModified = true;
+					}
+				}
+
+				if (bModified)
+				{
+					Data->Modify();
+					Data->CopyFrom(AssetImportInfo);
+					return true;
 				}
 
 				return false;
diff --git a/Engine/Source/Runtime/AssetRegistry/Private/AssetRegistry.cpp b/Engine/Source/Runtime/AssetRegistry/Private/AssetRegistry.cpp
--- a/Engine/Source/Runtime/AssetRegistry/Private/AssetRegistry.cpp
+++ b/Engine/Source/Runtime/AssetRegistry/Private/AssetRegistry.cpp
@@ -169,9 +169,6 @@ FAssetRegistry::~FAssetRegistry()
 	CachedAssetsByClass.Empty();
 	CachedAssetsByTag.Empty();
 	CachedDependsNodes.Empty();
-#if WITH_EDITORONLY_DATA
-	CachedAssetsBySourceFileName.Empty();
-#endif
 
 	// Stop listening for content mount point events
 	FPackageName::OnContentPathMounted().RemoveAll( this );
@@ -552,21 +549,6 @@ bool FAssetRegistry::GetAssets(const FARFilter& Filter, TArray<FAssetData>& OutA
 		}
 	}
 
-#if WITH_EDITORONLY_DATA
-	// Asset source file filter
-	if ( GIsEditor && Filter.SourceFilenames.Num() )
-	{
-		for (const auto& Filename : Filter.SourceFilenames)
-		{
-			const auto* Assets = CachedAssetsBySourceFileName.Find(Filename);
-			if (Assets)
-			{
-				DiskFilterSets.Add(*Assets);
-			}
-		}
-	}
-#endif
-
 	// If we have any filter sets, add the assets which are contained in the sets to OutAssetData
 	if ( DiskFilterSets.Num() > 0 )
 	{
@@ -1684,15 +1666,6 @@ void FAssetRegistry::AddAssetData(FAssetData* AssetData)
 
 		auto& TagAssets = CachedAssetsByTag.FindOrAdd(Key);
 		TagAssets.Add(AssetData);
-
-#if WITH_EDITORONLY_DATA
-		// Don't check FName number so we still accumulate them even if there are multiple
-		if (GIsEditor && TagIt.Key().IsEqual(UObject::SourceFileTagName(), ENameCase::IgnoreCase, false))
-		{
-			auto& Assets = CachedAssetsBySourceFileName.FindOrAdd(*FPaths::GetCleanFilename(TagIt.Value()));
-			Assets.Add(AssetData);
-		}
-#endif
 	}
 
 	// Notify subscribers
@@ -1714,6 +1687,22 @@ void FAssetRegistry::AddAssetData(FAssetData* AssetData)
 
 void FAssetRegistry::UpdateAssetData(FAssetData* AssetData, const FAssetData& NewAssetData)
 {
+	// Determine if tags need to be remapped
+	bool bTagsChanged = AssetData->TagsAndValues.Num() != NewAssetData.TagsAndValues.Num();
+	
+	// If the old and new asset data has the same number of tags, see if any are different (its ok if values are different)
+	if (!bTagsChanged)
+	{
+		for (TMap<FName, FString>::TConstIterator TagIt(AssetData->TagsAndValues); TagIt; ++TagIt)
+		{
+			if ( !NewAssetData.TagsAndValues.Contains(TagIt.Key()) )
+			{
+				bTagsChanged = true;
+				break;
+			}
+		}
+	}
+
 	// Update ObjectPath
 	if ( AssetData->PackageName != NewAssetData.PackageName || AssetData->AssetName != NewAssetData.AssetName)
 	{
@@ -1752,71 +1741,23 @@ void FAssetRegistry::UpdateAssetData(FAssetData* AssetData, const FAssetData& Ne
 	}
 
 	// Update Tags
+	if (bTagsChanged)
 	{
-		const FName& SourceFileTagName = UObject::SourceFileTagName();
-		TArray<FString> OldSourcePaths, NewSourcePaths;
-
-		// Remove any tags that no longer exist
-		for (const auto& Pair : AssetData->TagsAndValues)
+		for (TMap<FName, FString>::TConstIterator TagIt(AssetData->TagsAndValues); TagIt; ++TagIt)
 		{
-#if WITH_EDITORONLY_DATA
-			// Accumulate a list of old source filenames
-			if (GIsEditor && Pair.Key.IsEqual(SourceFileTagName, ENameCase::IgnoreCase, false))
-			{
-				OldSourcePaths.Add(Pair.Value);
-			}
-#endif
-			if (!NewAssetData.TagsAndValues.Contains(Pair.Key))
-			{
-				auto OldTagAssets = CachedAssetsByTag.Find(Pair.Key);
-				OldTagAssets->Remove(AssetData);
-			}
-		}
+			const FName FNameKey = TagIt.Key();
+			auto OldTagAssets = CachedAssetsByTag.Find(FNameKey);
 
-		// Add new tags
-		for (const auto& Pair : NewAssetData.TagsAndValues)
-		{
-#if WITH_EDITORONLY_DATA
-			// Accumulate a list of old source filenames
-			if (GIsEditor && Pair.Key.IsEqual(SourceFileTagName, ENameCase::IgnoreCase, false))
-			{
-				NewSourcePaths.Add(Pair.Value);
-			}
-#endif
-			if (!AssetData->TagsAndValues.Contains(Pair.Key))
-			{
-				auto& NewTagAssets = CachedAssetsByTag.FindOrAdd(Pair.Key);
-				NewTagAssets.Add(AssetData);
-			}
+			OldTagAssets->Remove(AssetData);
 		}
 
-#if WITH_EDITORONLY_DATA
-		// Fixup source file -> asset mappings
-		for (const auto& SourcePath : OldSourcePaths)
+		for (TMap<FName, FString>::TConstIterator TagIt(NewAssetData.TagsAndValues); TagIt; ++TagIt)
 		{
-			if (!NewSourcePaths.Contains(SourcePath))
-			{
-				FName CleanFilename = *FPaths::GetCleanFilename(SourcePath);
-				auto* Assets = CachedAssetsBySourceFileName.Find(CleanFilename);
-				if (Assets)
-				{
-					Assets->Remove(AssetData);
-					if (Assets->Num() == 0)
-					{
-						CachedAssetsBySourceFileName.Remove(CleanFilename);
-					}
-				}
-			}
-		}
-		for (const auto& SourcePath : NewSourcePaths)
-		{
-			if (!OldSourcePaths.Contains(SourcePath))
-			{
-				auto& Assets = CachedAssetsBySourceFileName.FindOrAdd(*FPaths::GetCleanFilename(SourcePath));
-				Assets.Add(AssetData);
-			}
+			const FName FNameKey = TagIt.Key();
+			auto& NewTagAssets = CachedAssetsByTag.FindOrAdd(FNameKey);
+
+			NewTagAssets.Add(AssetData);
 		}
-#endif
 	}
 
 	// Update the class map if updating a blueprint
@@ -1874,29 +1815,10 @@ bool FAssetRegistry::RemoveAssetData(FAssetData* AssetData)
 		OldPathAssets->Remove(AssetData);
 		OldClassAssets->Remove(AssetData);
 
-		const FName& SourceFileTagName = UObject::SourceFileTagName();
 		for (TMap<FName, FString>::TConstIterator TagIt(AssetData->TagsAndValues); TagIt; ++TagIt)
 		{
 			auto OldTagAssets = CachedAssetsByTag.Find(TagIt.Key());
 			OldTagAssets->Remove(AssetData);
-
-
-#if WITH_EDITORONLY_DATA
-			// Don't check FName number so we still remove them even if there are multiple
-			if (GIsEditor && TagIt.Key().IsEqual(SourceFileTagName, ENameCase::IgnoreCase, false))
-			{
-				FName CleanFilename = *FPaths::GetCleanFilename(TagIt.Value());
-				auto* Assets = CachedAssetsBySourceFileName.Find(CleanFilename);
-				if (Assets)
-				{
-					Assets->Remove(AssetData);
-					if (Assets->Num() == 0)
-					{
-						CachedAssetsBySourceFileName.Remove(CleanFilename);
-					}
-				}
-			}
-#endif
 		}
 
 		// We need to update the cached dependencies references cache so that they know we no
diff --git a/Engine/Source/Runtime/AssetRegistry/Private/BackgroundAssetData.cpp b/Engine/Source/Runtime/AssetRegistry/Private/BackgroundAssetData.cpp
--- a/Engine/Source/Runtime/AssetRegistry/Private/BackgroundAssetData.cpp
+++ b/Engine/Source/Runtime/AssetRegistry/Private/BackgroundAssetData.cpp
@@ -2,7 +2,7 @@
 
 #include "AssetRegistryPCH.h"
 
-FBackgroundAssetData::FBackgroundAssetData(const FString& InPackageName, const FString& InPackagePath, const FString& InGroupNames, const FString& InAssetName, const FString& InAssetClass, const TMultiMap<FString, FString>& InTags, const TArray<int32>& InChunkIDs)
+FBackgroundAssetData::FBackgroundAssetData(const FString& InPackageName, const FString& InPackagePath, const FString& InGroupNames, const FString& InAssetName, const FString& InAssetClass, const TMap<FString, FString>& InTags, const TArray<int32>& InChunkIDs)
 {
 	PackageName = InPackageName;
 	PackagePath = InPackagePath;
@@ -44,23 +44,9 @@ FBackgroundAssetData::FBackgroundAssetData(const FAssetData& InAssetData)
 FAssetData FBackgroundAssetData::ToAssetData() const
 {
 	TMap<FName, FString> CopiedTagsAndValues;
-
-	// Copy over tags and values, retaining unique entries in the map for duplicate names
-	TSet<FString> VisitedKeys;
-	for (const auto& Pair : TagsAndValues)
+	for (TMap<FString,FString>::TConstIterator TagIt(TagsAndValues); TagIt; ++TagIt)
 	{
-		if (VisitedKeys.Find(Pair.Key))
-		{
-			continue;
-		}
-
-		VisitedKeys.Add(Pair.Key);
-
-		int32 NameIndex = 0;
-		for (auto It = TagsAndValues.CreateConstKeyIterator(Pair.Key); It; ++It, ++NameIndex)
-		{
-			CopiedTagsAndValues.Add(FName(FName(*It.Key()), NameIndex), *It.Value());
-		}
+		CopiedTagsAndValues.Add(FName(*TagIt.Key()), TagIt.Value());
 	}
 
 	return FAssetData(
diff --git a/Engine/Source/Runtime/AssetRegistry/Private/BackgroundAssetData.h b/Engine/Source/Runtime/AssetRegistry/Private/BackgroundAssetData.h
--- a/Engine/Source/Runtime/AssetRegistry/Private/BackgroundAssetData.h
+++ b/Engine/Source/Runtime/AssetRegistry/Private/BackgroundAssetData.h
@@ -7,7 +7,7 @@ class FBackgroundAssetData
 {
 public:
 	/** Constructor */
-	FBackgroundAssetData(const FString& InPackageName, const FString& InPackagePath, const FString& InGroupNames, const FString& InAssetName, const FString& InAssetClass, const TMultiMap<FString, FString>& InTags, const TArray<int32>& InChunkIDs);
+	FBackgroundAssetData(const FString& InPackageName, const FString& InPackagePath, const FString& InGroupNames, const FString& InAssetName, const FString& InAssetClass, const TMap<FString, FString>& InTags, const TArray<int32>& InChunkIDs);
 	FBackgroundAssetData(const FAssetData& InAssetData);
 
 	/** Creates an AssetData object based on this object */
@@ -26,7 +26,7 @@ public:
 	/** The name of the asset's class */
 	FString AssetClass;
 	/** The map of values for properties that were marked AssetRegistrySearchable */
-	TMultiMap<FString, FString> TagsAndValues;
+	TMap<FString, FString> TagsAndValues;
 	/** The IDs of the chunks this asset is located in for streaming install.  Empty if not assigned to a chunk */
 	TArray<int32> ChunkIDs;
 };
diff --git a/Engine/Source/Runtime/AssetRegistry/Public/ARFilter.h b/Engine/Source/Runtime/AssetRegistry/Public/ARFilter.h
--- a/Engine/Source/Runtime/AssetRegistry/Public/ARFilter.h
+++ b/Engine/Source/Runtime/AssetRegistry/Public/ARFilter.h
@@ -15,11 +15,6 @@ struct FARFilter
 	TArray<FName> ClassNames;
 	/** The filter component for properties marked with the AssetRegistrySearchable flag */
 	TMultiMap<FName, FString> TagsAndValues;
-#if WITH_EDITORONLY_DATA
-	/** The filter component for asset source file names */
-	TArray<FName> SourceFilenames;
-#endif
-
 	/** If bRecursiveClasses is true, this results will exclude classes (including subclasses) in this list */
 	TSet<FName> RecursiveClassesExclusionSet;
 	/** If true, PackagePath components will be recursive */
@@ -43,9 +38,7 @@ struct FARFilter
 		PackagePaths.Append(Other.PackagePaths);
 		ObjectPaths.Append(Other.ObjectPaths);
 		ClassNames.Append(Other.ClassNames);
-#if WITH_EDITORONLY_DATA
-		SourceFilenames.Append(Other.SourceFilenames);
-#endif
+
 		for (auto TagIt = Other.TagsAndValues.CreateConstIterator(); TagIt; ++TagIt)
 		{
 			TagsAndValues.Add(TagIt.Key(), TagIt.Value());
@@ -61,11 +54,7 @@ struct FARFilter
 	/** Returns true if this filter has no entries */
 	bool IsEmpty() const
 	{
-		return PackageNames.Num() + PackagePaths.Num() + ObjectPaths.Num() + ClassNames.Num() + TagsAndValues.Num()
-	#if WITH_EDITORONLY_DATA
-			+ SourceFilenames.Num()
-	#endif
-			== 0;
+		return PackageNames.Num() + PackagePaths.Num() + ObjectPaths.Num() + ClassNames.Num() + TagsAndValues.Num() == 0;
 	}
 
 	/** Clears this filter of all entries */
@@ -76,9 +65,6 @@ struct FARFilter
 		ObjectPaths.Empty();
 		ClassNames.Empty();
 		TagsAndValues.Empty();
-#if WITH_EDITORONLY_DATA
-		SourceFilenames.Empty();
-#endif
 		RecursiveClassesExclusionSet.Empty();
 
 		bRecursivePaths = false;
diff --git a/Engine/Source/Runtime/Core/Private/Misc/SecureHash.cpp b/Engine/Source/Runtime/Core/Private/Misc/SecureHash.cpp
--- a/Engine/Source/Runtime/Core/Private/Misc/SecureHash.cpp
+++ b/Engine/Source/Runtime/Core/Private/Misc/SecureHash.cpp
@@ -457,6 +457,100 @@ void FMD5::Decode( uint32* output, uint8* input, int32 len )
 // 		(((uint32)input[j+2]) << 16) | (((uint32)input[j+3]) << 24);
 // }
 
+namespace LexicalConversion
+{
+	FString ToString(const FMD5Hash& Hash)
+	{
+		if (!Hash.bIsValid)
+		{
+			return FString();
+		}
+
+		return FString::Printf(TEXT("%02x%02x%02x%02x%02x%02x%02x%02x%02x%02x%02x%02x%02x%02x%02x%02x"),
+			Hash.Bytes[0], Hash.Bytes[1], Hash.Bytes[2], Hash.Bytes[3], Hash.Bytes[4], Hash.Bytes[5], Hash.Bytes[6], Hash.Bytes[7],
+			Hash.Bytes[8], Hash.Bytes[9], Hash.Bytes[10], Hash.Bytes[11], Hash.Bytes[12], Hash.Bytes[13], Hash.Bytes[14], Hash.Bytes[15]);
+	}
+
+	void FromString(FMD5Hash& Hash, const TCHAR* Buffer)
+	{
+		uint8 Bytes[16];
+		for (int32 Index = 0; Index < 16; ++Index)
+		{
+			if (Buffer[Index] == '\0')
+			{
+				return;
+			}
+
+			TCHAR Base = 0;
+			if (Buffer[Index] >= '0' && Buffer[Index] <= '9')
+			{
+				Base = '0';
+			}
+			else if (Buffer[Index] >= 'A' && Buffer[Index] <= 'F')
+			{
+				Base = 'A';
+			}
+			else if (Buffer[Index] >= 'a' && Buffer[Index] <= 'f')
+			{
+				Base = 'a';
+			}
+			else
+			{
+				// Invalid hex char
+				return;
+			}
+
+			if (Index % 2)
+			{
+				Bytes[Index] = (Buffer[Index] - Base) << 4;
+			}
+			else
+			{
+				Bytes[Index] = (Bytes[Index] & ((Buffer[Index] - Base) | 0xF0));
+			}	
+		}
+
+		FMemory::Memcpy(Hash.Bytes, Bytes, 16);
+		Hash.bIsValid = true;
+	}
+
+}
+
+FMD5Hash FMD5Hash::HashFile(const TCHAR* InFilename, TArray<uint8>* Buffer)
+{
+	FArchive* Ar = IFileManager::Get().CreateFileReader(InFilename);
+
+	FMD5Hash Hash;
+	if (Ar)
+	{
+		TArray<uint8> LocalScratch;
+		if (!Buffer)
+		{
+			LocalScratch.SetNumUninitialized(1024*64);
+			Buffer = &LocalScratch; //-V506
+		}
+		FMD5 MD5;
+
+		const int64 Size = Ar->TotalSize();
+		int64 Position = 0;
+
+		// Read in BufferSize chunks
+		while (Position < Size)
+		{
+			const auto ReadNum = FMath::Min(Size - Position, (int64)Buffer->Num());
+			Ar->Serialize(Buffer->GetData(), ReadNum);
+			MD5.Update(Buffer->GetData(), ReadNum);
+
+			Position += ReadNum;
+		}
+
+		Hash.Set(MD5);
+		delete Ar;
+	}
+
+	return Hash;
+}
+
 /*-----------------------------------------------------------------------------
 	SHA-1
 -----------------------------------------------------------------------------*/
diff --git a/Engine/Source/Runtime/Core/Public/UObject/ObjectVersion.h b/Engine/Source/Runtime/Core/Public/UObject/ObjectVersion.h
--- a/Engine/Source/Runtime/Core/Public/UObject/ObjectVersion.h
+++ b/Engine/Source/Runtime/Core/Public/UObject/ObjectVersion.h
@@ -522,6 +522,8 @@ enum EUnrealEngineObjectUE4Version
 	VER_UE4_REPLACE_SPRING_NOZ_PROPERTY,
 	// Keep enums tight and serialize their values as pairs of FName and value. Don't insert dummy values.
 	VER_UE4_TIGHTLY_PACKED_ENUMS,
+	// Changed Asset import data to serialize file meta data as JSON
+	VER_UE4_ASSET_IMPORT_DATA_AS_JSON,
 
 	// -----<new versions can be added before this line>-------------------------------------------------
 	// - this needs to be the last line (see note below)
diff --git a/Engine/Source/Runtime/CoreUObject/Private/UObject/Obj.cpp b/Engine/Source/Runtime/CoreUObject/Private/UObject/Obj.cpp
--- a/Engine/Source/Runtime/CoreUObject/Private/UObject/Obj.cpp
+++ b/Engine/Source/Runtime/CoreUObject/Private/UObject/Obj.cpp
@@ -1265,7 +1265,7 @@ void UObject::GetAssetRegistryTags(TArray<FAssetRegistryTag>& OutTags) const
 
 const FName& UObject::SourceFileTagName()
 {
-	static const FName SourceFilePathName("SourceFile");
+	static const FName SourceFilePathName("AssetImportData");
 	return SourceFilePathName;
 }
 
diff --git a/Engine/Source/Runtime/CoreUObject/Private/UObject/UObjectGlobals.cpp b/Engine/Source/Runtime/CoreUObject/Private/UObject/UObjectGlobals.cpp
--- a/Engine/Source/Runtime/CoreUObject/Private/UObject/UObjectGlobals.cpp
+++ b/Engine/Source/Runtime/CoreUObject/Private/UObject/UObjectGlobals.cpp
@@ -71,7 +71,7 @@ FCoreUObjectDelegates::FPackageCreatedForLoad FCoreUObjectDelegates::PackageCrea
 /** Check wehether we should report progress or not */
 bool ShouldReportProgress()
 {
-	return GIsEditor && !IsRunningCommandlet() && !IsAsyncLoading();
+	return GIsEditor && IsInGameThread() && !IsRunningCommandlet() && !IsAsyncLoading();
 }
 
 /**
diff --git a/Engine/Source/Runtime/Engine/Classes/Engine/CurveTable.h b/Engine/Source/Runtime/Engine/Classes/Engine/CurveTable.h
--- a/Engine/Source/Runtime/Engine/Classes/Engine/CurveTable.h
+++ b/Engine/Source/Runtime/Engine/Classes/Engine/CurveTable.h
@@ -24,20 +24,26 @@ class UCurveTable
 {
 	GENERATED_UCLASS_BODY()
 
-	/** The filename that was used to create this object. Relative to the object's package, BaseDir() or absolute */
-	UPROPERTY()
-	FString ImportPath;
-
-
 	/** Map of name of row to row data structure. */
 	TMap<FName, FRichCurve*>	RowMap;
 
 	// Begin UObject interface.
 	virtual void FinishDestroy() override;
 	virtual void Serialize( FArchive& Ar ) override;
+
 #if WITH_EDITORONLY_DATA
 	virtual void GetAssetRegistryTags(TArray<FAssetRegistryTag>& OutTags) const override;
-#endif
+	virtual void PostLoad() override;
+
+	UPROPERTY()
+	class UAssetImportData* AssetImportData;
+
+	/** The filename imported to create this object. Relative to this object's package, BaseDir() or absolute */
+	UPROPERTY()
+	FString ImportPath_DEPRECATED;
+	
+#endif	// WITH_EDITORONLY_DATA
+
 	// End  UObject interface
 
 	// Begin UCurveTable interface
diff --git a/Engine/Source/Runtime/Engine/Classes/Engine/DataTable.h b/Engine/Source/Runtime/Engine/Classes/Engine/DataTable.h
--- a/Engine/Source/Runtime/Engine/Classes/Engine/DataTable.h
+++ b/Engine/Source/Runtime/Engine/Classes/Engine/DataTable.h
@@ -40,10 +40,6 @@ class UDataTable
 	UPROPERTY()
 	UScriptStruct*			RowStruct;
 
-	/** The filename that was used to create this object. Relative to the object's package, BaseDir() or absolute */
-	UPROPERTY()
-	FString ImportPath;
-
 	/** Map of name of row to row data structure. */
 	TMap<FName, uint8*>		RowMap;
 
@@ -53,7 +49,16 @@ class UDataTable
 	ENGINE_API static void AddReferencedObjects(UObject* InThis, FReferenceCollector& Collector);
 #if WITH_EDITORONLY_DATA
 	ENGINE_API virtual void GetAssetRegistryTags(TArray<FAssetRegistryTag>& OutTags) const override;
-#endif
+	virtual void PostLoad() override;
+
+	UPROPERTY()
+	class UAssetImportData* AssetImportData;
+
+	/** The filename imported to create this object. Relative to this object's package, BaseDir() or absolute */
+	UPROPERTY()
+	FString ImportPath_DEPRECATED;
+#endif	// WITH_EDITORONLY_DATA
+
 	// End  UObject interface
 
 	// Begin UDataTable interface
diff --git a/Engine/Source/Runtime/Engine/Classes/Engine/Texture.h b/Engine/Source/Runtime/Engine/Classes/Engine/Texture.h
--- a/Engine/Source/Runtime/Engine/Classes/Engine/Texture.h
+++ b/Engine/Source/Runtime/Engine/Classes/Engine/Texture.h
@@ -357,13 +357,14 @@ private:
 	FGuid LightingGuid;
 
 public:
-	/** Path to the resource used to construct this texture. Relative to the object's package, BaseDir() or absolute */
-	UPROPERTY(Category=Texture, VisibleAnywhere, BlueprintReadWrite)
-	FString SourceFilePath;
+	
+	UPROPERTY()
+	FString SourceFilePath_DEPRECATED;
 
-	/** Date/Time-stamp of the file from the last import. */
-	UPROPERTY(Category=Texture, VisibleAnywhere, BlueprintReadWrite)
-	FString SourceFileTimestamp;
+	UPROPERTY()
+	class UAssetImportData* AssetImportData;
+
+public:
 
 	/** Static texture brightness adjustment (scales HSV value.)  (Non-destructive; Requires texture source art to be available.) */
 	UPROPERTY(EditAnywhere, BlueprintReadWrite, Category=Adjustments, meta=(DisplayName = "Brightness"))
diff --git a/Engine/Source/Runtime/Engine/Classes/Sound/SoundWave.h b/Engine/Source/Runtime/Engine/Classes/Sound/SoundWave.h
--- a/Engine/Source/Runtime/Engine/Classes/Sound/SoundWave.h
+++ b/Engine/Source/Runtime/Engine/Classes/Sound/SoundWave.h
@@ -210,14 +210,14 @@ class ENGINE_API USoundWave : public USoundBase
 	TArray<struct FLocalizedSubtitle> LocalizedSubtitles;
 
 #if WITH_EDITORONLY_DATA
-	/** Path to the resource used to construct this sound node wave. Relative to the object's package, BaseDir() or absolute. */
-	UPROPERTY(Category=Info, VisibleAnywhere)
-	FString SourceFilePath;
-
-	/** Date/Time-stamp of the file from the last import */
-	UPROPERTY(Category=Info, VisibleAnywhere)
-	FString SourceFileTimestamp;
+	UPROPERTY()
+	FString SourceFilePath_DEPRECATED;
+	UPROPERTY()
+	FString SourceFileTimestamp_DEPRECATED;
 
+	UPROPERTY()
+	class UAssetImportData* AssetImportData;
+	
 #endif // WITH_EDITORONLY_DATA
 
 public:	
diff --git a/Engine/Source/Runtime/Engine/Private/Animation/AnimSequence.cpp b/Engine/Source/Runtime/Engine/Private/Animation/AnimSequence.cpp
--- a/Engine/Source/Runtime/Engine/Private/Animation/AnimSequence.cpp
+++ b/Engine/Source/Runtime/Engine/Private/Animation/AnimSequence.cpp
@@ -127,7 +127,7 @@ void UAnimSequence::GetAssetRegistryTags(TArray<FAssetRegistryTag>& OutTags) con
 #if WITH_EDITORONLY_DATA
 	if (AssetImportData)
 	{
-		OutTags.Add( FAssetRegistryTag(SourceFileTagName(), AssetImportData->SourceFilePath, FAssetRegistryTag::TT_Hidden) );
+		OutTags.Add( FAssetRegistryTag(SourceFileTagName(), AssetImportData->ToJson(), FAssetRegistryTag::TT_Hidden) );
 	}
 #endif
 
@@ -283,11 +283,13 @@ void UAnimSequence::Serialize(FArchive& Ar)
 	{
 		if ( AssetImportData == NULL )
 		{
-			AssetImportData = NewObject<UAssetImportData>(this);
+			AssetImportData = CreateEditorOnlyDefaultSubobject<UAssetImportData>(TEXT("AssetImportData"));
 		}
 		
-		AssetImportData->SourceFilePath = SourceFilePath_DEPRECATED;
-		AssetImportData->SourceFileTimestamp = SourceFileTimestamp_DEPRECATED;
+		FAssetImportInfo Info;
+		Info.Insert(FAssetImportInfo::FSourceFile(SourceFilePath_DEPRECATED));
+		AssetImportData->CopyFrom(Info);
+		
 		SourceFilePath_DEPRECATED = TEXT("");
 		SourceFileTimestamp_DEPRECATED = TEXT("");
 	}
diff --git a/Engine/Source/Runtime/Engine/Private/CurveBase.cpp b/Engine/Source/Runtime/Engine/Private/CurveBase.cpp
--- a/Engine/Source/Runtime/Engine/Private/CurveBase.cpp
+++ b/Engine/Source/Runtime/Engine/Private/CurveBase.cpp
@@ -2,6 +2,7 @@
 
 #include "EnginePrivate.h"
 #include "CsvParser.h"
+#include "EditorFramework/AssetImportData.h"
 
 void FKeyHandleMap::Add( const FKeyHandle& InHandle, int32 InIndex )
 {
@@ -828,15 +829,29 @@ bool FRichCurve::operator==(const FRichCurve& Curve) const
 UCurveBase::UCurveBase(const FObjectInitializer& ObjectInitializer)
 	: Super(ObjectInitializer)
 {
+#if WITH_EDITORONLY_DATA
+	AssetImportData = CreateEditorOnlyDefaultSubobject<UAssetImportData>(TEXT("AssetImportData"));
+#endif
 }
 
 #if WITH_EDITORONLY_DATA
 void UCurveBase::GetAssetRegistryTags(TArray<FAssetRegistryTag>& OutTags) const
 {
-	OutTags.Add( FAssetRegistryTag(SourceFileTagName(), ImportPath, FAssetRegistryTag::TT_Hidden) );
+	OutTags.Add( FAssetRegistryTag(SourceFileTagName(), AssetImportData->ToJson(), FAssetRegistryTag::TT_Hidden) );
 
 	Super::GetAssetRegistryTags(OutTags);
 }
+
+void UCurveBase::PostLoad()
+{
+	Super::PostLoad();
+	if (!ImportPath_DEPRECATED.IsEmpty())
+	{
+		FAssetImportInfo Info;
+		Info.Insert(FAssetImportInfo::FSourceFile(ImportPath_DEPRECATED));
+		AssetImportData->CopyFrom(Info);
+	}
+}
 #endif
 
 void UCurveBase::GetTimeRange(float& MinTime, float& MaxTime) const
diff --git a/Engine/Source/Runtime/Engine/Private/CurveTable.cpp b/Engine/Source/Runtime/Engine/Private/CurveTable.cpp
--- a/Engine/Source/Runtime/Engine/Private/CurveTable.cpp
+++ b/Engine/Source/Runtime/Engine/Private/CurveTable.cpp
@@ -5,6 +5,8 @@
 #include "Json.h"
 #include "CsvParser.h"
 
+#include "EditorFramework/AssetImportData.h"
+
 DEFINE_LOG_CATEGORY(LogCurveTable);
 
 ENGINE_API const FString FCurveTableRowHandle::Unknown(TEXT("UNKNOWN"));
@@ -16,6 +18,9 @@ DECLARE_CYCLE_STAT(TEXT("CurveTableRowHandle Eval"),STAT_CurveTableRowHandleEval
 UCurveTable::UCurveTable(const FObjectInitializer& ObjectInitializer)
 	: Super(ObjectInitializer)
 {
+#if WITH_EDITORONLY_DATA
+	AssetImportData = CreateEditorOnlyDefaultSubobject<UAssetImportData>(TEXT("AssetImportData"));
+#endif
 }
 
 /** Util that removes invalid chars and then make an FName */
@@ -94,10 +99,21 @@ void UCurveTable::FinishDestroy()
 #if WITH_EDITORONLY_DATA
 void UCurveTable::GetAssetRegistryTags(TArray<FAssetRegistryTag>& OutTags) const
 {
-	OutTags.Add( FAssetRegistryTag(SourceFileTagName(), ImportPath, FAssetRegistryTag::TT_Hidden) );
+	OutTags.Add( FAssetRegistryTag(SourceFileTagName(), AssetImportData->ToJson(), FAssetRegistryTag::TT_Hidden) );
 
 	Super::GetAssetRegistryTags(OutTags);
 }
+
+void UCurveTable::PostLoad()
+{
+	Super::PostLoad();
+	if (!ImportPath_DEPRECATED.IsEmpty())
+	{
+		FAssetImportInfo Info;
+		Info.Insert(FAssetImportInfo::FSourceFile(ImportPath_DEPRECATED));
+		AssetImportData->CopyFrom(Info);
+	}
+}
 #endif
 
 FString UCurveTable::GetTableAsString() const
diff --git a/Engine/Source/Runtime/Engine/Private/DataTable.cpp b/Engine/Source/Runtime/Engine/Private/DataTable.cpp
--- a/Engine/Source/Runtime/Engine/Private/DataTable.cpp
+++ b/Engine/Source/Runtime/Engine/Private/DataTable.cpp
@@ -5,6 +5,7 @@
 #include "DataTableUtils.h"
 #include "DataTableCSV.h"
 #include "DataTableJSON.h"
+#include "EditorFramework/AssetImportData.h"
 
 DEFINE_LOG_CATEGORY(LogDataTable);
 
@@ -42,6 +43,8 @@ UDataTable::UDataTable(const FObjectInitializer& ObjectInitializer)
 			UPackage::GetTypeSpecificLocalizationDataGatheringCallbacks().Add(UDataTable::StaticClass(), &GatherDataTableForLocalization);
 		}
 	} AutomaticRegistrationOfLocalizationGatherer;
+
+	AssetImportData = CreateEditorOnlyDefaultSubobject<UAssetImportData>(TEXT("AssetImportData"));
 #endif
 }
 
@@ -160,10 +163,21 @@ void UDataTable::FinishDestroy()
 #if WITH_EDITORONLY_DATA
 void UDataTable::GetAssetRegistryTags(TArray<FAssetRegistryTag>& OutTags) const
 {
-	OutTags.Add( FAssetRegistryTag(SourceFileTagName(), ImportPath, FAssetRegistryTag::TT_Hidden) );
+	OutTags.Add( FAssetRegistryTag(SourceFileTagName(), AssetImportData->ToJson(), FAssetRegistryTag::TT_Hidden) );
 
 	Super::GetAssetRegistryTags(OutTags);
 }
+
+void UDataTable::PostLoad()
+{
+	Super::PostLoad();
+	if (!ImportPath_DEPRECATED.IsEmpty())
+	{
+		FAssetImportInfo Info;
+		Info.Insert(FAssetImportInfo::FSourceFile(ImportPath_DEPRECATED));
+		AssetImportData->CopyFrom(Info);
+	}
+}
 #endif
 
 void UDataTable::EmptyTable()
diff --git a/Engine/Source/Runtime/Engine/Private/Font.cpp b/Engine/Source/Runtime/Engine/Private/Font.cpp
--- a/Engine/Source/Runtime/Engine/Private/Font.cpp
+++ b/Engine/Source/Runtime/Engine/Private/Font.cpp
@@ -9,6 +9,8 @@
 #include "Engine/FontImportOptions.h"
 #include "SlateBasics.h"
 
+#include "EditorFramework/AssetImportData.h"
+
 UFontImportOptions::UFontImportOptions(const FObjectInitializer& ObjectInitializer)
 	: Super(ObjectInitializer)
 {
@@ -97,20 +99,24 @@ void UFont::PostLoad()
 #if WITH_EDITORONLY_DATA
 void UFont::GetAssetRegistryTags(TArray<FAssetRegistryTag>& OutTags) const
 {
+	FAssetImportInfo ImportInfo;
+
 	// Add all the font filenames
 	for( const FTypefaceEntry& TypefaceEntry : CompositeFont.DefaultTypeface.Fonts )
 	{
-		OutTags.Add( FAssetRegistryTag(SourceFileTagName(), TypefaceEntry.Font.FontFilename, FAssetRegistryTag::TT_Hidden) );
+		ImportInfo.Insert(FAssetImportInfo::FSourceFile(TypefaceEntry.Font.FontFilename));
 	}
 
 	for( const FCompositeSubFont& SubFont : CompositeFont.SubTypefaces )
 	{
 		for( const FTypefaceEntry& TypefaceEntry : SubFont.Typeface.Fonts )
 		{
-			OutTags.Add( FAssetRegistryTag(SourceFileTagName(), TypefaceEntry.Font.FontFilename, FAssetRegistryTag::TT_Hidden) );
+			ImportInfo.Insert(FAssetImportInfo::FSourceFile(TypefaceEntry.Font.FontFilename));
 		}
 	}
 
+	OutTags.Add( FAssetRegistryTag(SourceFileTagName(), ImportInfo.ToJson(), FAssetRegistryTag::TT_Hidden) );
+
 	Super::GetAssetRegistryTags(OutTags);
 }
 #endif
diff --git a/Engine/Source/Runtime/Engine/Private/SkeletalMesh.cpp b/Engine/Source/Runtime/Engine/Private/SkeletalMesh.cpp
--- a/Engine/Source/Runtime/Engine/Private/SkeletalMesh.cpp
+++ b/Engine/Source/Runtime/Engine/Private/SkeletalMesh.cpp
@@ -2350,8 +2350,10 @@ void USkeletalMesh::Serialize( FArchive& Ar )
 			AssetImportData = NewObject<UAssetImportData>(this);
 		}
 
-		AssetImportData->SourceFilePath = SourceFilePath_DEPRECATED;
-		AssetImportData->SourceFileTimestamp = SourceFileTimestamp_DEPRECATED;
+		FAssetImportInfo Info;
+		Info.Insert(FAssetImportInfo::FSourceFile(SourceFilePath_DEPRECATED));
+		AssetImportData->CopyFrom(Info);
+		
 		SourceFilePath_DEPRECATED = TEXT("");
 		SourceFileTimestamp_DEPRECATED = TEXT("");
 	}
@@ -2704,7 +2706,7 @@ void USkeletalMesh::GetAssetRegistryTags(TArray<FAssetRegistryTag>& OutTags) con
 #if WITH_EDITORONLY_DATA
 	if (AssetImportData)
 	{
-		OutTags.Add( FAssetRegistryTag(SourceFileTagName(), AssetImportData->SourceFilePath, FAssetRegistryTag::TT_Hidden) );
+		OutTags.Add( FAssetRegistryTag(SourceFileTagName(), AssetImportData->ToJson(), FAssetRegistryTag::TT_Hidden) );
 	}
 #endif
 	
diff --git a/Engine/Source/Runtime/Engine/Private/SoundWave.cpp b/Engine/Source/Runtime/Engine/Private/SoundWave.cpp
--- a/Engine/Source/Runtime/Engine/Private/SoundWave.cpp
+++ b/Engine/Source/Runtime/Engine/Private/SoundWave.cpp
@@ -10,6 +10,8 @@
 #include "AudioDerivedData.h"
 #include "SubtitleManager.h"
 #include "DerivedDataCacheInterface.h"
+#include "EditorFramework/AssetImportData.h"
+
 /*-----------------------------------------------------------------------------
 	FStreamedAudioChunk
 -----------------------------------------------------------------------------*/
@@ -59,6 +61,10 @@ USoundWave::USoundWave(const FObjectInitializer& ObjectInitializer)
 	Volume = 1.0;
 	Pitch = 1.0;
 	CompressionQuality = 40;
+
+#if WITH_EDITORONLY_DATA
+	AssetImportData = CreateEditorOnlyDefaultSubobject<UAssetImportData>(TEXT("AssetImportData"));
+#endif
 }
 
 SIZE_T USoundWave::GetResourceSize(EResourceSizeMode::Type Mode)
@@ -129,7 +135,7 @@ void USoundWave::GetAssetRegistryTags(TArray<FAssetRegistryTag>& OutTags) const
 	Super::GetAssetRegistryTags(OutTags);
 	
 #if WITH_EDITORONLY_DATA
-	OutTags.Add( FAssetRegistryTag(SourceFileTagName(), SourceFilePath, FAssetRegistryTag::TT_Hidden) );
+	OutTags.Add( FAssetRegistryTag(SourceFileTagName(), AssetImportData->ToJson(), FAssetRegistryTag::TT_Hidden) );
 #endif
 	// GetCompressedDataSize could technically modify this->CompressedFormatData therefore it is not const, however this information
 	// is very useful in the asset registry so we will allow GetCompressedDataSize to be modified if the formats do not exist
@@ -349,6 +355,15 @@ void USoundWave::PostLoad()
 		IStreamingManager::Get().GetAudioStreamingManager().AddStreamingSoundWave(this);
 	}
 
+#if WITH_EDITORONLY_DATA
+	if (!SourceFilePath_DEPRECATED.IsEmpty())
+	{
+		FAssetImportInfo Info;
+		Info.Insert(FAssetImportInfo::FSourceFile(SourceFilePath_DEPRECATED));
+		AssetImportData->CopyFrom(Info);
+	}
+#endif // #if WITH_EDITORONLY_DATA
+
 	INC_FLOAT_STAT_BY( STAT_AudioBufferTime, Duration );
 	INC_FLOAT_STAT_BY( STAT_AudioBufferTimeChannels, NumChannels * Duration );
 }
diff --git a/Engine/Source/Runtime/Engine/Private/Texture.cpp b/Engine/Source/Runtime/Engine/Private/Texture.cpp
--- a/Engine/Source/Runtime/Engine/Private/Texture.cpp
+++ b/Engine/Source/Runtime/Engine/Private/Texture.cpp
@@ -12,6 +12,7 @@
 
 #if WITH_EDITORONLY_DATA
 #include "DDSLoader.h"
+#include "EditorFramework/AssetImportData.h"
 #endif
 #include "Engine/TextureCube.h"
 
@@ -58,6 +59,9 @@ UTexture::UTexture(const FObjectInitializer& ObjectInitializer)
 	PaddingColor = FColor::Black;
 	ChromaKeyColor = FColorList::Magenta;
 	ChromaKeyThreshold = 1.0f / 255.0f;
+
+	AssetImportData = CreateEditorOnlyDefaultSubobject<UAssetImportData>(TEXT("AssetImportData"));
+	
 #endif // #if WITH_EDITORONLY_DATA
 
 	if (FApp::CanEverRender() && !IsTemplate())
@@ -265,6 +269,15 @@ void UTexture::PostLoad()
 {
 	Super::PostLoad();
 
+#if WITH_EDITORONLY_DATA
+	if (!SourceFilePath_DEPRECATED.IsEmpty())
+	{
+		FAssetImportInfo Info;
+		Info.Insert(FAssetImportInfo::FSourceFile(SourceFilePath_DEPRECATED));
+		AssetImportData->CopyFrom(Info);
+	}
+#endif
+
 	if( !IsTemplate() )
 	{
 		// Update cached LOD bias.
@@ -378,7 +391,7 @@ void UTexture::PreSave()
 #if WITH_EDITORONLY_DATA
 void UTexture::GetAssetRegistryTags(TArray<FAssetRegistryTag>& OutTags) const
 {
-	OutTags.Add( FAssetRegistryTag(SourceFileTagName(), SourceFilePath, FAssetRegistryTag::TT_Hidden) );
+	OutTags.Add( FAssetRegistryTag(SourceFileTagName(), AssetImportData->ToJson(), FAssetRegistryTag::TT_Hidden) );
 
 	Super::GetAssetRegistryTags(OutTags);
 }
diff --git a/Engine/Source/Runtime/Engine/Private/VectorField.cpp b/Engine/Source/Runtime/Engine/Private/VectorField.cpp
--- a/Engine/Source/Runtime/Engine/Private/VectorField.cpp
+++ b/Engine/Source/Runtime/Engine/Private/VectorField.cpp
@@ -21,6 +21,10 @@
 #include "ComponentReregisterContext.h"
 #include "Components/VectorFieldComponent.h"
 
+#if WITH_EDITORONLY_DATA
+	#include "EditorFramework/AssetImportData.h"
+#endif
+	
 #define MAX_GLOBAL_VECTOR_FIELDS (16)
 DEFINE_LOG_CATEGORY(LogVectorField)
 
@@ -257,6 +261,9 @@ private:
 UVectorFieldStatic::UVectorFieldStatic(const FObjectInitializer& ObjectInitializer)
 	: Super(ObjectInitializer)
 {
+#if WITH_EDITORONLY_DATA
+	AssetImportData = CreateEditorOnlyDefaultSubobject<UAssetImportData>(TEXT("AssetImportData"));
+#endif
 }
 
 void UVectorFieldStatic::InitInstance(FVectorFieldInstance* Instance, bool bPreviewInstance)
@@ -317,6 +324,15 @@ void UVectorFieldStatic::PostLoad()
 	{
 		InitResource();
 	}
+
+#if WITH_EDITORONLY_DATA
+	if (!SourceFilePath_DEPRECATED.IsEmpty())
+	{
+		FAssetImportInfo Info;
+		Info.Insert(FAssetImportInfo::FSourceFile(SourceFilePath_DEPRECATED));
+		AssetImportData->CopyFrom(Info);
+	}
+#endif
 }
 
 void UVectorFieldStatic::BeginDestroy()
@@ -336,7 +352,7 @@ void UVectorFieldStatic::PostEditChangeProperty(FPropertyChangedEvent& PropertyC
 #if WITH_EDITORONLY_DATA
 void UVectorFieldStatic::GetAssetRegistryTags(TArray<FAssetRegistryTag>& OutTags) const
 {
-	OutTags.Add( FAssetRegistryTag(SourceFileTagName(), SourceFilePath, FAssetRegistryTag::TT_Hidden) );
+	OutTags.Add( FAssetRegistryTag(SourceFileTagName(), AssetImportData->ToJson(), FAssetRegistryTag::TT_Hidden) );
 
 	Super::GetAssetRegistryTags(OutTags);
 }
