Optimized FAssetData::GetFullName and FAssetData::GetExportTextName

These were both doing two FName -> FString conversions (into a new string), and then passing those into FString::Printf to produce a third string, discarding the other two.

They will now create a single string, and append all the required items into it (they can also write to a string you give them).

This helps the Content Browser text filter, which calls FAssetData::GetExportTextName for each item.

diff --git a/Engine/Source/Runtime/AssetRegistry/Public/AssetData.h b/Engine/Source/Runtime/AssetRegistry/Public/AssetData.h
--- a/Engine/Source/Runtime/AssetRegistry/Public/AssetData.h
+++ b/Engine/Source/Runtime/AssetRegistry/Public/AssetData.h
@@ -113,13 +113,36 @@ public:
 	/** Returns the full name for the asset in the form: Class ObjectPath */
 	FString GetFullName() const
 	{
-		return FString::Printf(TEXT("%s %s"), *AssetClass.ToString(), *ObjectPath.ToString());
+		FString FullName;
+		GetFullName(FullName);
+		return FullName;
+	}
+
+	/** Populates OutFullName with the full name for the asset in the form: Class ObjectPath */
+	void GetFullName(FString& OutFullName) const
+	{
+		OutFullName.Reset();
+		AssetClass.AppendString(OutFullName);
+		OutFullName.AppendChar(' ');
+		ObjectPath.AppendString(OutFullName);
 	}
 
 	/** Returns the name for the asset in the form: Class'ObjectPath' */
 	FString GetExportTextName() const
 	{
-		return FString::Printf(TEXT("%s'%s'"), *AssetClass.ToString(), *ObjectPath.ToString());
+		FString ExportTextName;
+		GetExportTextName(ExportTextName);
+		return ExportTextName;
+	}
+
+	/** Populates OutExportTextName with the name for the asset in the form: Class'ObjectPath' */
+	void GetExportTextName(FString& OutExportTextName) const
+	{
+		OutExportTextName.Reset();
+		AssetClass.AppendString(OutExportTextName);
+		OutExportTextName.AppendChar('\'');
+		ObjectPath.AppendString(OutExportTextName);
+		OutExportTextName.AppendChar('\'');
 	}
 
 	/** Returns true if the this asset is a redirector. */
