Tidied a few asset registry uses: added consts, removed redundant code.

diff --git a/Engine/Source/Editor/ClassViewer/Private/SClassViewer.cpp b/Engine/Source/Editor/ClassViewer/Private/SClassViewer.cpp
--- a/Engine/Source/Editor/ClassViewer/Private/SClassViewer.cpp
+++ b/Engine/Source/Editor/ClassViewer/Private/SClassViewer.cpp
@@ -1510,9 +1510,7 @@ static TSharedPtr< FClassViewerNode > CreateNodeForClass(UClass* Class, const TM
 	if ( BlueprintList.Num() )
 	{
 		// Grab the generated class name and check it before assigning. Objects that haven't been saved since this has started to be exported do not have the information.
-		FString* GeneratedClassnamePtr = BlueprintList[0].TagsAndValues.Find(FName("GeneratedClass"));
-
-		if ( GeneratedClassnamePtr )
+		if (auto GeneratedClassnamePtr = BlueprintList[0].TagsAndValues.Find(FName("GeneratedClass")))
 		{
 			NewNode->GeneratedClassname = FName(**GeneratedClassnamePtr);
 		}
diff --git a/Engine/Source/Editor/ContentBrowser/Private/AssetContextMenu.cpp b/Engine/Source/Editor/ContentBrowser/Private/AssetContextMenu.cpp
--- a/Engine/Source/Editor/ContentBrowser/Private/AssetContextMenu.cpp
+++ b/Engine/Source/Editor/ContentBrowser/Private/AssetContextMenu.cpp
@@ -511,7 +511,7 @@ bool FAssetContextMenu::AddDocumentationMenuOptions(FMenuBuilder& MenuBuilder)
 		const bool bIsBlueprint = SelectedClass->IsChildOf<UBlueprint>();
 		if (bIsBlueprint)
 		{
-			FString* ParentClassPath = SelectedAssets[0].TagsAndValues.Find(GET_MEMBER_NAME_CHECKED(UBlueprint,ParentClass));
+			auto ParentClassPath = SelectedAssets[0].TagsAndValues.Find(GET_MEMBER_NAME_CHECKED(UBlueprint,ParentClass));
 			if (ParentClassPath)
 			{
 				SelectedClass = FindObject<UClass>(nullptr,**ParentClassPath);
@@ -560,7 +560,7 @@ bool FAssetContextMenu::AddDocumentationMenuOptions(FMenuBuilder& MenuBuilder)
 						}
 
 						UEnum* BlueprintTypeEnum = FindObject<UEnum>(ANY_PACKAGE, TEXT("EBlueprintType"), true);
-						FString* EnumString = SelectedAssets[0].TagsAndValues.Find(GET_MEMBER_NAME_CHECKED(UBlueprint,BlueprintType));
+						auto EnumString = SelectedAssets[0].TagsAndValues.Find(GET_MEMBER_NAME_CHECKED(UBlueprint,BlueprintType));
 						EBlueprintType BlueprintType = (EnumString ? (EBlueprintType)BlueprintTypeEnum->FindEnumIndex(**EnumString) : BPTYPE_Normal);
 
 						switch (BlueprintType)
diff --git a/Engine/Source/Editor/Kismet/Private/BlueprintEditor.cpp b/Engine/Source/Editor/Kismet/Private/BlueprintEditor.cpp
--- a/Engine/Source/Editor/Kismet/Private/BlueprintEditor.cpp
+++ b/Engine/Source/Editor/Kismet/Private/BlueprintEditor.cpp
@@ -1484,20 +1484,10 @@ void FBlueprintEditor::LoadLibrariesFromAssetRegistry()
 
 			for (int32 AssetIndex = 0; AssetIndex < AssetData.Num(); ++AssetIndex)
 			{
-				bool FoundBPType = false;
-				for (TMap<FName, FString>::TConstIterator TagIt(AssetData[ AssetIndex ].TagsAndValues); TagIt; ++TagIt)
-				{
-					FString TagValue = AssetData[ AssetIndex ].TagsAndValues.FindRef(BPTypeName);
-
-					//Only check for Blueprint Macros & Functions in the asset data for loading
-					if ( TagValue == BPMacroTypeStr || TagValue == BPFunctionTypeStr )
-					{
-						FoundBPType = true;
-						break;
-					}
-				}
+				FString TagValue = AssetData[ AssetIndex ].TagsAndValues.FindRef(BPTypeName);
 
-				if( FoundBPType )
+				//Only check for Blueprint Macros & Functions in the asset data for loading
+				if ( TagValue == BPMacroTypeStr || TagValue == BPFunctionTypeStr )
 				{
 					FString BlueprintPath = AssetData[AssetIndex].ObjectPath.ToString();
 
diff --git a/Engine/Source/Editor/Persona/Private/PersonaModule.cpp b/Engine/Source/Editor/Persona/Private/PersonaModule.cpp
--- a/Engine/Source/Editor/Persona/Private/PersonaModule.cpp
+++ b/Engine/Source/Editor/Persona/Private/PersonaModule.cpp
@@ -37,18 +37,8 @@ void FPersonaModule::StartupModule()
 
 		for (int32 AssetIndex = 0; AssetIndex < AssetData.Num(); ++AssetIndex)
 		{
-			bool FoundBPNotify = false;
-			for (TMap<FName, FString>::TConstIterator TagIt(AssetData[ AssetIndex ].TagsAndValues); TagIt; ++TagIt)
-			{
-				FString TagValue = AssetData[ AssetIndex ].TagsAndValues.FindRef(BPParentClassName);
-				if(TagValue == BPAnimNotify)
-				{
-					FoundBPNotify = true;
-					break;
-				}
-			}
-
-			if(FoundBPNotify)
+			FString TagValue = AssetData[ AssetIndex ].TagsAndValues.FindRef(BPParentClassName);
+			if (TagValue == BPAnimNotify)
 			{
 				FString BlueprintPath = AssetData[AssetIndex].ObjectPath.ToString();
 				LoadObject<UBlueprint>(NULL, *BlueprintPath, NULL, 0, NULL);
diff --git a/Engine/Source/Editor/UnrealEd/Private/ComponentTypeRegistry.cpp b/Engine/Source/Editor/UnrealEd/Private/ComponentTypeRegistry.cpp
--- a/Engine/Source/Editor/UnrealEd/Private/ComponentTypeRegistry.cpp
+++ b/Engine/Source/Editor/UnrealEd/Private/ComponentTypeRegistry.cpp
@@ -362,7 +362,7 @@ void FComponentTypeRegistryData::Tick(float)
 			const FName BPParentClassName(GET_MEMBER_NAME_CHECKED(UBlueprint, ParentClass));
 
 			bool FoundBPNotify = false;
-			for (TMap<FName, FString>::TConstIterator TagIt(Asset.TagsAndValues); TagIt; ++TagIt)
+			for (auto TagIt = Asset.TagsAndValues.CreateConstIterator(); TagIt; ++TagIt)
 			{
 				FString TagValue = Asset.TagsAndValues.FindRef(BPParentClassName);
 				const FString ObjectPath = FPackageName::ExportTextPathToObjectPath(TagValue);
diff --git a/Engine/Source/Runtime/AssetRegistry/Private/AssetRegistry.cpp b/Engine/Source/Runtime/AssetRegistry/Private/AssetRegistry.cpp
--- a/Engine/Source/Runtime/AssetRegistry/Private/AssetRegistry.cpp
+++ b/Engine/Source/Runtime/AssetRegistry/Private/AssetRegistry.cpp
@@ -1660,7 +1660,7 @@ void FAssetRegistry::AddAssetData(FAssetData* AssetData)
 	PathAssets.Add(AssetData);
 	ClassAssets.Add(AssetData);
 
-	for (TMap<FName, FString>::TConstIterator TagIt(AssetData->TagsAndValues); TagIt; ++TagIt)
+	for (auto TagIt = AssetData->TagsAndValues.CreateConstIterator(); TagIt; ++TagIt)
 	{
 		FName Key = TagIt.Key();
 
@@ -1674,8 +1674,8 @@ void FAssetRegistry::AddAssetData(FAssetData* AssetData)
 	// Populate the class map if adding blueprint
 	if (ClassGeneratorNames.Contains(AssetData->AssetClass))
 	{
-		FString* GeneratedClassPtr = AssetData->TagsAndValues.Find("GeneratedClass");
-		FString* ParentClassPtr = AssetData->TagsAndValues.Find("ParentClass");
+		auto GeneratedClassPtr = AssetData->TagsAndValues.Find("GeneratedClass");
+		auto ParentClassPtr = AssetData->TagsAndValues.Find("ParentClass");
 		if ( GeneratedClassPtr && ParentClassPtr )
 		{
 			const FName GeneratedClassFName = FName(*ExportTextPathToObjectName(*GeneratedClassPtr));
@@ -1693,7 +1693,7 @@ void FAssetRegistry::UpdateAssetData(FAssetData* AssetData, const FAssetData& Ne
 	// If the old and new asset data has the same number of tags, see if any are different (its ok if values are different)
 	if (!bTagsChanged)
 	{
-		for (TMap<FName, FString>::TConstIterator TagIt(AssetData->TagsAndValues); TagIt; ++TagIt)
+		for (auto TagIt = AssetData->TagsAndValues.CreateConstIterator(); TagIt; ++TagIt)
 		{
 			if ( !NewAssetData.TagsAndValues.Contains(TagIt.Key()) )
 			{
@@ -1743,7 +1743,7 @@ void FAssetRegistry::UpdateAssetData(FAssetData* AssetData, const FAssetData& Ne
 	// Update Tags
 	if (bTagsChanged)
 	{
-		for (TMap<FName, FString>::TConstIterator TagIt(AssetData->TagsAndValues); TagIt; ++TagIt)
+		for (auto TagIt = AssetData->TagsAndValues.CreateConstIterator(); TagIt; ++TagIt)
 		{
 			const FName FNameKey = TagIt.Key();
 			auto OldTagAssets = CachedAssetsByTag.Find(FNameKey);
@@ -1751,7 +1751,7 @@ void FAssetRegistry::UpdateAssetData(FAssetData* AssetData, const FAssetData& Ne
 			OldTagAssets->Remove(AssetData);
 		}
 
-		for (TMap<FName, FString>::TConstIterator TagIt(NewAssetData.TagsAndValues); TagIt; ++TagIt)
+		for (auto TagIt = AssetData->TagsAndValues.CreateConstIterator(); TagIt; ++TagIt)
 		{
 			const FName FNameKey = TagIt.Key();
 			auto& NewTagAssets = CachedAssetsByTag.FindOrAdd(FNameKey);
@@ -1763,7 +1763,7 @@ void FAssetRegistry::UpdateAssetData(FAssetData* AssetData, const FAssetData& Ne
 	// Update the class map if updating a blueprint
 	if (ClassGeneratorNames.Contains(AssetData->AssetClass))
 	{
-		FString* OldGeneratedClassPtr = AssetData->TagsAndValues.Find("GeneratedClass");
+		auto OldGeneratedClassPtr = AssetData->TagsAndValues.Find("GeneratedClass");
 
 		if ( OldGeneratedClassPtr )
 		{
@@ -1797,7 +1797,7 @@ bool FAssetRegistry::RemoveAssetData(FAssetData* AssetData)
 		// Remove from the class map if removing a blueprint
 		if (ClassGeneratorNames.Contains(AssetData->AssetClass))
 		{
-			FString* OldGeneratedClassPtr = AssetData->TagsAndValues.Find("GeneratedClass");
+			auto OldGeneratedClassPtr = AssetData->TagsAndValues.Find("GeneratedClass");
 
 			if ( OldGeneratedClassPtr )
 			{
@@ -1815,7 +1815,7 @@ bool FAssetRegistry::RemoveAssetData(FAssetData* AssetData)
 		OldPathAssets->Remove(AssetData);
 		OldClassAssets->Remove(AssetData);
 
-		for (TMap<FName, FString>::TConstIterator TagIt(AssetData->TagsAndValues); TagIt; ++TagIt)
+		for (auto TagIt = AssetData->TagsAndValues.CreateConstIterator(); TagIt; ++TagIt)
 		{
 			auto OldTagAssets = CachedAssetsByTag.Find(TagIt.Key());
 			OldTagAssets->Remove(AssetData);
diff --git a/Engine/Source/Runtime/AssetRegistry/Public/AssetData.h b/Engine/Source/Runtime/AssetRegistry/Public/AssetData.h
--- a/Engine/Source/Runtime/AssetRegistry/Public/AssetData.h
+++ b/Engine/Source/Runtime/AssetRegistry/Public/AssetData.h
@@ -103,7 +103,7 @@ public:
 	{
 		return ObjectPath != NAME_None;
 	}
-	
+
 	/** Returns true if this asset was found in a UAsset file */
 	bool IsUAsset() const
 	{
@@ -229,7 +229,7 @@ public:
 	{
 		return IsValid() && FindObject<UObject>(NULL, *ObjectPath.ToString()) != NULL;
 	}
-	
+
 	/** Prints the details of the asset to the log */
 	void PrintAssetData() const
 	{
@@ -242,16 +242,16 @@ public:
 		UE_LOG(LogAssetData, Log, TEXT("        AssetClass: %s"), *AssetClass.ToString());
 		UE_LOG(LogAssetData, Log, TEXT("        TagsAndValues: %d"), TagsAndValues.Num());
 
-		for (TMap<FName, FString>::TConstIterator TagIt(TagsAndValues); TagIt; ++TagIt)
+		for (const auto& TagValue: TagsAndValues)
 		{
-			UE_LOG(LogAssetData, Log, TEXT("            %s : %s"), *TagIt.Key().ToString(), *TagIt.Value());
+			UE_LOG(LogAssetData, Log, TEXT("            %s : %s"), *TagValue.Key.ToString(), *TagValue.Value);
 		}
 
 		UE_LOG(LogAssetData, Log, TEXT("        ChunkIDs: %d"), ChunkIDs.Num());
 
-		for (auto ChunkIt = ChunkIDs.CreateConstIterator(); ChunkIt; ++ChunkIt)
+		for (int32 Chunk: ChunkIDs)
 		{
-			UE_LOG(LogAssetData, Log, TEXT("                 %d"), *ChunkIt);
+			UE_LOG(LogAssetData, Log, TEXT("                 %d"), Chunk);
 		}
 	}
 
