Moved SharedMapView out of Core and next to asset registry where it's used - it really only provides a different syntax for a TSharedPtr<const TMap<...>>.

diff --git a/Engine/Source/Runtime/AssetRegistry/Public/SharedMapView.h b/Engine/Source/Runtime/AssetRegistry/Public/SharedMapView.h
new file mode 100644
--- /dev/null
+++ b/Engine/Source/Runtime/AssetRegistry/Public/SharedMapView.h
@@ -0,0 +1,109 @@
+// Copyright 1998-2015 Epic Games, Inc. All Rights Reserved.
+
+#pragma once
+
+#include "Map.h"
+
+/** Immutable wrapper of shared pointer to a map */
+template <class FKey, class FValue>
+class TSharedMapView
+{
+	typedef TMap<FKey, FValue> FMap;
+public:
+	/** Default constructor - empty map */
+	TSharedMapView() : Map(new FMap())
+	{
+	}
+
+	/** Constructor from an existing map pointer */
+	explicit TSharedMapView(TSharedRef<const FMap> InMap)
+		: Map(InMap)
+	{
+	}
+
+	/** Find a value by key (nullptr if not found) */
+	const FString* Find(FName Name) const
+	{
+		return Map->Find(Name);
+	}
+
+	/** Find a value by key (abort if not found) */
+	const FString& FindChecked(FName Name) const
+	{
+		return Map->FindChecked(Name);
+	}
+
+	/** Find a value by key (default value if not found) */
+	FString FindRef(FName Name) const
+	{
+		return Map->FindRef(Name);
+	}
+
+	/** Determine whether a key is present in the map */
+	bool Contains(FName Name) const
+	{
+		return Map->Contains(Name);
+	}
+
+	/** Retrieve size of map */
+	int32 Num() const
+	{
+		return Map->Num();
+	}
+
+	/** Populate an array with all the map's keys */
+	template <class FAllocator>
+	int32 GetKeys(TArray<FKey, FAllocator>& OutKeys) const
+	{
+		return Map->GetKeys(OutKeys);
+	}
+
+	/** Populate an array with all the map's keys */
+	template <class FAllocator>
+	int32 GenerateKeyArray(TArray<FKey, FAllocator>& OutKeys) const
+	{
+		return Map->GenerateKeyArray(OutKeys);
+	}
+
+	/** Populate an array with all the map's values */
+	template <class FAllocator>
+	int32 GenerateValueArray(TArray<FKey, FAllocator>& OutValues) const
+	{
+		return Map->GenerateValueArray(OutValues);
+	}
+
+	/** Iterate all key-value pairs */
+	typename FMap::TConstIterator CreateConstIterator() const
+	{
+		return Map->CreateConstIterator();
+	}
+
+	/** Const access to the underlying map, mainly for taking a copy */
+	const FMap& GetMap() const
+	{
+		return *Map;
+	}
+
+	/** Range for iterator access - DO NOT USE DIRECTLY */
+	friend typename FMap::TConstIterator begin(const TSharedMapView& View)
+	{
+		return begin(*View.Map);
+	}
+
+	/** Range for iterator access - DO NOT USE DIRECTLY */
+	friend typename FMap::TConstIterator end(const TSharedMapView& View)
+	{
+		return end(*View.Map);
+	}
+
+private:
+	/** Pointer to map being wrapped */
+	TSharedRef<const FMap> Map;
+};
+
+/** Helper function to wrap a TMap in a TSharedMapView */
+template <class FKey, class FValue>
+TSharedMapView<FKey, FValue> MakeSharedMapView(TMap<FKey, FValue> Map)
+{
+	return TSharedMapView<FKey, FValue>(MakeShareable(new TMap<FKey, FValue>(MoveTemp(Map))));
+}
diff --git a/Engine/Source/Runtime/Core/Public/Containers/SharedMapView.h b/Engine/Source/Runtime/Core/Public/Containers/SharedMapView.h
deleted file mode 100644
