Separate asset registry dependencies into "hard" and "soft", where hard references are required to be loaded before the asset, and soft ones are optional (string asset references basically)

diff --git a/Engine/Source/Runtime/AssetRegistry/Private/AssetRegistry.cpp b/Engine/Source/Runtime/AssetRegistry/Private/AssetRegistry.cpp
--- a/Engine/Source/Runtime/AssetRegistry/Private/AssetRegistry.cpp
+++ b/Engine/Source/Runtime/AssetRegistry/Private/AssetRegistry.cpp
@@ -11,6 +11,20 @@
 
 DEFINE_LOG_CATEGORY(LogAssetRegistry);
 
+/** Enum for tracking versions of the "mini" runtime asset registry that gets used by iterative cooking and the runtime
+	dependency preloading system */
+enum class ERuntimeRegistryVersion
+{
+	PreVersioning,					// From before file versioning was implemented
+	HardSoftDependencies,			// The first version of the runtime asset registry to include file versioning. 
+
+	LatestPlusOne,
+	Latest = (LatestPlusOne - 1),
+};
+
+/** Guid for the "mini" runtime asset registry file format, used for identification purposes */
+static const FGuid GRuntimeRegistryGuid(0x717F9EE7, 0xE9B0493A, 0x88B39132, 0x1B388107);
+
 /** Returns the appropriate ChunkProgressReportingType for the given Asset enum */
 EChunkProgressReportingType::Type GetChunkAvailabilityProgressType(EAssetAvailabilityProgressReportingType::Type ReportType)
 {
@@ -31,6 +45,31 @@ EChunkProgressReportingType::Type GetChunkAvailabilityProgressType(EAssetAvailab
 	return ChunkReportType;
 }
 
+/** Helper function for reading the header of the "mini" runtime asset registry. It deals with the header not
+    existing for old formats of the file */
+inline ERuntimeRegistryVersion ReadRuntimeRegistryVersion(FArchive& Ar)
+{
+	auto InitialLocation = Ar.Tell();
+	FGuid Guid;
+	Ar << Guid;
+
+	ERuntimeRegistryVersion Version = ERuntimeRegistryVersion::PreVersioning;
+
+	if (Guid == GRuntimeRegistryGuid)
+	{
+		int32 VersionInt;
+		Ar << VersionInt;
+		Version = (ERuntimeRegistryVersion)VersionInt;
+	}
+	else
+	{
+		// This is an old format file, so skip back to where we started
+		Ar.Seek(InitialLocation);
+	}
+
+	return Version;
+}
+
 FAssetRegistry::FAssetRegistry()
 	: PathTreeRoot(FString())
 	, PreallocatedAssetDataBuffer(NULL)
@@ -693,7 +732,7 @@ bool FAssetRegistry::GetAllAssets(TArray<FAssetData>& OutAssetData) const
 	return true;
 }
 
-bool FAssetRegistry::GetDependencies(FName PackageName, TArray<FName>& OutDependencies) const
+bool FAssetRegistry::GetDependencies(FName PackageName, TArray<FName>& OutDependencies, EAssetRegistryDependencyType::Type InDependencyType) const
 {
 	const FDependsNode*const* NodePtr = CachedDependsNodes.Find(PackageName);
 	const FDependsNode* Node = NULL;
@@ -704,14 +743,7 @@ bool FAssetRegistry::GetDependencies(FName PackageName, TArray<FName>& OutDepend
 
 	if (Node != NULL)
 	{
-		TArray<FDependsNode*> DependencyNodes;
-		Node->GetDependencies(DependencyNodes);
-
-		for ( auto NodeIt = DependencyNodes.CreateConstIterator(); NodeIt; ++NodeIt )
-		{
-			OutDependencies.Add((*NodeIt)->PackageName);
-		}
-
+		Node->GetDependencies(OutDependencies, InDependencyType);
 		return true;
 	}
 	else
@@ -736,7 +768,7 @@ bool FAssetRegistry::GetReferencers(FName PackageName, TArray<FName>& OutReferen
 
 		for ( auto NodeIt = DependencyNodes.CreateConstIterator(); NodeIt; ++NodeIt )
 		{
-			OutReferencers.Add((*NodeIt)->PackageName);
+			OutReferencers.Add((*NodeIt)->GetPackageName());
 		}
 
 		return true;
@@ -1312,6 +1344,8 @@ void FAssetRegistry::Serialize(FArchive& Ar)
 	// load in by building the TMap
 	else
 	{
+		auto Version = ReadRuntimeRegistryVersion(Ar);
+
 		// serialize number of objects
 		int32 LocalNumAssets = 0;
 		Ar << LocalNumAssets;
@@ -1332,41 +1366,59 @@ void FAssetRegistry::Serialize(FArchive& Ar)
 			AddAssetPath(NewAssetData->PackagePath.ToString());
 		}
 
-		int32 LocalNumDependsNodes = 0;
-		Ar << LocalNumDependsNodes;
+		int32 LocalNumDependsNodes = LocalNumAssets;
+		if (Version == ERuntimeRegistryVersion::PreVersioning)
+		{
+			Ar << LocalNumDependsNodes;
+		}
 
 		PreallocatedDependsNodeDataBuffer = new FDependsNode[LocalNumDependsNodes];
 
 		for (int32 DependsNodeIndex = 0; DependsNodeIndex < LocalNumDependsNodes; DependsNodeIndex++)
 		{
 			auto NewDependsNodeData = &PreallocatedDependsNodeDataBuffer[DependsNodeIndex];
-			int32 AssetIndex = 0;
-			Ar << AssetIndex;
-			NewDependsNodeData->PackageName = PreallocatedAssetDataBuffer[AssetIndex].PackageName;
+			int32 AssetIndex = DependsNodeIndex;
+			if (Version == ERuntimeRegistryVersion::PreVersioning)
+			{
+				Ar << AssetIndex;
+			}
+			NewDependsNodeData->SetPackageName(PreallocatedAssetDataBuffer[AssetIndex].PackageName);
 
-			CachedDependsNodes.Add(NewDependsNodeData->PackageName, NewDependsNodeData);
+			CachedDependsNodes.Add(NewDependsNodeData->GetPackageName(), NewDependsNodeData);
 		}
 
 		for (int32 DependsNodeIndex = 0; DependsNodeIndex < LocalNumDependsNodes; DependsNodeIndex++)
 		{
 			auto NewDependsNodeData = &PreallocatedDependsNodeDataBuffer[DependsNodeIndex];
-			int32 LocalNumDependencies = 0;
+			int32 LocalNumHardDependencies = 0;
+			int32 LocalNumSoftDependencies = 0;
 			int32 LocalNumReferencers = 0;
-			Ar << LocalNumDependencies;
+			Ar << LocalNumHardDependencies;
+			if (Version >= ERuntimeRegistryVersion::HardSoftDependencies)
+			{
+				Ar << LocalNumSoftDependencies;
+			}
 			Ar << LocalNumReferencers;
 
-			for (int32 DependencyIndex = 0; DependencyIndex < LocalNumDependencies; ++DependencyIndex)
+			for (int32 DependencyIndex = 0; DependencyIndex < LocalNumHardDependencies; ++DependencyIndex)
+			{
+				int32 Index = 0;
+				Ar << Index;
+				NewDependsNodeData->AddDependency(&PreallocatedDependsNodeDataBuffer[Index], EAssetRegistryDependencyType::Hard);
+			}
+
+			for (int32 DependencyIndex = 0; DependencyIndex < LocalNumSoftDependencies; ++DependencyIndex)
 			{
 				int32 Index = 0;
 				Ar << Index;
-				NewDependsNodeData->Dependencies.Add(&PreallocatedDependsNodeDataBuffer[Index]);
+				NewDependsNodeData->AddDependency(&PreallocatedDependsNodeDataBuffer[Index], EAssetRegistryDependencyType::Soft);
 			}
 
 			for (int32 ReferencerIndex = 0; ReferencerIndex < LocalNumReferencers; ++ReferencerIndex)
 			{
 				int32 Index = 0;
 				Ar << Index;
-				NewDependsNodeData->Referencers.Add(&PreallocatedDependsNodeDataBuffer[Index]);
+				NewDependsNodeData->AddReferencer(&PreallocatedDependsNodeDataBuffer[Index]);
 			}
 		}
 	}
@@ -1376,25 +1428,24 @@ void AddDependsNodesRecursive(FDependsNode* InDependsNode, TArray<FDependsNode*>
 {
 	if (!OutArray.Contains(InDependsNode))
 	{
-		if (InAssets.Contains(InDependsNode->PackageName))
+		if (InAssets.Contains(InDependsNode->GetPackageName()))
 		{
 			OutArray.Add(InDependsNode);
 
-			for (auto Dependency : InDependsNode->Dependencies)
-			{
-				AddDependsNodesRecursive(Dependency, OutArray, InAssets);
-			}
-
-			for (auto Referencer : InDependsNode->Referencers)
-			{
-				AddDependsNodesRecursive(Referencer, OutArray, InAssets);
-			}
+			InDependsNode->IterateOverDependencies([&](FDependsNode* InDependency, EAssetRegistryDependencyType::Type InDependencyType){AddDependsNodesRecursive(InDependency, OutArray, InAssets); });
+			InDependsNode->IterateOverReferencers([&](FDependsNode* InReferencer){AddDependsNodesRecursive(InReferencer, OutArray, InAssets); });
 		}
 	}
 }
 
 void FAssetRegistry::SaveRegistryData(FArchive& Ar, TMap<FName, FAssetData*>& Data, TArray<FName>* InMaps /* = nullptr */)
 {
+	// Write mini asset registry header
+	FGuid LocalGuid = GRuntimeRegistryGuid;
+	LocalGuid.Serialize(Ar);
+	int32 VersionInt = (int32)ERuntimeRegistryVersion::Latest;
+	Ar << VersionInt;
+
 	// serialize number of objects
 	int32 AssetCount = Data.Num();
 	Ar << AssetCount;
@@ -1409,73 +1460,83 @@ void FAssetRegistry::SaveRegistryData(FArchive& Ar, TMap<FName, FAssetData*>& Da
 	{
 		Ar << *It.Value();
 		AssetIndexMap.Add(It.Value()->PackageName, AssetIndexMap.Num());
-
-		AddDependsNodesRecursive(FindDependsNode(It.Value()->PackageName), Dependencies, Data);
-	}
-
-	int32 DependsNodeCount = Dependencies.Num();
-	Ar << DependsNodeCount;
-
-	TMap<FDependsNode*, int32> DependencyIndexMap;
-	DependencyIndexMap.Reserve(DependsNodeCount);
-
-	for (auto DependentNode : Dependencies)
-	{
-		int32 AssetIndex = AssetIndexMap[DependentNode->PackageName];
-		Ar << AssetIndex;
-		DependencyIndexMap.Add(DependentNode, DependencyIndexMap.Num());
+		Dependencies.Add(FindDependsNode(It.Value()->PackageName));
 	}
 
 	for (auto DependentNode : Dependencies)
 	{
-		auto DependencyCount = 0;
+		auto HardDependencyCount = 0;
+		auto SoftDependencyCount = 0;
 		auto ReferencerCount = 0;
 
-		for (auto Dependency : DependentNode->Dependencies)
+		DependentNode->IterateOverDependencies([&](FDependsNode* InDependency, EAssetRegistryDependencyType::Type InDependencyType)
 		{
-			bool bIsMap = InMaps && InMaps->Contains(Dependency->PackageName);
-		
-			if (!bIsMap && DependencyIndexMap.Contains(Dependency))
+			bool bIsMap = InMaps && InMaps->Contains(InDependency->GetPackageName());
+
+			if (!bIsMap && AssetIndexMap.Contains(InDependency->GetPackageName()))
 			{
-				DependencyCount++;
+				if (InDependencyType == EAssetRegistryDependencyType::Hard)
+				{
+					HardDependencyCount++;
+				}
+				else
+				{
+					SoftDependencyCount++;
+				}
 			}
-		}
+		});
 
-		for (auto Referencer : DependentNode->Referencers)
+		DependentNode->IterateOverReferencers([&](FDependsNode* InReferencer)
 		{
-			if (DependencyIndexMap.Contains(Referencer))
+			if (AssetIndexMap.Contains(InReferencer->GetPackageName()))
 			{
 				ReferencerCount++;
 			}
-		}
+		});
 
-		Ar << DependencyCount;
+		Ar << HardDependencyCount;
+		Ar << SoftDependencyCount;
 		Ar << ReferencerCount;
 
-		for (auto Dependency : DependentNode->Dependencies)
+		DependentNode->IterateOverDependencies([&](FDependsNode* InDependency, EAssetRegistryDependencyType::Type InDependencyType)
 		{
-			bool bIsMap = InMaps && InMaps->Contains(Dependency->PackageName);
+			bool bIsMap = InMaps && InMaps->Contains(InDependency->GetPackageName());
 
-			if (!bIsMap && DependencyIndexMap.Contains(Dependency))
+			if (!bIsMap && AssetIndexMap.Contains(InDependency->GetPackageName()))
 			{
-				int32 Index = DependencyIndexMap[Dependency];
+				int32 Index = AssetIndexMap[InDependency->GetPackageName()];
 				Ar << Index;
 			}
-		}
+		},
+		EAssetRegistryDependencyType::Hard);
 
-		for (auto Referencer : DependentNode->Referencers)
+		DependentNode->IterateOverDependencies([&](FDependsNode* InDependency, EAssetRegistryDependencyType::Type InDependencyType)
 		{
-			if (DependencyIndexMap.Contains(Referencer))
+			bool bIsMap = InMaps && InMaps->Contains(InDependency->GetPackageName());
+
+			if (!bIsMap && AssetIndexMap.Contains(InDependency->GetPackageName()))
 			{
-				int32 Index = DependencyIndexMap[Referencer];
+				int32 Index = AssetIndexMap[InDependency->GetPackageName()];
 				Ar << Index;
 			}
-		}
+		},
+		EAssetRegistryDependencyType::Soft);
+
+		DependentNode->IterateOverReferencers([&](FDependsNode* InReferencer)
+		{
+			if (AssetIndexMap.Contains(InReferencer->GetPackageName()))
+			{
+				int32 Index = AssetIndexMap[InReferencer->GetPackageName()];
+				Ar << Index;
+			}
+		});
 	}
 }
 
 void FAssetRegistry::LoadRegistryData(FArchive& Ar, TMap<FName, FAssetData*>& Data, TArray<FDependsNode*>& OutDependencyData)
 {
+	auto Version = ReadRuntimeRegistryVersion(Ar);
+
 	check(Ar.IsLoading());
 	// serialize number of objects
 	int AssetCount = 0;
@@ -1498,15 +1559,27 @@ void FAssetRegistry::LoadRegistryData(FArchive& Ar, TMap<FName, FAssetData*>& Da
 
 	if (Ar.TotalSize() > Ar.Tell())
 	{
-		int DependsNodeCount = 0;
-		Ar << DependsNodeCount;
+		int DependsNodeCount = AssetCount;
+
+		if (Version == ERuntimeRegistryVersion::PreVersioning)
+		{
+			Ar << DependsNodeCount;
+		}
 
 		OutDependencyData.Reserve(DependsNodeCount);
 
 		for (int32 DependsNodeIndex = 0; DependsNodeIndex < DependsNodeCount; DependsNodeIndex++)
 		{
-			int32 AssetIndex;
-			Ar << AssetIndex;
+			int32 AssetIndex = 0;
+
+			if (Version == ERuntimeRegistryVersion::PreVersioning)
+			{
+				Ar << AssetIndex;
+			}
+			else
+			{
+				AssetIndex = OutDependencyData.Num();
+			}
 
 			auto PackageName = AssetIndexMap[AssetIndex];
 			auto NewNode = new FDependsNode(PackageName);
@@ -1515,22 +1588,34 @@ void FAssetRegistry::LoadRegistryData(FArchive& Ar, TMap<FName, FAssetData*>& Da
 
 		for (int32 DependsNodeIndex = 0; DependsNodeIndex < DependsNodeCount; DependsNodeIndex++)
 		{
-			int32 NumDependencies = 0;
-			Ar << NumDependencies;
+			int32 NumHardDependencies = 0;
+			Ar << NumHardDependencies;
+
+			int32 NumSoftDependencies = 0;
+			if (Version >= ERuntimeRegistryVersion::HardSoftDependencies)
+			{
+				Ar << NumSoftDependencies;
+			}
 
 			int32 NumReferencers = 0;
 			Ar << NumReferencers;
 
 			auto DependsNode = OutDependencyData[DependsNodeIndex];
-			DependsNode->Dependencies.Reserve(NumDependencies);
-			DependsNode->Referencers.Reserve(NumReferencers);
 
-			for (int32 i = 0; i < NumDependencies; ++i)
+			for (int32 i = 0; i < NumHardDependencies; ++i)
+			{
+				int32 DependencyIndex = 0;
+				Ar << DependencyIndex;
+
+				DependsNode->AddDependency(OutDependencyData[DependencyIndex], EAssetRegistryDependencyType::Hard);
+			}
+
+			for (int32 i = 0; i < NumSoftDependencies; ++i)
 			{
 				int32 DependencyIndex = 0;
 				Ar << DependencyIndex;
 
-				DependsNode->Dependencies.Add(OutDependencyData[DependencyIndex]);
+				DependsNode->AddDependency(OutDependencyData[DependencyIndex], EAssetRegistryDependencyType::Soft);
 			}
 
 			for (int32 i = 0; i < NumReferencers; ++i)
@@ -1538,7 +1623,7 @@ void FAssetRegistry::LoadRegistryData(FArchive& Ar, TMap<FName, FAssetData*>& Da
 				int32 Index = 0;
 				Ar << Index;
 
-				DependsNode->Referencers.Add(OutDependencyData[Index]);
+				DependsNode->AddReferencer(OutDependencyData[Index]);
 			}
 		}
 	}
@@ -1684,35 +1769,41 @@ void FAssetRegistry::DependencyDataGathered(const double TickStartTime, TArray<F
 
 		// We will populate the node dependencies below. Empty the set here in case this file was already read
 		// Also remove references to all existing dependencies, those will be also repopulated below
-		for (auto OldDependsIt = Node->Dependencies.CreateConstIterator(); OldDependsIt; ++OldDependsIt)
+		Node->IterateOverDependencies([Node](FDependsNode* InDependency, EAssetRegistryDependencyType::Type InDependencyType)
 		{
-			(*OldDependsIt)->Referencers.Remove(Node);
-		}
-		Node->Dependencies.Empty();
+			InDependency->RemoveReferencer(Node);
+		});
+
+		Node->ClearDependencies();
 
 		// Determine the new package dependencies
-		TSet<FName> PackageDependencies;
+		TMap<FName, EAssetRegistryDependencyType::Type> PackageDependencies;
 		for (int32 ImportIdx = 0; ImportIdx < Result.ImportMap.Num(); ++ImportIdx)
 		{
-			PackageDependencies.Add( Result.GetImportPackageName(ImportIdx) );
+			PackageDependencies.Add(Result.GetImportPackageName(ImportIdx), EAssetRegistryDependencyType::Hard);
 		}
+
 		for (int32 StringAssetRefIdx = 0; StringAssetRefIdx < Result.StringAssetReferencesMap.Num(); ++StringAssetRefIdx)
 		{
 			FString PackageName, ObjName;
 			Result.StringAssetReferencesMap[StringAssetRefIdx].Split(".", &PackageName, &ObjName, ESearchCase::CaseSensitive);
-			PackageDependencies.Add(*PackageName);
+
+			if (!PackageDependencies.Contains(*PackageName))
+			{
+				PackageDependencies.Add(*PackageName, EAssetRegistryDependencyType::Soft);
+			}
 		}
 
 
 		// Doubly-link all new dependencies for this package
-		for (auto NewDependsIt = PackageDependencies.CreateConstIterator(); NewDependsIt; ++NewDependsIt)
+		for (auto NewDependsIt : PackageDependencies)
 		{
-			FDependsNode* DependsNode = CreateOrFindDependsNode(*NewDependsIt);
+			FDependsNode* DependsNode = CreateOrFindDependsNode(NewDependsIt.Key);
 			
 			if (DependsNode != NULL)
 			{
-				Node->Dependencies.Add(DependsNode);
-				DependsNode->Referencers.Add(Node);
+				Node->AddDependency(DependsNode, NewDependsIt.Value);
+				DependsNode->AddReferencer(Node);
 			}
 		}
 
@@ -1772,7 +1863,7 @@ bool FAssetRegistry::RemoveDependsNode( FName PackageName )
 			// Remove the reference to this node from all dependencies
 			for ( FDependsNode* DependencyNode : DependencyNodes )
 			{
-				DependencyNode->Referencers.Remove( Node );
+				DependencyNode->RemoveReferencer( Node );
 			}
 
 			TArray<FDependsNode*> ReferencerNodes;
@@ -1781,7 +1872,7 @@ bool FAssetRegistry::RemoveDependsNode( FName PackageName )
 			// Remove the reference to this node from all referencers
 			for ( FDependsNode* ReferencerNode : ReferencerNodes )
 			{
-				ReferencerNode->Dependencies.Remove( Node );
+				ReferencerNode->RemoveDependency(Node);
 			}
 
 			// Remove the node and delete it
diff --git a/Engine/Source/Runtime/AssetRegistry/Private/AssetRegistry.h b/Engine/Source/Runtime/AssetRegistry/Private/AssetRegistry.h
--- a/Engine/Source/Runtime/AssetRegistry/Private/AssetRegistry.h
+++ b/Engine/Source/Runtime/AssetRegistry/Private/AssetRegistry.h
@@ -21,7 +21,7 @@ public:
 	virtual bool GetAssets(const FARFilter& Filter, TArray<FAssetData>& OutAssetData) const override;
 	virtual FAssetData GetAssetByObjectPath( const FName ObjectPath ) const override;
 	virtual bool GetAllAssets(TArray<FAssetData>& OutAssetData) const override;
-	virtual bool GetDependencies(FName PackageName, TArray<FName>& OutDependencies) const override;
+	virtual bool GetDependencies(FName PackageName, TArray<FName>& OutDependencies, EAssetRegistryDependencyType::Type InDependencyType = EAssetRegistryDependencyType::All) const override;
 	virtual bool GetReferencers(FName PackageName, TArray<FName>& OutReferencers) const override;
 	virtual bool GetAncestorClassNames(FName ClassName, TArray<FName>& OutAncestorClassNames) const override;
 	virtual void GetDerivedClassNames(const TArray<FName>& ClassNames, const TSet<FName>& ExcludedClassNames, TSet<FName>& OutDerivedClassNames) const override;
diff --git a/Engine/Source/Runtime/AssetRegistry/Private/DependsNode.cpp b/Engine/Source/Runtime/AssetRegistry/Private/DependsNode.cpp
--- a/Engine/Source/Runtime/AssetRegistry/Private/DependsNode.cpp
+++ b/Engine/Source/Runtime/AssetRegistry/Private/DependsNode.cpp
@@ -30,12 +30,22 @@ void FDependsNode::PrintReferencers() const
 	PrintReferencersRecursive(TEXT(""), VisitedNodes);
 }
 
-void FDependsNode::GetDependencies(TArray<FDependsNode*>& OutDependencies) const
+void FDependsNode::GetDependencies(TArray<FDependsNode*>& OutDependencies, EAssetRegistryDependencyType::Type InDependencyType) const
 {
-	for (auto DependsIt = Dependencies.CreateConstIterator(); DependsIt; ++DependsIt)
+	IterateOverDependencies([&](FDependsNode* InDependency, EAssetRegistryDependencyType::Type /*InDependencyType*/)
 	{
-		OutDependencies.Add(*DependsIt);
-	}
+		OutDependencies.Add(InDependency);
+	}, 
+	InDependencyType);
+}
+
+void FDependsNode::GetDependencies(TArray<FName>& OutDependencies, EAssetRegistryDependencyType::Type InDependencyType) const
+{
+	IterateOverDependencies([&](FDependsNode* InDependency, EAssetRegistryDependencyType::Type /*InDependencyType*/)
+	{
+		OutDependencies.Add(InDependency->GetPackageName());
+	},
+	InDependencyType);
 }
 
 void FDependsNode::GetReferencers(TArray<FDependsNode*>& OutReferencers) const
@@ -61,10 +71,12 @@ void FDependsNode::PrintDependenciesRecursive(const FString& Indent, TSet<const
 		UE_LOG(LogAssetRegistry, Log, TEXT("%s%s"), *Indent, *PackageName.ToString());
 		VisitedNodes.Add(this);
 
-		for (auto DependsIt = Dependencies.CreateConstIterator(); DependsIt; ++DependsIt)
+		IterateOverDependencies([&](FDependsNode* InDependency, EAssetRegistryDependencyType::Type /*InDependencyType*/)
 		{
-			(*DependsIt)->PrintDependenciesRecursive(Indent + TEXT("  "), VisitedNodes);
-		}
+			InDependency->PrintDependenciesRecursive(Indent + TEXT("  "), VisitedNodes);
+		},
+		EAssetRegistryDependencyType::All
+		);
 	}
 }
 
diff --git a/Engine/Source/Runtime/AssetRegistry/Private/DependsNode.h b/Engine/Source/Runtime/AssetRegistry/Private/DependsNode.h
--- a/Engine/Source/Runtime/AssetRegistry/Private/DependsNode.h
+++ b/Engine/Source/Runtime/AssetRegistry/Private/DependsNode.h
@@ -16,9 +16,90 @@ public:
 	/** Prints the referencers to this node to the log */
 	void PrintReferencers() const;
 	/** Gets the list of dependencies for this node */
-	void GetDependencies(TArray<FDependsNode*>& OutDependencies) const;
+	void GetDependencies(TArray<FDependsNode*>& OutDependencies, EAssetRegistryDependencyType::Type InDependencyType = EAssetRegistryDependencyType::All) const;
+	/** Gets the list of dependency names for this node */
+	void GetDependencies(TArray<FName>& OutDependencies, EAssetRegistryDependencyType::Type InDependencyType = EAssetRegistryDependencyType::All) const;
 	/** Gets the list of referencers to this node */
 	void GetReferencers(TArray<FDependsNode*>& OutReferencers) const;
+	/** Gets the name of the package that this node represents */
+	FName GetPackageName() const { return PackageName; }
+	/** Sets the name of the package that this node represents */
+	void SetPackageName(FName InName) { PackageName = InName; }
+	/** Add a dependency to this node */
+	void AddDependency(FDependsNode* InDependency, EAssetRegistryDependencyType::Type InDependencyType = EAssetRegistryDependencyType::Hard)
+	{ 
+		if (InDependencyType == EAssetRegistryDependencyType::Hard)
+		{
+			HardDependencies.Add(InDependency);
+		}
+		else
+		{
+			SoftDependencies.Add(InDependency);
+		}
+	}
+	/** Add a referencer to this node */
+	void AddReferencer(FDependsNode* InReferencer) { Referencers.Add(InReferencer); }
+	/** Remove a dependency from this node */
+	void RemoveDependency(FDependsNode* InDependency) { HardDependencies.Remove(InDependency); SoftDependencies.Remove(InDependency); }
+	/** Remove a referencer from this node */
+	void RemoveReferencer(FDependsNode* InReferencer) { Referencers.Remove(InReferencer); }
+	/** Clear all dependency records from this node */
+	void ClearDependencies() { HardDependencies.Empty(); SoftDependencies.Empty(); }
+
+	/** Iterate over all the dependencies of this node, filtered by the supplied type parameter, and call the supplied lambda
+		parameter on the record */
+	template <class T>
+	void IterateOverDependencies(T InCallback, EAssetRegistryDependencyType::Type InDependencyType = EAssetRegistryDependencyType::All)
+	{
+		if (InDependencyType & EAssetRegistryDependencyType::Hard)
+		{
+			for (auto Dependency : HardDependencies)
+			{
+				InCallback(Dependency, EAssetRegistryDependencyType::Hard);
+			}
+		}
+
+		if (InDependencyType & EAssetRegistryDependencyType::Soft)
+		{
+			for (auto Dependency : SoftDependencies)
+			{
+				InCallback(Dependency, EAssetRegistryDependencyType::Soft);
+			}
+		}
+	}
+
+	/** Iterate over all the dependencies of this node, filtered by the supplied type parameter, and call the supplied lambda
+	parameter on the record */
+	template <class T>
+	void IterateOverDependencies(T InCallback, EAssetRegistryDependencyType::Type InDependencyType = EAssetRegistryDependencyType::All) const
+	{
+		if (InDependencyType & EAssetRegistryDependencyType::Hard)
+		{
+			for (auto Dependency : HardDependencies)
+			{
+				InCallback(Dependency, EAssetRegistryDependencyType::Hard);
+			}
+		}
+
+		if (InDependencyType & EAssetRegistryDependencyType::Soft)
+		{
+			for (auto Dependency : SoftDependencies)
+			{
+				InCallback(Dependency, EAssetRegistryDependencyType::Soft);
+			}
+		}
+	}
+
+	/** Iterate over all the referencers of this node and call the supplied lambda
+		parameter on the record */
+	template <class T>
+	void IterateOverReferencers(T InCallback)
+	{
+		for (auto Referencer : Referencers)
+		{
+			InCallback(Referencer);
+		}
+	}
 
 private:
 	/** Recursively prints dependencies of the node starting with the specified indent. VisitedNodes should be an empty set at first which is populated recursively. */
@@ -26,11 +107,12 @@ private:
 	/** Recursively prints referencers to the node starting with the specified indent. VisitedNodes should be an empty set at first which is populated recursively. */
 	void PrintReferencersRecursive(const FString& Indent, TSet<const FDependsNode*>& VisitedNodes) const;
 
-public:
 	/** The name of the package this node represents */
 	FName PackageName;
-	/** The list of dependencies for this node */
-	TSet<FDependsNode*> Dependencies;
+	/** The list of hard dependencies for this node */
+	TSet<FDependsNode*> HardDependencies;
+	/** The list of soft dependencies for this node */
+	TSet<FDependsNode*> SoftDependencies;
 	/** The list of referencers to this node */
 	TSet<FDependsNode*> Referencers;
 };
\ No newline at end of file
diff --git a/Engine/Source/Runtime/AssetRegistry/Public/AssetRegistryModule.h b/Engine/Source/Runtime/AssetRegistry/Public/AssetRegistryModule.h
--- a/Engine/Source/Runtime/AssetRegistry/Public/AssetRegistryModule.h
+++ b/Engine/Source/Runtime/AssetRegistry/Public/AssetRegistryModule.h
@@ -67,9 +67,9 @@ public:
 	}
 
 	/** Access the dependent package names for a given source package */
-	void GetDependencies(FName InPackageName, TArray<FName>& OutDependencies) override
+	void GetDependencies(FName InPackageName, TArray<FName>& OutDependencies, EAssetRegistryDependencyType::Type InDependencyType) override
 	{
-		Get().GetDependencies(InPackageName, OutDependencies);
+		Get().GetDependencies(InPackageName, OutDependencies, InDependencyType);
 	}
 
 private:
diff --git a/Engine/Source/Runtime/AssetRegistry/Public/IAssetRegistry.h b/Engine/Source/Runtime/AssetRegistry/Public/IAssetRegistry.h
--- a/Engine/Source/Runtime/AssetRegistry/Public/IAssetRegistry.h
+++ b/Engine/Source/Runtime/AssetRegistry/Public/IAssetRegistry.h
@@ -4,6 +4,7 @@
 
 #include "AssetData.h"
 #include "ARFilter.h"
+#include "AssetRegistryInterface.h"
 
 
 namespace EAssetAvailability
@@ -100,8 +101,9 @@ public:
 	 *
 	 * @param PackageName		the name of the package for which to gather dependencies
 	 * @param OutDependencies	a list of paths to objects that are referenced by the package whose path is PackageName
+	 * @param InDependencyType	which kinds of dependency to include in the output list
 	 */
-	virtual bool GetDependencies(FName PackageName, TArray<FName>& OutDependencies) const = 0;
+	virtual bool GetDependencies(FName PackageName, TArray<FName>& OutDependencies, EAssetRegistryDependencyType::Type InDependencyType = EAssetRegistryDependencyType::All) const = 0;
 
 	/**
 	 * Gets a list of paths to objects that reference the supplied package. (On disk references ONLY)
diff --git a/Engine/Source/Runtime/CoreUObject/Private/Serialization/AsyncLoading.cpp b/Engine/Source/Runtime/CoreUObject/Private/Serialization/AsyncLoading.cpp
--- a/Engine/Source/Runtime/CoreUObject/Private/Serialization/AsyncLoading.cpp
+++ b/Engine/Source/Runtime/CoreUObject/Private/Serialization/AsyncLoading.cpp
@@ -449,12 +449,7 @@ void FAsyncLoadingThread::ProcessAsyncPackageRequest(FAsyncPackageDesc* InReques
 			if (AssetRegistry)
 			{
 				TArray<FName> Dependencies;
-				AssetRegistry->GetDependencies(Package->GetPackageName(), Dependencies);
-
-				int32 PIEInstanceID = INDEX_NONE;
-#if WITH_EDITOR
-				PIEInstanceID = InRequest->PIEInstanceID;
-#endif
+				AssetRegistry->GetDependencies(Package->GetPackageName(), Dependencies, EAssetRegistryDependencyType::Hard);
 
 				if (InRootPackage == nullptr)
 				{
@@ -468,7 +463,7 @@ void FAsyncLoadingThread::ProcessAsyncPackageRequest(FAsyncPackageDesc* InReques
 						QueuedPackagesCounter.Increment();
 						const int32 RequestID = GPackageRequestID.Increment();
 						FAsyncLoadingThread::Get().AddPendingRequest(RequestID);
-						FAsyncPackageDesc DependencyPackageRequest(RequestID, DependencyName, NAME_None, FGuid(), FLoadPackageAsyncDelegate(), InRequest->PackageFlags, PIEInstanceID, InRequest->Priority);
+						FAsyncPackageDesc DependencyPackageRequest(RequestID, DependencyName, NAME_None, FGuid(), FLoadPackageAsyncDelegate(), InRequest->PackageFlags, INDEX_NONE, InRequest->Priority);
 						InDependencyTracker.Add(InRequest->Name);
 						ProcessAsyncPackageRequest(&DependencyPackageRequest, InRootPackage, InDependencyTracker);
 					}
diff --git a/Engine/Source/Runtime/CoreUObject/Private/UObject/UObjectGlobals.cpp b/Engine/Source/Runtime/CoreUObject/Private/UObject/UObjectGlobals.cpp
--- a/Engine/Source/Runtime/CoreUObject/Private/UObject/UObjectGlobals.cpp
+++ b/Engine/Source/Runtime/CoreUObject/Private/UObject/UObjectGlobals.cpp
@@ -975,7 +975,7 @@ UPackage* LoadPackageInternal(UPackage* InOuter, const TCHAR* InLongPackageName,
 				TArray<FName> PackageDependencies;
 				FName PackageName(InLongPackageName);
 
-				AssetRegistry->GetDependencies(PackageName, PackageDependencies);
+				AssetRegistry->GetDependencies(PackageName, PackageDependencies, EAssetRegistryDependencyType::Hard);
 
 				Result->PackageFlags |= PKG_ProcessingDependencies;
 				for (auto Dependency : PackageDependencies)
diff --git a/Engine/Source/Runtime/CoreUObject/Public/Misc/AssetRegistryInterface.h b/Engine/Source/Runtime/CoreUObject/Public/Misc/AssetRegistryInterface.h
--- a/Engine/Source/Runtime/CoreUObject/Public/Misc/AssetRegistryInterface.h
+++ b/Engine/Source/Runtime/CoreUObject/Public/Misc/AssetRegistryInterface.h
@@ -5,6 +5,20 @@
 #include "ModuleManager.h"
 #include "CoreUObject.h"
 
+namespace EAssetRegistryDependencyType
+{
+	enum Type
+	{
+		// Dependencies which don't need to be loaded for the object to be used (i.e. string asset references)
+		Soft = 1,
+
+		// Dependencies which are required for correct usage of the source asset, and must be loaded at the same time
+		Hard = 2
+	};
+
+	static const Type All = (Type)(Soft | Hard);
+};
+
 /**
 * HotReload module interface
 */
@@ -28,6 +42,6 @@ public:
 	/**
 	* Lookup dependencies for the given package name and fill OutDependencies with direct dependencies
 	*/
-	virtual void GetDependencies(FName InPackageName, TArray<FName>& OutDependencies) = 0;
+	virtual void GetDependencies(FName InPackageName, TArray<FName>& OutDependencies, EAssetRegistryDependencyType::Type InDependencyType = EAssetRegistryDependencyType::All) = 0;
 };
 
