Changed the asset registry discovery to avoid some redundant copies and conversions

It would previously find an FAssetData in the cached asset registry, and then pass it to FBackgroundAssetData which would copy the tags and values map, and perform at least 6 FName -> FString conversions, before ultimately converting all of these FStrings back to FNames when FBackgroundAssetData::ToAssetData was called.

This change makes FBackgroundAssetData implement an interface (IGatheredAssetData), and adds FAssetDataWrapper to more efficiently wrap up an FAssetData that was read from the cached asset registry.

This saved 2-3 seconds when loading the largest project I have available (previously ~18 seconds, now ~15 seconds).

diff --git a/Engine/Source/Runtime/AssetRegistry/Private/AssetDataGatherer.cpp b/Engine/Source/Runtime/AssetRegistry/Private/AssetDataGatherer.cpp
--- a/Engine/Source/Runtime/AssetRegistry/Private/AssetDataGatherer.cpp
+++ b/Engine/Source/Runtime/AssetRegistry/Private/AssetDataGatherer.cpp
@@ -83,7 +83,7 @@ uint32 FAssetDataGatherer::Run()
 	}
 
 	TArray<FString> LocalFilesToSearch;
-	TArray<FBackgroundAssetData*> LocalAssetResults;
+	TArray<IGatheredAssetData*> LocalAssetResults;
 	TArray<FPackageDependencyData> LocalDependencyResults;
 
 	while ( StopTaskCounter.GetValue() == 0 )
@@ -169,7 +169,7 @@ uint32 FAssetDataGatherer::Run()
 					{
 						for ( auto CacheIt = DiskCachedAssetData->AssetDataList.CreateConstIterator(); CacheIt; ++CacheIt )
 						{
-							LocalAssetResults.Add(new FBackgroundAssetData(*CacheIt));
+							LocalAssetResults.Add(new FAssetDataWrapper(*CacheIt));
 						}
 
 						LocalDependencyResults.Add(DiskCachedAssetData->DependencyData);
@@ -267,7 +267,7 @@ void FAssetDataGatherer::EnsureCompletion()
     Thread = NULL;
 }
 
-bool FAssetDataGatherer::GetAndTrimSearchResults(TArray<FBackgroundAssetData*>& OutAssetResults, TArray<FString>& OutPathResults, TArray<FPackageDependencyData>& OutDependencyResults, TArray<double>& OutSearchTimes, int32& OutNumFilesToSearch, int32& OutNumPathsToSearch)
+bool FAssetDataGatherer::GetAndTrimSearchResults(TArray<IGatheredAssetData*>& OutAssetResults, TArray<FString>& OutPathResults, TArray<FPackageDependencyData>& OutDependencyResults, TArray<double>& OutSearchTimes, int32& OutNumFilesToSearch, int32& OutNumPathsToSearch)
 {
 	FScopeLock CritSectionLock(&WorkerThreadCriticalSection);
 
diff --git a/Engine/Source/Runtime/AssetRegistry/Private/AssetDataGatherer.h b/Engine/Source/Runtime/AssetRegistry/Private/AssetDataGatherer.h
--- a/Engine/Source/Runtime/AssetRegistry/Private/AssetDataGatherer.h
+++ b/Engine/Source/Runtime/AssetRegistry/Private/AssetDataGatherer.h
@@ -22,7 +22,7 @@ public:
 	void EnsureCompletion();
 
 	/** Gets search results from the data gatherer */
-	bool GetAndTrimSearchResults(TArray<FBackgroundAssetData*>& OutAssetResults, TArray<FString>& OutPathResults, TArray<FPackageDependencyData>& OutDependencyResults, TArray<double>& OutSearchTimes, int32& OutNumFilesToSearch, int32& OutNumPathsToSearch);
+	bool GetAndTrimSearchResults(TArray<IGatheredAssetData*>& OutAssetResults, TArray<FString>& OutPathResults, TArray<FPackageDependencyData>& OutDependencyResults, TArray<double>& OutSearchTimes, int32& OutNumFilesToSearch, int32& OutNumPathsToSearch);
 
 	/** Adds a root path to the search queue. Only works when searching asynchronously */
 	void AddPathToSearch(const FString& Path);
@@ -80,7 +80,7 @@ private:
 	TArray<FString> PathsToSearch;
 
 	/** The asset data gathered from the searched files */
-	TArray<FBackgroundAssetData*> AssetResults;
+	TArray<IGatheredAssetData*> AssetResults;
 
 	/** Dependency data for scanned packages */
 	TArray<FPackageDependencyData> DependencyResults;
diff --git a/Engine/Source/Runtime/AssetRegistry/Private/AssetRegistry.cpp b/Engine/Source/Runtime/AssetRegistry/Private/AssetRegistry.cpp
--- a/Engine/Source/Runtime/AssetRegistry/Private/AssetRegistry.cpp
+++ b/Engine/Source/Runtime/AssetRegistry/Private/AssetRegistry.cpp
@@ -1148,7 +1148,7 @@ void FAssetRegistry::PrioritizeSearchPath(const FString& PathToPrioritize)
 		int32 LowestNonPriorityFileIdx = 0;
 		for (int32 ResultIdx = 0; ResultIdx < BackgroundAssetResults.Num(); ++ResultIdx)
 		{
-			if (BackgroundAssetResults[ResultIdx]->PackagePath.StartsWith(PathToPrioritize))
+			if (BackgroundAssetResults[ResultIdx]->IsWithinSearchPath(PathToPrioritize))
 			{
 				BackgroundAssetResults.Swap(ResultIdx, LowestNonPriorityFileIdx);
 				LowestNonPriorityFileIdx++;
@@ -1650,7 +1650,7 @@ void FAssetRegistry::ScanPathsSynchronous_Internal(const TArray<FString>& InPath
 		FAssetDataGatherer AssetSearch(PathsToScan, /*bSynchronous=*/true, bUseCache);
 
 		// Get the search results
-		TArray<FBackgroundAssetData*> AssetResults;
+		TArray<IGatheredAssetData*> AssetResults;
 		TArray<FString> PathResults;
 		TArray<FPackageDependencyData> DependencyResults;
 		TArray<double> SearchTimes;
@@ -1680,7 +1680,7 @@ void FAssetRegistry::ScanPathsSynchronous_Internal(const TArray<FString>& InPath
 	}
 }
 
-void FAssetRegistry::AssetSearchDataGathered(const double TickStartTime, TArray<FBackgroundAssetData*>& AssetResults)
+void FAssetRegistry::AssetSearchDataGathered(const double TickStartTime, TArray<IGatheredAssetData*>& AssetResults)
 {
 	const bool bFlushFullBuffer = TickStartTime < 0;
 	TSet<FName> ModifiedPaths;
@@ -1689,7 +1689,7 @@ void FAssetRegistry::AssetSearchDataGathered(const double TickStartTime, TArray<
 	int32 AssetIndex = 0;
 	for (AssetIndex = 0; AssetIndex < AssetResults.Num(); ++AssetIndex)
 	{
-		FBackgroundAssetData*& BackgroundResult = AssetResults[AssetIndex];
+		IGatheredAssetData*& BackgroundResult = AssetResults[AssetIndex];
 		FAssetData Result = BackgroundResult->ToAssetData();
 
 		// Try to update any asset data that may already exist
@@ -1712,7 +1712,7 @@ void FAssetRegistry::AssetSearchDataGathered(const double TickStartTime, TArray<
 		}
 
 		// Populate the path tree
-		AddAssetPath(BackgroundResult->PackagePath);
+		AddAssetPath(Result.PackagePath.ToString());
 
 		// Delete the result that was originally created by an FPackageReader
 		delete BackgroundResult;
diff --git a/Engine/Source/Runtime/AssetRegistry/Private/AssetRegistry.h b/Engine/Source/Runtime/AssetRegistry/Private/AssetRegistry.h
--- a/Engine/Source/Runtime/AssetRegistry/Private/AssetRegistry.h
+++ b/Engine/Source/Runtime/AssetRegistry/Private/AssetRegistry.h
@@ -85,7 +85,7 @@ private:
 	void ScanPathsSynchronous_Internal(const TArray<FString>& InPaths, bool bForceRescan, bool bUseCache);
 
 	/** Called every tick to when data is retrieved by the background asset search. If TickStartTime is < 0, the entire list of gathered assets will be cached. Also used in sychronous searches */
-	void AssetSearchDataGathered(const double TickStartTime, TArray<FBackgroundAssetData*>& AssetResults);
+	void AssetSearchDataGathered(const double TickStartTime, TArray<IGatheredAssetData*>& AssetResults);
 
 	/** Called every tick to when data is retrieved by the background path search. If TickStartTime is < 0, the entire list of gathered assets will be cached. Also used in sychronous searches */
 	void PathDataGathered(const double TickStartTime, TArray<FString>& PathResults);
@@ -201,7 +201,7 @@ private:
 	TSharedPtr< class FAssetDataGatherer > BackgroundAssetSearch;
 
 	/** A list of results that were gathered from the background thread that are waiting to get processed by the main thread */
-	TArray<class FBackgroundAssetData*> BackgroundAssetResults;
+	TArray<class IGatheredAssetData*> BackgroundAssetResults;
 	TArray<FString> BackgroundPathResults;
 	TArray<class FPackageDependencyData> BackgroundDependencyResults;
 
diff --git a/Engine/Source/Runtime/AssetRegistry/Private/BackgroundAssetData.cpp b/Engine/Source/Runtime/AssetRegistry/Private/BackgroundAssetData.cpp
--- a/Engine/Source/Runtime/AssetRegistry/Private/BackgroundAssetData.cpp
+++ b/Engine/Source/Runtime/AssetRegistry/Private/BackgroundAssetData.cpp
@@ -2,51 +2,31 @@
 
 #include "AssetRegistryPCH.h"
 
-FBackgroundAssetData::FBackgroundAssetData(const FString& InPackageName, const FString& InPackagePath, const FString& InGroupNames, const FString& InAssetName, const FString& InAssetClass, const TMap<FString, FString>& InTags, const TArray<int32>& InChunkIDs)
+FBackgroundAssetData::FBackgroundAssetData(FString InPackageName, FString InPackagePath, FString InGroupNames, FString InAssetName, FString InAssetClass, TMap<FString, FString> InTags, TArray<int32> InChunkIDs)
+	: PackageName(MoveTemp(InPackageName))
+	, PackagePath(MoveTemp(InPackagePath))
+	, GroupNames(MoveTemp(InGroupNames))
+	, AssetName(MoveTemp(InAssetName))
+	, AssetClass(MoveTemp(InAssetClass))
+	, TagsAndValues(MakeSharedMapView(MoveTemp(InTags)))
+	, ChunkIDs(MoveTemp(InChunkIDs))
 {
-	PackageName = InPackageName;
-	PackagePath = InPackagePath;
-	GroupNames = InGroupNames;
-	AssetName = InAssetName;
-	AssetClass = InAssetClass;
-	TagsAndValues = InTags;
-
 	ObjectPath = PackageName + TEXT(".");
 
-	if ( GroupNames.Len() )
+	if (GroupNames.Len())
 	{
 		ObjectPath += GroupNames + TEXT(".");
 	}
 
 	ObjectPath += AssetName;
-
-	ChunkIDs = InChunkIDs;
-}
-
-FBackgroundAssetData::FBackgroundAssetData(const FAssetData& InAssetData)
-{
-	PackageName = InAssetData.PackageName.ToString();
-	PackagePath = InAssetData.PackagePath.ToString();
-	GroupNames = InAssetData.GroupNames.ToString();
-	AssetName = InAssetData.AssetName.ToString();
-	AssetClass = InAssetData.AssetClass.ToString();
-
-	for ( auto TagIt = InAssetData.TagsAndValues.CreateConstIterator(); TagIt; ++TagIt )
-	{
-		TagsAndValues.Add(TagIt.Key().ToString(), TagIt.Value());
-	}
-
-	ObjectPath = InAssetData.ObjectPath.ToString();
-
-	ChunkIDs = InAssetData.ChunkIDs;
 }
 
 FAssetData FBackgroundAssetData::ToAssetData() const
 {
 	TMap<FName, FString> CopiedTagsAndValues;
-	for (TMap<FString,FString>::TConstIterator TagIt(TagsAndValues); TagIt; ++TagIt)
+	for (const auto& TagsAndValuesEntry : TagsAndValues)
 	{
-		CopiedTagsAndValues.Add(FName(*TagIt.Key()), TagIt.Value());
+		CopiedTagsAndValues.Add(FName(*TagsAndValuesEntry.Key), TagsAndValuesEntry.Value);
 	}
 
 	return FAssetData(
@@ -55,7 +35,28 @@ FAssetData FBackgroundAssetData::ToAssetData() const
 		FName(*GroupNames),
 		FName(*AssetName),
 		FName(*AssetClass),
-		CopiedTagsAndValues,
+		MoveTemp(CopiedTagsAndValues),
 		ChunkIDs
 		);
 }
+
+bool FBackgroundAssetData::IsWithinSearchPath(const FString& InSearchPath) const
+{
+	return PackagePath.StartsWith(InSearchPath);
+}
+
+FAssetDataWrapper::FAssetDataWrapper(FAssetData InAssetData)
+	: WrappedAssetData(MoveTemp(InAssetData))
+	, PackagePathStr(WrappedAssetData.PackagePath.ToString())
+{
+}
+
+FAssetData FAssetDataWrapper::ToAssetData() const
+{
+	return WrappedAssetData;
+}
+
+bool FAssetDataWrapper::IsWithinSearchPath(const FString& InSearchPath) const
+{
+	return PackagePathStr.StartsWith(InSearchPath);
+}
diff --git a/Engine/Source/Runtime/AssetRegistry/Private/BackgroundAssetData.h b/Engine/Source/Runtime/AssetRegistry/Private/BackgroundAssetData.h
--- a/Engine/Source/Runtime/AssetRegistry/Private/BackgroundAssetData.h
+++ b/Engine/Source/Runtime/AssetRegistry/Private/BackgroundAssetData.h
@@ -2,17 +2,33 @@
 
 #pragma once
 
+#include "SharedMapView.h"
+
+/** Interface to the gathered asset data */
+class IGatheredAssetData
+{
+public:
+	virtual ~IGatheredAssetData() {}
+
+	/** Creates an FAssetData object for this object */
+	virtual FAssetData ToAssetData() const = 0;
+
+	/** Test to see if this asset data is within the given search path */
+	virtual bool IsWithinSearchPath(const FString& InSearchPath) const = 0;
+};
+
 /** A class to hold important information about an assets found by the Asset Registry. Intended for use by the background gathering thread to avoid creation of FNames */
-class FBackgroundAssetData
+class FBackgroundAssetData : public IGatheredAssetData
 {
 public:
 	/** Constructor */
-	FBackgroundAssetData(const FString& InPackageName, const FString& InPackagePath, const FString& InGroupNames, const FString& InAssetName, const FString& InAssetClass, const TMap<FString, FString>& InTags, const TArray<int32>& InChunkIDs);
-	FBackgroundAssetData(const FAssetData& InAssetData);
+	FBackgroundAssetData(FString InPackageName, FString InPackagePath, FString InGroupNames, FString InAssetName, FString InAssetClass, TMap<FString, FString> InTags, TArray<int32> InChunkIDs);
 
-	/** Creates an AssetData object based on this object */
-	FAssetData ToAssetData() const;
+	/** IGatheredAssetData interface */
+	virtual FAssetData ToAssetData() const override;
+	virtual bool IsWithinSearchPath(const FString& InSearchPath) const override;
 	
+private:
 	/** The object path for the asset in the form 'Package.GroupNames.AssetName' */
 	FString ObjectPath;
 	/** The name of the package in which the asset is found */
@@ -26,7 +42,25 @@ public:
 	/** The name of the asset's class */
 	FString AssetClass;
 	/** The map of values for properties that were marked AssetRegistrySearchable */
-	TMap<FString, FString> TagsAndValues;
+	TSharedMapView<FString, FString> TagsAndValues;
 	/** The IDs of the chunks this asset is located in for streaming install.  Empty if not assigned to a chunk */
 	TArray<int32> ChunkIDs;
 };
+
+/** A class to wrap up an existing FAssetData to avoid redundant FName -> FString -> FName conversion */
+class FAssetDataWrapper : public IGatheredAssetData
+{
+public:
+	/** Constructor */
+	FAssetDataWrapper(FAssetData InAssetData);
+
+	/** IGatheredAssetData interface */
+	virtual FAssetData ToAssetData() const override;
+	virtual bool IsWithinSearchPath(const FString& InSearchPath) const override;
+
+private:
+	/** The asset data we're wrapping up */
+	FAssetData WrappedAssetData;
+	/** The value of WrappedAssetData.PackagePath as an FString for use by IsWithinSearchPath */
+	FString PackagePathStr;
+};
diff --git a/Engine/Source/Runtime/AssetRegistry/Private/PackageReader.cpp b/Engine/Source/Runtime/AssetRegistry/Private/PackageReader.cpp
--- a/Engine/Source/Runtime/AssetRegistry/Private/PackageReader.cpp
+++ b/Engine/Source/Runtime/AssetRegistry/Private/PackageReader.cpp
@@ -140,9 +140,8 @@ bool FPackageReader::ReadAssetRegistryData (TArray<FBackgroundAssetData*>& Asset
 		const bool bNoMapAsset = (ObjectCount == 0);
 		if (bLegacyPackage || bNoMapAsset)
 		{
-			const FString AssetName = FPackageName::GetLongPackageAssetName(PackageName);
-			TMap<FString, FString> TagsAndValues;
-			AssetDataList.Add(new FBackgroundAssetData(PackageName, PackagePath, TEXT(""), AssetName, TEXT("World"), TagsAndValues, PackageFileSummary.ChunkIDs));
+			FString AssetName = FPackageName::GetLongPackageAssetName(PackageName);
+			AssetDataList.Add(new FBackgroundAssetData(PackageName, PackagePath, FString(), MoveTemp(AssetName), TEXT("World"), TMap<FString, FString>(), PackageFileSummary.ChunkIDs));
 		}
 	}
 
@@ -191,7 +190,7 @@ bool FPackageReader::ReadAssetRegistryData (TArray<FBackgroundAssetData*>& Asset
 		}
 
 		// Create a new FBackgroundAssetData for this asset and update it with the gathered data
-		AssetDataList.Add(new FBackgroundAssetData(PackageName, PackagePath, GroupNames, AssetName, ObjectClassName, TagsAndValues, PackageFileSummary.ChunkIDs));
+		AssetDataList.Add(new FBackgroundAssetData(PackageName, PackagePath, MoveTemp(GroupNames), MoveTemp(AssetName), MoveTemp(ObjectClassName), MoveTemp(TagsAndValues), PackageFileSummary.ChunkIDs));
 	}
 
 	return true;
@@ -246,8 +245,7 @@ bool FPackageReader::ReadAssetDataFromThumbnailCache(TArray<FBackgroundAssetData
 		}
 
 		// Create a new FBackgroundAssetData for this asset and update it with the gathered data
-		TMap<FString, FString> TagsAndValues;
-		AssetDataList.Add(new FBackgroundAssetData(PackageName, PackagePath, GroupNames, AssetName, AssetClassName, TagsAndValues, PackageFileSummary.ChunkIDs));
+		AssetDataList.Add(new FBackgroundAssetData(PackageName, PackagePath, MoveTemp(GroupNames), MoveTemp(AssetName), MoveTemp(AssetClassName), TMap<FString, FString>(), PackageFileSummary.ChunkIDs));
 	}
 
 	return true;
diff --git a/Engine/Source/Runtime/AssetRegistry/Public/AssetData.h b/Engine/Source/Runtime/AssetRegistry/Public/AssetData.h
--- a/Engine/Source/Runtime/AssetRegistry/Public/AssetData.h
+++ b/Engine/Source/Runtime/AssetRegistry/Public/AssetData.h
@@ -33,15 +33,15 @@ public:
 	{}
 
 	/** Constructor */
-	FAssetData(FName InPackageName, FName InPackagePath, FName InGroupNames, FName InAssetName, FName InAssetClass, TMap<FName, FString> InTags, const TArray<int32>& InChunkIDs)
+	FAssetData(FName InPackageName, FName InPackagePath, FName InGroupNames, FName InAssetName, FName InAssetClass, TMap<FName, FString> InTags, TArray<int32> InChunkIDs)
+		: PackageName(InPackageName)
+		, PackagePath(InPackagePath)
+		, GroupNames(InGroupNames)
+		, AssetName(InAssetName)
+		, AssetClass(InAssetClass)
+		, TagsAndValues(MakeSharedMapView(MoveTemp(InTags)))
+		, ChunkIDs(MoveTemp(InChunkIDs))
 	{
-		PackageName = InPackageName;
-		PackagePath = InPackagePath;
-		GroupNames = InGroupNames;
-		AssetName = InAssetName;
-		AssetClass = InAssetClass;
-		TagsAndValues = MakeSharedMapView(MoveTemp(InTags));
-
 		FString ObjectPathStr = PackageName.ToString() + TEXT(".");
 
 		if ( GroupNames != NAME_None )
@@ -52,7 +52,6 @@ public:
 		ObjectPathStr += AssetName.ToString();
 
 		ObjectPath = FName(*ObjectPathStr);
-		ChunkIDs = InChunkIDs;
 	}
 
 	/** Constructor taking a UObject */
