Fixed asset registry crashing if no harddependencies are found.

diff --git a/Engine/Source/Runtime/AssetRegistry/Private/DependsNode.h b/Engine/Source/Runtime/AssetRegistry/Private/DependsNode.h
--- a/Engine/Source/Runtime/AssetRegistry/Private/DependsNode.h
+++ b/Engine/Source/Runtime/AssetRegistry/Private/DependsNode.h
@@ -51,7 +51,7 @@ public:
 	template <class T>
 	void IterateOverDependencies(T InCallback, EAssetRegistryDependencyType::Type InDependencyType = EAssetRegistryDependencyType::All)
 	{
-		if (InDependencyType & EAssetRegistryDependencyType::Hard)
+		if (InDependencyType & EAssetRegistryDependencyType::Hard && HardDependencies)
 		{
 			for (auto Dependency : HardDependencies)
 			{
@@ -59,7 +59,7 @@ public:
 			}
 		}
 
-		if (InDependencyType & EAssetRegistryDependencyType::Soft)
+		if (InDependencyType & EAssetRegistryDependencyType::Soft && SoftDependencies)
 		{
 			for (auto Dependency : SoftDependencies)
 			{
