Fixed compilation error.

diff --git a/Engine/Source/Runtime/AssetRegistry/Private/DependsNode.h b/Engine/Source/Runtime/AssetRegistry/Private/DependsNode.h
--- a/Engine/Source/Runtime/AssetRegistry/Private/DependsNode.h
+++ b/Engine/Source/Runtime/AssetRegistry/Private/DependsNode.h
@@ -51,19 +51,25 @@ public:
 	template <class T>
 	void IterateOverDependencies(T InCallback, EAssetRegistryDependencyType::Type InDependencyType = EAssetRegistryDependencyType::All)
 	{
-		if (InDependencyType & EAssetRegistryDependencyType::Hard && HardDependencies)
+		if (InDependencyType & EAssetRegistryDependencyType::Hard)
 		{
 			for (auto Dependency : HardDependencies)
 			{
-				InCallback(Dependency, EAssetRegistryDependencyType::Hard);
+				if (Dependency)
+				{
+					InCallback(Dependency, EAssetRegistryDependencyType::Hard);
+				}
 			}
 		}
 
-		if (InDependencyType & EAssetRegistryDependencyType::Soft && SoftDependencies)
+		if (InDependencyType & EAssetRegistryDependencyType::Soft)
 		{
 			for (auto Dependency : SoftDependencies)
 			{
-				InCallback(Dependency, EAssetRegistryDependencyType::Soft);
+				if (Dependency)
+				{
+					InCallback(Dependency, EAssetRegistryDependencyType::Soft);
+				}
 			}
 		}
 	}
