Fixed crash in asset registry when dependency info is missing.

diff --git a/Engine/Source/Runtime/AssetRegistry/Private/AssetRegistry.cpp b/Engine/Source/Runtime/AssetRegistry/Private/AssetRegistry.cpp
--- a/Engine/Source/Runtime/AssetRegistry/Private/AssetRegistry.cpp
+++ b/Engine/Source/Runtime/AssetRegistry/Private/AssetRegistry.cpp
@@ -1460,7 +1460,11 @@ void FAssetRegistry::SaveRegistryData(FArchive& Ar, TMap<FName, FAssetData*>& Da
 	{
 		Ar << *It.Value();
 		AssetIndexMap.Add(It.Value()->PackageName, AssetIndexMap.Num());
-		Dependencies.Add(FindDependsNode(It.Value()->PackageName));
+		FDependsNode* DependencyNode = FindDependsNode(It.Value()->PackageName);
+		if (DependencyNode)
+		{
+			Dependencies.Add(DependencyNode);
+		}
 	}
 
 	for (auto DependentNode : Dependencies)
