Added support for loading cooked content in editor.

diff --git a/Engine/Source/Editor/ContentBrowser/Private/ContentBrowserUtils.cpp b/Engine/Source/Editor/ContentBrowser/Private/ContentBrowserUtils.cpp
--- a/Engine/Source/Editor/ContentBrowser/Private/ContentBrowserUtils.cpp
+++ b/Engine/Source/Editor/ContentBrowser/Private/ContentBrowserUtils.cpp
@@ -413,8 +413,8 @@ bool ContentBrowserUtils::CanRenameFolder(const FString& InFolderPath)
 
 bool ContentBrowserUtils::CanRenameAsset(const FAssetData& InAssetData)
 {
-	// Cannot rename redirectors or classes
-	return !InAssetData.IsRedirector() && InAssetData.AssetClass != NAME_Class;
+	// Cannot rename redirectors or classes or cooked packages
+	return !InAssetData.IsRedirector() && InAssetData.AssetClass != NAME_Class && !(InAssetData.PackageFlags & PKG_FilterEditorOnly);
 }
 
 void ContentBrowserUtils::RenameAsset(UObject* Asset, const FString& NewName, FText& ErrorMessage)
diff --git a/Engine/Source/Editor/ContentBrowser/Private/SAssetView.cpp b/Engine/Source/Editor/ContentBrowser/Private/SAssetView.cpp
--- a/Engine/Source/Editor/ContentBrowser/Private/SAssetView.cpp
+++ b/Engine/Source/Editor/ContentBrowser/Private/SAssetView.cpp
@@ -498,8 +498,9 @@ void SAssetView::DeferredCreateNewAsset()
 		FName AssetClassName = DeferredAssetToCreate->AssetClass->GetFName();
 		TMap<FName, FString> EmptyTags;
 		TArray<int32> EmptyChunkIDs;
-	
-		FAssetData NewAssetData(PackageName, PackagePathFName, NAME_None, AssetName, AssetClassName, EmptyTags, EmptyChunkIDs);
+		const uint32 EmptyPackageFlags = 0;
+
+		FAssetData NewAssetData(PackageName, PackagePathFName, NAME_None, AssetName, AssetClassName, EmptyTags, EmptyChunkIDs, EmptyPackageFlags);
 		TSharedPtr<FAssetViewItem> NewItem = MakeShareable(new FAssetViewCreation(NewAssetData, DeferredAssetToCreate->AssetClass, DeferredAssetToCreate->Factory));
 
 		NewItem->bRenameWhenScrolledIntoview = true;
@@ -536,8 +537,9 @@ void SAssetView::DuplicateAsset(const FString& PackagePath, const TWeakObjectPtr
 	FName AssetClass = OriginalObject->GetClass()->GetFName();
 	TMap<FName, FString> EmptyTags;
 	TArray<int32> EmptyChunkIDs;
+	const uint32 EmptyPackageFlags = 0;
 
-	FAssetData NewAssetData(PackageName, PackagePathFName, NAME_None, AssetName, AssetClass, EmptyTags, EmptyChunkIDs);
+	FAssetData NewAssetData(PackageName, PackagePathFName, NAME_None, AssetName, AssetClass, EmptyTags, EmptyChunkIDs, EmptyPackageFlags);
 	TSharedPtr<FAssetViewItem> NewItem = MakeShareable(new FAssetViewDuplication(NewAssetData, OriginalObject));
 	NewItem->bRenameWhenScrolledIntoview = true;
 
diff --git a/Engine/Source/Editor/UnrealEd/Private/ObjectTools.cpp b/Engine/Source/Editor/UnrealEd/Private/ObjectTools.cpp
--- a/Engine/Source/Editor/UnrealEd/Private/ObjectTools.cpp
+++ b/Engine/Source/Editor/UnrealEd/Private/ObjectTools.cpp
@@ -2592,6 +2592,10 @@ namespace ObjectTools
 					UPackage* NewPackage = CreatePackage( NULL, *FullPackageName );
 					NewPackage->GetOutermost()->FullyLoad();
 
+					// Make sure we copy all the cooked package flags if the asset was already cooked.
+					NewPackage->PackageFlags |= (Object->GetOutermost()->PackageFlags & PKG_FilterEditorOnly);
+					NewPackage->bIsCookedForEditor = Object->GetOutermost()->bIsCookedForEditor;
+
 					UObjectRedirector* Redirector = Cast<UObjectRedirector>( StaticFindObject(UObjectRedirector::StaticClass(), NewPackage, *NewObjectName) );
 					bool bFoundCompatibleRedirector = false;
 					// If we found a redirector, check that the object it points to is of the same class.
diff --git a/Engine/Source/Editor/UnrealEd/Private/Toolkits/AssetEditorManager.cpp b/Engine/Source/Editor/UnrealEd/Private/Toolkits/AssetEditorManager.cpp
--- a/Engine/Source/Editor/UnrealEd/Private/Toolkits/AssetEditorManager.cpp
+++ b/Engine/Source/Editor/UnrealEd/Private/Toolkits/AssetEditorManager.cpp
@@ -245,8 +245,16 @@ bool FAssetEditorManager::OpenEditorForAsset(UObject* Asset, const EToolkitMode:
 	//    being edited within, we should decide whether to disallow "Edit Here" in that case, or to close the old asset
 	//    editor and summon it in the new level editor, or to just foreground the old level editor (current behavior)
 
-
 	const bool bBringToFrontIfOpen = true;
+
+	// Don't open asset editors for cooked packages
+	if (UPackage* Package = Asset->GetOutermost())
+	{
+		if (Package->bIsCookedForEditor)
+		{
+			return false;
+		}
+	}
 	
 	AssetEditorRequestOpenEvent.Broadcast(Asset);
 
diff --git a/Engine/Source/Runtime/AssetRegistry/Private/AssetDataGatherer.cpp b/Engine/Source/Runtime/AssetRegistry/Private/AssetDataGatherer.cpp
--- a/Engine/Source/Runtime/AssetRegistry/Private/AssetDataGatherer.cpp
+++ b/Engine/Source/Runtime/AssetRegistry/Private/AssetDataGatherer.cpp
@@ -190,17 +190,30 @@ uint32 FAssetDataGatherer::Run()
 
 						if ( bLoadAndSaveCache )
 						{
-							// Update the cache
-							const FName PackageName = FName(*FPackageName::FilenameToLongPackageName(AssetFile));
-							const FDateTime& FileTimestamp = IFileManager::Get().GetTimeStamp(*AssetFile);
-							FDiskCachedAssetData* NewData = new FDiskCachedAssetData(PackageName, FileTimestamp);
-							for ( auto AssetIt = AssetDataFromFile.CreateConstIterator(); AssetIt; ++AssetIt )
+							// don't store info on cooked packages
+							bool bIsCooked = false;
+							for (const auto& AssetData : AssetDataFromFile)
 							{
-								NewData->AssetDataList.Add((*AssetIt)->ToAssetData());
+								bIsCooked |= AssetData->IsCooked();
+								if (bIsCooked)
+								{
+									break;
+								}
+							}
+							if (bIsCooked == false)
+							{
+								// Update the cache
+								const FName PackageName = FName(*FPackageName::FilenameToLongPackageName(AssetFile));
+								const FDateTime& FileTimestamp = IFileManager::Get().GetTimeStamp(*AssetFile);
+								FDiskCachedAssetData* NewData = new FDiskCachedAssetData(PackageName, FileTimestamp);
+								for (auto AssetIt = AssetDataFromFile.CreateConstIterator(); AssetIt; ++AssetIt)
+								{
+									NewData->AssetDataList.Add((*AssetIt)->ToAssetData());
+								}
+								NewData->DependencyData = DependencyData;
+								NewCachedAssetData.Add(NewData);
+								NewCachedAssetDataMap.Add(PackageName, NewData);
 							}
-							NewData->DependencyData = DependencyData;
-							NewCachedAssetData.Add(NewData);
-							NewCachedAssetDataMap.Add(PackageName, NewData);
 						}
 					}
 				}
@@ -415,7 +428,66 @@ bool FAssetDataGatherer::IsValidPackageFileToRead(const FString& Filename) const
 	return false;
 }
 
-bool FAssetDataGatherer::ReadAssetFile(const FString& AssetFilename, TArray<FBackgroundAssetData*>& AssetDataList, FPackageDependencyData& DependencyData) const
+FBackgroundAssetData* FAssetDataGatherer::CreateAssetDataFromLinkerTables(const FString& AssetFilename, uint32 InPackageFlags, const FObjectExport& AssetExport, const TArray<FObjectImport>& ImportMap, const TArray<FObjectExport>& ExportMap) const
+{
+	const FString PackageName = FPackageName::FilenameToLongPackageName(AssetFilename);
+	const FString PackagePath = FPaths::GetPath(PackageName);
+	FString GroupNames; // Not used for anything
+	TMap<FString, FString> Tags; // Not used for anything
+	TArray<int32> ChunkIDs; // Not used for anything
+
+	// We need to get the class name from the import/export maps
+	FName ObjectClassName;
+	if (AssetExport.ClassIndex.IsNull())
+	{
+		ObjectClassName = UClass::StaticClass()->GetFName();
+	}
+	else if (AssetExport.ClassIndex.IsExport())
+	{
+		const FObjectExport& ClassExport = ExportMap[AssetExport.ClassIndex.ToExport()];
+		ObjectClassName = ClassExport.ObjectName;
+	}
+	else if (AssetExport.ClassIndex.IsImport())
+	{
+		const FObjectImport& ClassImport = ImportMap[AssetExport.ClassIndex.ToImport()];
+		ObjectClassName = ClassImport.ObjectName;
+	}
+	return new FBackgroundAssetData(PackageName, PackagePath, GroupNames, AssetExport.ObjectName.ToString(), ObjectClassName.ToString(), Tags, ChunkIDs, InPackageFlags);
+}
+
+FBackgroundAssetData* FAssetDataGatherer::CreateAssetDataFromCookedPackage(const FString& AssetFilename, uint32 InPackageFlags, FPackageReader& PackageReader) const
+{
+	FString PackageName = FPackageName::FilenameToLongPackageName(AssetFilename);
+	FBackgroundAssetData* Result = nullptr;
+
+	// If the packaged is saved with the right version we have the information
+	// which of the objects in the export map as the asset.
+	// Otherwise we need to store a temp minimal data and then force load the asset
+	// to re-generate its registry data
+	if (PackageReader.UE4Ver() >= VER_UE4_COOKED_ASSETS_IN_EDITOR_SUPPORT)
+	{
+		TArray<FObjectImport> ImportMap;
+		TArray<FObjectExport> ExportMap;		
+		PackageReader.SerializeNameMap();
+		PackageReader.SerializeImportMap(ImportMap);
+		PackageReader.SerializeExportMap(ExportMap);
+		for (FObjectExport& Export : ExportMap)
+		{
+			if (Export.bIsAsset)
+			{
+				Result = CreateAssetDataFromLinkerTables(AssetFilename, InPackageFlags, Export, ImportMap, ExportMap);
+				break;
+			}
+		}
+	}	
+	if (!Result)
+	{
+		Result = new FBackgroundAssetData(PackageName, InPackageFlags);
+	}
+	return Result;
+}
+
+bool FAssetDataGatherer::ReadAssetFile(const FString& AssetFilename, TArray<FBackgroundAssetData*>& AssetDataList, FPackageDependencyData& DependencyData ) const
 {
 	FPackageReader PackageReader;
 
@@ -424,6 +496,15 @@ bool FAssetDataGatherer::ReadAssetFile(const FString& AssetFilename, TArray<FBac
 		return false;
 	}
 
+	if (!!(PackageReader.GetPackageFlags() & PKG_FilterEditorOnly))
+	{
+		// Try to reconstruct asset data from the cooked package by serializing its import and export tables
+		FBackgroundAssetData* CookedPackageData = CreateAssetDataFromCookedPackage(AssetFilename, PackageReader.GetPackageFlags(), PackageReader);
+		AssetDataList.Add(CookedPackageData);
+		return true;
+	}
+
+
 	if ( !PackageReader.ReadAssetRegistryData(AssetDataList) )
 	{
 		if ( !PackageReader.ReadAssetDataFromThumbnailCache(AssetDataList) )
diff --git a/Engine/Source/Runtime/AssetRegistry/Private/AssetDataGatherer.h b/Engine/Source/Runtime/AssetRegistry/Private/AssetDataGatherer.h
--- a/Engine/Source/Runtime/AssetRegistry/Private/AssetDataGatherer.h
+++ b/Engine/Source/Runtime/AssetRegistry/Private/AssetDataGatherer.h
@@ -55,6 +55,12 @@ private:
 	/** Serializes the timestamped cache of discovered assets. Used for quick loading of data for assets that have not changed on disk */
 	void SerializeCache(FArchive& Ar);
 
+	/** Creates asset data reconstructing all the required info from linker tables */
+	FBackgroundAssetData* CreateAssetDataFromLinkerTables(const FString& AssetFilename, uint32 InPackageFlags, const FObjectExport& AssetExport, const TArray<FObjectImport>& ImportMap, const TArray<FObjectExport>& ExportMap) const;
+
+	/** Creates asset data reconstructing all the required data from cooked package info */
+	FBackgroundAssetData* CreateAssetDataFromCookedPackage(const FString& AssetFilename, uint32 InPackageFlags, FPackageReader& PackageReader) const;
+
 private:
 	/** A critical section to protect data transfer to the main thread */
 	FCriticalSection WorkerThreadCriticalSection;
diff --git a/Engine/Source/Runtime/AssetRegistry/Private/AssetRegistry.cpp b/Engine/Source/Runtime/AssetRegistry/Private/AssetRegistry.cpp
--- a/Engine/Source/Runtime/AssetRegistry/Private/AssetRegistry.cpp
+++ b/Engine/Source/Runtime/AssetRegistry/Private/AssetRegistry.cpp
@@ -469,7 +469,7 @@ bool FAssetRegistry::GetAssets(const FARFilter& Filter, TArray<FAssetData>& OutA
 				}
 
 				// This asset is in memory and passes all filters
-				FAssetData* AssetData = new (OutAssetData) FAssetData(PackageName, PackagePath, FName(*GroupNamesStr), Obj->GetFName(), Obj->GetClass()->GetFName(), TagMap, InMemoryPackage->GetChunkIDs());
+				FAssetData* AssetData = new (OutAssetData)FAssetData(PackageName, PackagePath, FName(*GroupNamesStr), Obj->GetFName(), Obj->GetClass()->GetFName(), TagMap, InMemoryPackage->GetChunkIDs(), InMemoryPackage->PackageFlags);
 			}
 		};
 
@@ -1714,6 +1714,16 @@ void FAssetRegistry::AssetSearchDataGathered(const double TickStartTime, TArray<
 	for (AssetIndex = 0; AssetIndex < AssetResults.Num(); ++AssetIndex)
 	{
 		IGatheredAssetData*& BackgroundResult = AssetResults[AssetIndex];
+
+		// If this data is cooked and it we couldn't find any asset in its export table then try load the entire package 
+		if (BackgroundResult->IsCooked())
+		{
+			LoadPackage(nullptr, *BackgroundResult->GetPackageName(), 0);
+			delete BackgroundResult;
+			BackgroundResult = nullptr;
+			continue;
+		}
+
 		FAssetData Result = BackgroundResult->ToAssetData();
 
 		// Try to update any asset data that may already exist
diff --git a/Engine/Source/Runtime/AssetRegistry/Private/BackgroundAssetData.cpp b/Engine/Source/Runtime/AssetRegistry/Private/BackgroundAssetData.cpp
--- a/Engine/Source/Runtime/AssetRegistry/Private/BackgroundAssetData.cpp
+++ b/Engine/Source/Runtime/AssetRegistry/Private/BackgroundAssetData.cpp
@@ -2,7 +2,7 @@
 
 #include "AssetRegistryPCH.h"
 
-FBackgroundAssetData::FBackgroundAssetData(FString InPackageName, FString InPackagePath, FString InGroupNames, FString InAssetName, FString InAssetClass, TMap<FString, FString> InTags, TArray<int32> InChunkIDs)
+FBackgroundAssetData::FBackgroundAssetData(FString InPackageName, FString InPackagePath, FString InGroupNames, FString InAssetName, FString InAssetClass, TMap<FString, FString> InTags, TArray<int32> InChunkIDs, uint32 InPackageFlags)
 	: PackageName(MoveTemp(InPackageName))
 	, PackagePath(MoveTemp(InPackagePath))
 	, GroupNames(MoveTemp(InGroupNames))
@@ -10,6 +10,7 @@ FBackgroundAssetData::FBackgroundAssetData(FString InPackageName, FString InPack
 	, AssetClass(MoveTemp(InAssetClass))
 	, TagsAndValues(MakeSharedMapView(MoveTemp(InTags)))
 	, ChunkIDs(MoveTemp(InChunkIDs))
+	, PackageFlags(InPackageFlags)
 {
 	ObjectPath = PackageName + TEXT(".");
 
@@ -21,6 +22,13 @@ FBackgroundAssetData::FBackgroundAssetData(FString InPackageName, FString InPack
 	ObjectPath += AssetName;
 }
 
+FBackgroundAssetData::FBackgroundAssetData(FString InPackageName, uint32 InPackageFlags)
+	: PackageName(InPackageName)
+	, PackageFlags(InPackageFlags)
+{
+}
+
+
 FAssetData FBackgroundAssetData::ToAssetData() const
 {
 	TMap<FName, FString> CopiedTagsAndValues;
@@ -36,7 +44,8 @@ FAssetData FBackgroundAssetData::ToAssetData() const
 		FName(*AssetName),
 		FName(*AssetClass),
 		MoveTemp(CopiedTagsAndValues),
-		ChunkIDs
+		ChunkIDs,
+		PackageFlags
 		);
 }
 
@@ -60,3 +69,9 @@ bool FAssetDataWrapper::IsWithinSearchPath(const FString& InSearchPath) const
 {
 	return PackagePathStr.StartsWith(InSearchPath);
 }
+
+
+FString FAssetDataWrapper::GetPackageName() const
+{ 
+	return WrappedAssetData.PackageName.ToString();  
+}
\ No newline at end of file
diff --git a/Engine/Source/Runtime/AssetRegistry/Private/BackgroundAssetData.h b/Engine/Source/Runtime/AssetRegistry/Private/BackgroundAssetData.h
--- a/Engine/Source/Runtime/AssetRegistry/Private/BackgroundAssetData.h
+++ b/Engine/Source/Runtime/AssetRegistry/Private/BackgroundAssetData.h
@@ -15,6 +15,13 @@ public:
 
 	/** Test to see if this asset data is within the given search path */
 	virtual bool IsWithinSearchPath(const FString& InSearchPath) const = 0;
+
+	/** Returns true if the asset was cooked */
+	virtual bool IsCooked() const = 0;
+
+	/** Returns the package name */
+	virtual FString GetPackageName() const = 0;
+
 };
 
 /** A class to hold important information about an assets found by the Asset Registry. Intended for use by the background gathering thread to avoid creation of FNames */
@@ -22,12 +29,15 @@ class FBackgroundAssetData : public IGatheredAssetData
 {
 public:
 	/** Constructor */
-	FBackgroundAssetData(FString InPackageName, FString InPackagePath, FString InGroupNames, FString InAssetName, FString InAssetClass, TMap<FString, FString> InTags, TArray<int32> InChunkIDs);
-
+	FBackgroundAssetData(FString InPackageName, FString InPackagePath, FString InGroupNames, FString InAssetName, FString InAssetClass, TMap<FString, FString> InTags, TArray<int32> InChunkIDs, uint32 InPackageFlags);
+	FBackgroundAssetData(FString InPackageName, uint32 InPackageFlags);
 	/** IGatheredAssetData interface */
 	virtual FAssetData ToAssetData() const override;
 	virtual bool IsWithinSearchPath(const FString& InSearchPath) const override;
 	
+	virtual bool IsCooked()const override { return !!(PackageFlags & PKG_FilterEditorOnly); }
+
+	virtual FString GetPackageName() const override { return PackageName; }
 private:
 	/** The object path for the asset in the form 'Package.GroupNames.AssetName' */
 	FString ObjectPath;
@@ -45,6 +55,8 @@ private:
 	TSharedMapView<FString, FString> TagsAndValues;
 	/** The IDs of the chunks this asset is located in for streaming install.  Empty if not assigned to a chunk */
 	TArray<int32> ChunkIDs;
+	/** Asset package flags */
+	uint32 PackageFlags;
 };
 
 /** A class to wrap up an existing FAssetData to avoid redundant FName -> FString -> FName conversion */
@@ -57,7 +69,8 @@ public:
 	/** IGatheredAssetData interface */
 	virtual FAssetData ToAssetData() const override;
 	virtual bool IsWithinSearchPath(const FString& InSearchPath) const override;
-
+	virtual bool IsCooked() const override { return false;  }
+	virtual FString GetPackageName()const override;
 private:
 	/** The asset data we're wrapping up */
 	FAssetData WrappedAssetData;
diff --git a/Engine/Source/Runtime/AssetRegistry/Private/PackageReader.cpp b/Engine/Source/Runtime/AssetRegistry/Private/PackageReader.cpp
--- a/Engine/Source/Runtime/AssetRegistry/Private/PackageReader.cpp
+++ b/Engine/Source/Runtime/AssetRegistry/Private/PackageReader.cpp
@@ -141,7 +141,7 @@ bool FPackageReader::ReadAssetRegistryData (TArray<FBackgroundAssetData*>& Asset
 		if (bLegacyPackage || bNoMapAsset)
 		{
 			FString AssetName = FPackageName::GetLongPackageAssetName(PackageName);
-			AssetDataList.Add(new FBackgroundAssetData(PackageName, PackagePath, FString(), MoveTemp(AssetName), TEXT("World"), TMap<FString, FString>(), PackageFileSummary.ChunkIDs));
+			AssetDataList.Add(new FBackgroundAssetData(PackageName, PackagePath, FString(), MoveTemp(AssetName), TEXT("World"), TMap<FString, FString>(), PackageFileSummary.ChunkIDs, PackageFileSummary.PackageFlags));
 		}
 	}
 
@@ -190,7 +190,7 @@ bool FPackageReader::ReadAssetRegistryData (TArray<FBackgroundAssetData*>& Asset
 		}
 
 		// Create a new FBackgroundAssetData for this asset and update it with the gathered data
-		AssetDataList.Add(new FBackgroundAssetData(PackageName, PackagePath, MoveTemp(GroupNames), MoveTemp(AssetName), MoveTemp(ObjectClassName), MoveTemp(TagsAndValues), PackageFileSummary.ChunkIDs));
+		AssetDataList.Add(new FBackgroundAssetData(PackageName, PackagePath, MoveTemp(GroupNames), MoveTemp(AssetName), MoveTemp(ObjectClassName), MoveTemp(TagsAndValues), PackageFileSummary.ChunkIDs, PackageFileSummary.PackageFlags));
 	}
 
 	return true;
@@ -245,7 +245,7 @@ bool FPackageReader::ReadAssetDataFromThumbnailCache(TArray<FBackgroundAssetData
 		}
 
 		// Create a new FBackgroundAssetData for this asset and update it with the gathered data
-		AssetDataList.Add(new FBackgroundAssetData(PackageName, PackagePath, MoveTemp(GroupNames), MoveTemp(AssetName), MoveTemp(AssetClassName), TMap<FString, FString>(), PackageFileSummary.ChunkIDs));
+		AssetDataList.Add(new FBackgroundAssetData(PackageName, PackagePath, MoveTemp(GroupNames), MoveTemp(AssetName), MoveTemp(AssetClassName), TMap<FString, FString>(), PackageFileSummary.ChunkIDs, PackageFileSummary.PackageFlags));
 	}
 
 	return true;
@@ -296,6 +296,19 @@ void FPackageReader::SerializeImportMap(TArray<FObjectImport>& OutImportMap)
 	}
 }
 
+void FPackageReader::SerializeExportMap(TArray<FObjectExport>& OutExportMap)
+{
+	if (PackageFileSummary.ExportCount > 0)
+	{
+		Seek(PackageFileSummary.ExportOffset);
+		for (int32 ExportMapIdx = 0; ExportMapIdx < PackageFileSummary.ExportCount; ++ExportMapIdx)
+		{
+			FObjectExport* Export = new(OutExportMap)FObjectExport;
+			*this << *Export;
+		}
+	}
+}
+
 void FPackageReader::SerializeStringAssetReferencesMap(TArray<FString>& OutStringAssetReferencesMap)
 {
 	if (UE4Ver() >= VER_UE4_ADD_STRING_ASSET_REFERENCES_MAP && PackageFileSummary.StringAssetReferencesCount > 0)
@@ -369,6 +382,11 @@ int64 FPackageReader::TotalSize()
 	return Loader->TotalSize();
 }
 
+uint32 FPackageReader::GetPackageFlags() const
+{
+	return PackageFileSummary.PackageFlags;
+}
+
 FArchive& FPackageReader::operator<<( FName& Name )
 {
 	check(Loader);
diff --git a/Engine/Source/Runtime/AssetRegistry/Private/PackageReader.h b/Engine/Source/Runtime/AssetRegistry/Private/PackageReader.h
--- a/Engine/Source/Runtime/AssetRegistry/Private/PackageReader.h
+++ b/Engine/Source/Runtime/AssetRegistry/Private/PackageReader.h
@@ -23,8 +23,12 @@ public:
 	/** Serializers for different package maps */
 	void SerializeNameMap();
 	void SerializeImportMap(TArray<FObjectImport>& OutImportMap);
+	void SerializeExportMap(TArray<FObjectExport>& OutExportMap);
 	void SerializeStringAssetReferencesMap(TArray<FString>& OutStringAssetReferencesMap);
 
+	/** Returns flags the asset package was saved with */
+	uint32 GetPackageFlags() const;
+
 	// Farchive implementation to redirect requests to the Loader
 	void Serialize( void* V, int64 Length );
 	bool Precache( int64 PrecacheOffset, int64 PrecacheSize );
diff --git a/Engine/Source/Runtime/AssetRegistry/Public/AssetData.h b/Engine/Source/Runtime/AssetRegistry/Public/AssetData.h
--- a/Engine/Source/Runtime/AssetRegistry/Public/AssetData.h
+++ b/Engine/Source/Runtime/AssetRegistry/Public/AssetData.h
@@ -26,6 +26,8 @@ public:
 	TSharedMapView<FName, FString> TagsAndValues;
 	/** The IDs of the chunks this asset is located in for streaming install.  Empty if not assigned to a chunk */
 	TArray<int32> ChunkIDs;
+	/** Asset package flags */
+	uint32 PackageFlags;
 
 public:
 	/** Default constructor */
@@ -33,7 +35,7 @@ public:
 	{}
 
 	/** Constructor */
-	FAssetData(FName InPackageName, FName InPackagePath, FName InGroupNames, FName InAssetName, FName InAssetClass, TMap<FName, FString> InTags, TArray<int32> InChunkIDs)
+	FAssetData(FName InPackageName, FName InPackagePath, FName InGroupNames, FName InAssetName, FName InAssetClass, TMap<FName, FString> InTags, TArray<int32> InChunkIDs, uint32 InPackageFlags)
 		: PackageName(InPackageName)
 		, PackagePath(InPackagePath)
 		, GroupNames(InGroupNames)
@@ -41,6 +43,7 @@ public:
 		, AssetClass(InAssetClass)
 		, TagsAndValues(MakeSharedMapView(MoveTemp(InTags)))
 		, ChunkIDs(MoveTemp(InChunkIDs))
+		, PackageFlags(InPackageFlags)
 	{
 		FString ObjectPathStr = PackageName.ToString() + TEXT(".");
 
@@ -92,6 +95,7 @@ public:
 			TagsAndValues = MakeSharedMapView(MoveTemp(TagsAndValuesMap));
 
 			ChunkIDs = Outermost->GetChunkIDs();
+			PackageFlags = Outermost->PackageFlags;
 		}
 	}
 
@@ -256,6 +260,8 @@ public:
 		{
 			UE_LOG(LogAssetData, Log, TEXT("                 %d"), Chunk);
 		}
+
+		UE_LOG(LogAssetData, Log, TEXT("        PackageFlahs: %d"), PackageFlags);
 	}
 
 	/** Get the first FAssetData of a particular class from an Array of FAssetData */
@@ -286,9 +292,6 @@ public:
 	/** Operator for serialization */
 	friend FArchive& operator<<(FArchive& Ar, FAssetData& AssetData)
 	{
-
-		//FNameAsStringProxyArchive Ar(InAr);
-
 		// serialize out the asset info
 		Ar << AssetData.ObjectPath;
 		Ar << AssetData.PackagePath;
@@ -334,6 +337,11 @@ public:
 			Ar << ChunkID;
 		}
 
+		if (Ar.UE4Ver() >= VER_UE4_COOKED_ASSETS_IN_EDITOR_SUPPORT)
+		{
+			Ar << AssetData.PackageFlags;
+		}
+
 		return Ar;
 	}
 };
diff --git a/Engine/Source/Runtime/Core/Public/UObject/ObjectVersion.h b/Engine/Source/Runtime/Core/Public/UObject/ObjectVersion.h
--- a/Engine/Source/Runtime/Core/Public/UObject/ObjectVersion.h
+++ b/Engine/Source/Runtime/Core/Public/UObject/ObjectVersion.h
@@ -564,6 +564,8 @@ enum EUnrealEngineObjectUE4Version
 	VER_UE4_DECAL_SIZE,
 	// Keep only package names in StringAssetReferencesMap
 	VER_UE4_KEEP_ONLY_PACKAGE_NAMES_IN_STRING_ASSET_REFERENCES_MAP,
+	// Support sound cue not saving out editor only data
+	VER_UE4_COOKED_ASSETS_IN_EDITOR_SUPPORT,
 	// -----<new versions can be added before this line>-------------------------------------------------
 	// - this needs to be the last line (see note below)
 	VER_UE4_AUTOMATIC_VERSION_PLUS_ONE,
diff --git a/Engine/Source/Runtime/CoreUObject/Private/UObject/LinkerLoad.cpp b/Engine/Source/Runtime/CoreUObject/Private/UObject/LinkerLoad.cpp
--- a/Engine/Source/Runtime/CoreUObject/Private/UObject/LinkerLoad.cpp
+++ b/Engine/Source/Runtime/CoreUObject/Private/UObject/LinkerLoad.cpp
@@ -995,13 +995,37 @@ FLinkerLoad::ELinkerStatus FLinkerLoad::SerializePackageFileSummary()
 		}
 
 		// don't load packages that contain editor only data in builds that don't support that and vise versa
-		if ( (!FPlatformProperties::HasEditorOnlyData() && !(Summary.PackageFlags & PKG_FilterEditorOnly)) ||
-			 (FPlatformProperties::HasEditorOnlyData() && !!(Summary.PackageFlags & PKG_FilterEditorOnly)) )
+		if (!FPlatformProperties::HasEditorOnlyData() && !(Summary.PackageFlags & PKG_FilterEditorOnly))
 		{
-			UE_LOG(LogLinker, Warning, TEXT("Unable to load package (%s). Package contains EditorOnly data which is not supported by the current build or vice versa."), *Filename );
+			UE_LOG(LogLinker, Warning, TEXT("Unable to load package (%s). Package contains EditorOnly data which is not supported by the current build."), *Filename );
 			return LINKER_Failed;
 		}
 
+		// don't load packages that contain editor only data in builds that don't support that and vise versa
+		if (FPlatformProperties::HasEditorOnlyData() && !!(Summary.PackageFlags & PKG_FilterEditorOnly))
+		{
+			// This warning can be disabled in ini with [Core.System] AllowCookedDataInEditorBuilds=False
+			static struct FInitCookedDatataInEditorBuildsSupport
+			{
+				bool bAllowCookedData;
+				FInitCookedDatataInEditorBuildsSupport()
+				{
+					if (!GConfig->GetBool(TEXT("Core.System"), TEXT("AllowCookedDataInEditorBuilds"), bAllowCookedData, GEngineIni))
+					{
+						bAllowCookedData = true;
+					}
+				}
+				FORCEINLINE operator bool() const { return bAllowCookedData; }
+			} AllowCookedDataInEditorBuilds;
+			if (!AllowCookedDataInEditorBuilds)
+			{
+				UE_LOG(LogLinker, Warning, 
+					TEXT("Unable to load package (%s). Package contains cooked data which is not supported by the current build. Set [Core.System] AllowCookedDataInEditorBuilds to true in Engine.ini to allow it."), 
+					*Filename);
+				return LINKER_Failed;
+			}
+		}
+
 #if PLATFORM_WINDOWS
 		if (!FPlatformProperties::RequiresCookedData())
 		{
@@ -1123,6 +1147,10 @@ FLinkerLoad::ELinkerStatus FLinkerLoad::SerializePackageFileSummary()
 			LinkerRootPackage->LinkerPackageVersion = ArUE4Ver;
 			LinkerRootPackage->LinkerLicenseeVersion = ArLicenseeUE4Ver;
 			LinkerRootPackage->LinkerCustomVersion = SummaryVersions;
+
+#if WITH_EDITORONLY_DATA
+			LinkerRootPackage->bIsCookedForEditor = !!(Summary.PackageFlags & PKG_FilterEditorOnly);
+#endif
 		}
 		
 		// Propagate fact that package cannot use lazy loading to archive (aka this).
diff --git a/Engine/Source/Runtime/CoreUObject/Private/UObject/ObjectResource.cpp b/Engine/Source/Runtime/CoreUObject/Private/UObject/ObjectResource.cpp
--- a/Engine/Source/Runtime/CoreUObject/Private/UObject/ObjectResource.cpp
+++ b/Engine/Source/Runtime/CoreUObject/Private/UObject/ObjectResource.cpp
@@ -30,44 +30,48 @@ FObjectResource::FObjectResource( UObject* InObject )
 -----------------------------------------------------------------------------*/
 
 FObjectExport::FObjectExport()
-:	FObjectResource	()
-,	ObjectFlags		( RF_NoFlags												)
-,	bExportLoadFailed	( false													)
-,	SerialSize		( 0															)
-,	SerialOffset	( 0															)
-,	ScriptSerializationStartOffset	( 0											)
-,	ScriptSerializationEndOffset	( 0											)
-,	Object			( NULL														)
-,	HashNext		( INDEX_NONE												)
-,	bForcedExport	( false														)
-,	bNotForClient	( false														)
-,	bNotForServer	( false														)
-,	PackageGuid		( FGuid(0,0,0,0)											)
-,	bNotForEditorGame	(true													)
+: FObjectResource()
+, ObjectFlags(RF_NoFlags)
+, SerialSize(0)
+, SerialOffset(0)
+, ScriptSerializationStartOffset(0)
+, ScriptSerializationEndOffset(0)
+, Object(NULL)
+, HashNext(INDEX_NONE)
+, bForcedExport(false)
+, bNotForClient(false)
+, bNotForServer(false)
+, bNotForEditorGame(true)
+, bIsAsset(false)
+, bExportLoadFailed(false)
+, PackageGuid(FGuid(0, 0, 0, 0))
+, PackageFlags(0)
 {}
 
 FObjectExport::FObjectExport( UObject* InObject )
-:	FObjectResource	( InObject													)
-,	ObjectFlags		( InObject ? InObject->GetMaskedFlags() : RF_NoFlags		)
-,	bExportLoadFailed	( false													)
-,	SerialSize		( 0															)
-,	SerialOffset	( 0															)
-,	ScriptSerializationStartOffset	( 0											)
-,	ScriptSerializationEndOffset	( 0											)
-,	Object			( InObject													)
-,	HashNext		( INDEX_NONE												)
-,	bForcedExport	( false														)
-,	bNotForClient	( false														)
-,	bNotForServer	( false														)
-,	PackageGuid		( FGuid(0,0,0,0)											)
-,	PackageFlags	( 0															)
-,	bNotForEditorGame	(true													)
+: FObjectResource(InObject)
+, ObjectFlags(InObject ? InObject->GetMaskedFlags() : RF_NoFlags)
+, SerialSize(0)
+, SerialOffset(0)
+, ScriptSerializationStartOffset(0)
+, ScriptSerializationEndOffset(0)
+, Object(InObject)
+, HashNext(INDEX_NONE)
+, bForcedExport(false)
+, bNotForClient(false)
+, bNotForServer(false)
+, bNotForEditorGame(true)
+, bIsAsset(false)
+, bExportLoadFailed(false)
+, PackageGuid(FGuid(0, 0, 0, 0))
+, PackageFlags(0)
 {
 	if(Object)		
 	{
 		bNotForClient = Object->HasAnyMarks(OBJECTMARK_NotForClient);
 		bNotForServer = Object->HasAnyMarks(OBJECTMARK_NotForServer);
 		bNotForEditorGame = Object->HasAnyMarks(OBJECTMARK_NotForEditorGame);
+		bIsAsset = Object->IsAsset();
 	}
 }
 
@@ -100,6 +104,11 @@ FArchive& operator<<( FArchive& Ar, FObjectExport& E )
 		Ar << E.bNotForEditorGame;
 	}
 
+	if (Ar.UE4Ver() >= VER_UE4_COOKED_ASSETS_IN_EDITOR_SUPPORT)
+	{
+		Ar << E.bIsAsset;
+	}
+
 	return Ar;
 }
 
diff --git a/Engine/Source/Runtime/CoreUObject/Private/UObject/Package.cpp b/Engine/Source/Runtime/CoreUObject/Private/UObject/Package.cpp
--- a/Engine/Source/Runtime/CoreUObject/Private/UObject/Package.cpp
+++ b/Engine/Source/Runtime/CoreUObject/Private/UObject/Package.cpp
@@ -33,10 +33,12 @@ void UPackage::PostInitProperties()
 	MetaData = NULL;
 	LinkerPackageVersion = GPackageFileUE4Version;
 	LinkerLicenseeVersion = GPackageFileLicenseeUE4Version;
-
 #if WITH_EDITOR
 	PIEInstanceID = INDEX_NONE;
 #endif
+#if WITH_EDITORONLY_DATA
+	bIsCookedForEditor = false;
+#endif
 }
 
 
diff --git a/Engine/Source/Runtime/CoreUObject/Private/UObject/UObjectBaseUtility.cpp b/Engine/Source/Runtime/CoreUObject/Private/UObject/UObjectBaseUtility.cpp
--- a/Engine/Source/Runtime/CoreUObject/Private/UObject/UObjectBaseUtility.cpp
+++ b/Engine/Source/Runtime/CoreUObject/Private/UObject/UObjectBaseUtility.cpp
@@ -135,7 +135,7 @@ bool UObjectBaseUtility::MarkPackageDirty() const
 	{
 		UPackage* Package = GetOutermost();
 
-		if( Package != NULL )
+		if( Package != NULL	)
 		{
 			// It is against policy to dirty a map or package during load in the Editor, to enforce this policy
 			// we explicitly disable the ability to dirty a package or map during load.  Commandlets can still
@@ -145,6 +145,9 @@ bool UObjectBaseUtility::MarkPackageDirty() const
 #if WITH_HOT_RELOAD
 				&& !GIsHotReload
 #endif // WITH_HOT_RELOAD
+#if WITH_EDITORONLY_DATA
+				&& !Package->bIsCookedForEditor // Cooked packages can't be modified nor marked as dirty
+#endif
 				))
 			{
 				const bool bIsDirty = Package->IsDirty();
diff --git a/Engine/Source/Runtime/CoreUObject/Public/UObject/CoreObject.h b/Engine/Source/Runtime/CoreUObject/Public/UObject/CoreObject.h
--- a/Engine/Source/Runtime/CoreUObject/Public/UObject/CoreObject.h
+++ b/Engine/Source/Runtime/CoreUObject/Public/UObject/CoreObject.h
@@ -108,6 +108,11 @@ public:
 	int32 PIEInstanceID;
 #endif
 
+#if WITH_EDITORONLY_DATA
+	/** True if this packages has been cooked for the editor / opened cooked by the editor */
+	bool bIsCookedForEditor;
+#endif
+
 	/**
 	 * Called after the C++ constructor and after the properties have been initialized, but before the config has been loaded, etc.
 	 * mainly this is to emulate some behavior of when the constructor was called after the properties were intialized.
diff --git a/Engine/Source/Runtime/CoreUObject/Public/UObject/Linker.h b/Engine/Source/Runtime/CoreUObject/Public/UObject/Linker.h
--- a/Engine/Source/Runtime/CoreUObject/Public/UObject/Linker.h
+++ b/Engine/Source/Runtime/CoreUObject/Public/UObject/Linker.h
@@ -200,11 +200,6 @@ struct FObjectExport : public FObjectResource
 	 */
 	EObjectFlags	ObjectFlags;
 
-	/**
-	 * Force this export to not load, it failed because the outer didn't exist.
-	 */
-	bool         	bExportLoadFailed;
-
 	/**
 	 * The number of bytes to serialize when saving/loading this export's UObject.
 	 * Serialized
@@ -262,6 +257,24 @@ struct FObjectExport : public FObjectResource
 	 */
 	bool			bNotForServer;
 
+	/**
+	* whether the export should be always loaded in editor game
+	* False means that the object is necessary for editor game,
+	* True doesn't means, that the object won't be loaded.
+	* Serialized
+	*/
+	bool			bNotForEditorGame;
+
+	/**
+	* True if this export is an asset object.
+	*/
+	bool			bIsAsset;
+
+	/**
+	* Force this export to not load, it failed because the outer didn't exist.
+	*/
+	bool			bExportLoadFailed;
+
 	/** If this object is a top level package (which must have been forced into the export table via OBJECTMARK_ForceTagExp)
 	 * this is the GUID for the original package file
 	 * Serialized
@@ -274,14 +287,6 @@ struct FObjectExport : public FObjectResource
 	 */
 	uint32			PackageFlags;
 
-	/**
-	 * whether the export should be always loaded in editor game
-	 * False means that the object is necessary for editor game,
-	 * True doesn't means, that the object won't be loaded.
-	 * Serialized
-	 */
-	bool			bNotForEditorGame;
-
 	/**
 	 * Constructors
 	 */
@@ -790,7 +795,8 @@ public:
 		{
 			return Exp(Index);
 		}
-	}	/**
+	}	
+	/**
 	 * Return an import or export for this index
 	 * @param	Index	Package index to get
 	 * @return	the resource corresponding to this index, or NULL if the package index is null
diff --git a/Engine/Source/Runtime/Engine/Private/Animation/AnimInstance.cpp b/Engine/Source/Runtime/Engine/Private/Animation/AnimInstance.cpp
--- a/Engine/Source/Runtime/Engine/Private/Animation/AnimInstance.cpp
+++ b/Engine/Source/Runtime/Engine/Private/Animation/AnimInstance.cpp
@@ -313,9 +313,8 @@ void UAnimInstance::UpdateAnimation(float DeltaSeconds)
 		{
 			if (UAnimBlueprintGeneratedClass* AnimBlueprintClass = Cast<UAnimBlueprintGeneratedClass>(GetClass()))
 			{
-				UAnimBlueprint* AnimBP = CastChecked<UAnimBlueprint>(AnimBlueprintClass->ClassGeneratedBy);
-
-				if (AnimBP->GetObjectBeingDebugged() == this)
+				UAnimBlueprint* AnimBP = Cast<UAnimBlueprint>(AnimBlueprintClass->ClassGeneratedBy);
+				if (AnimBP && AnimBP->GetObjectBeingDebugged() == this)
 				{
 					AnimBlueprintClass->GetAnimBlueprintDebugData().ResetNodeVisitSites();
 				}
diff --git a/Engine/Source/Runtime/Engine/Private/Animation/ReferenceSkeleton.cpp b/Engine/Source/Runtime/Engine/Private/Animation/ReferenceSkeleton.cpp
--- a/Engine/Source/Runtime/Engine/Private/Animation/ReferenceSkeleton.cpp
+++ b/Engine/Source/Runtime/Engine/Private/Animation/ReferenceSkeleton.cpp
@@ -16,7 +16,7 @@ FArchive &operator<<(FArchive& Ar, FMeshBoneInfo& F)
 #if WITH_EDITORONLY_DATA
 	if (Ar.UE4Ver() >= VER_UE4_STORE_BONE_EXPORT_NAMES)
 	{
-		if (!Ar.IsCooking())
+		if (!Ar.IsCooking() && !Ar.IsFilterEditorOnly())
 		{
 			Ar << F.ExportName;
 		}
diff --git a/Engine/Source/Runtime/Engine/Private/Animation/Skeleton.cpp b/Engine/Source/Runtime/Engine/Private/Animation/Skeleton.cpp
--- a/Engine/Source/Runtime/Engine/Private/Animation/Skeleton.cpp
+++ b/Engine/Source/Runtime/Engine/Private/Animation/Skeleton.cpp
@@ -38,7 +38,7 @@ FArchive& operator<<(FArchive& Ar, FReferencePose & P)
 	Ar << P.ReferencePose;
 #if WITH_EDITORONLY_DATA
 	//TODO: we should use strip flags but we need to rev the serialization version
-	if (!Ar.IsCooking())
+	if (!Ar.IsCooking() && !Ar.IsFilterEditorOnly())
 	{
 		Ar << P.ReferenceMesh;
 	}
diff --git a/Engine/Source/Runtime/Engine/Private/SoundCue.cpp b/Engine/Source/Runtime/Engine/Private/SoundCue.cpp
--- a/Engine/Source/Runtime/Engine/Private/SoundCue.cpp
+++ b/Engine/Source/Runtime/Engine/Private/SoundCue.cpp
@@ -34,7 +34,19 @@ void USoundCue::PostInitProperties()
 	}
 }
 
-void USoundCue::Serialize( FArchive& Ar )
+
+void USoundCue::AddReferencedObjects(UObject* InThis, FReferenceCollector& Collector)
+{
+	USoundCue* This = CastChecked<USoundCue>(InThis);
+
+	Collector.AddReferencedObject(This->SoundCueGraph, This);
+
+	Super::AddReferencedObjects(InThis, Collector);
+}
+#endif // WITH_EDITOR
+
+
+void USoundCue::Serialize(FArchive& Ar)
 {
 	// Always force the duration to be updated when we are saving or cooking
 	if (Ar.IsSaving() || Ar.IsCooking())
@@ -42,42 +54,52 @@ void USoundCue::Serialize( FArchive& Ar )
 		Duration = (FirstNode ? FirstNode->GetDuration() : 0.f);
 	}
 
-	Super::Serialize( Ar );
+	Super::Serialize(Ar);
 
-	if (!Ar.IsCooking())
+	if (Ar.UE4Ver() >= VER_UE4_COOKED_ASSETS_IN_EDITOR_SUPPORT)
+	{
+		FStripDataFlags StripFlags(Ar);
+#if WITH_EDITORONLY_DATA
+		if (!StripFlags.IsEditorDataStripped())
+		{
+			Ar << SoundCueGraph;
+		}
+#endif
+	}
+#if WITH_EDITOR
+	else
 	{
 		Ar << SoundCueGraph;
 	}
+#endif
 }
 
-void USoundCue::AddReferencedObjects(UObject* InThis, FReferenceCollector& Collector)
-{
-	USoundCue* This = CastChecked<USoundCue>(InThis);
-
-	Collector.AddReferencedObject(This->SoundCueGraph, This);
-
-	Super::AddReferencedObjects(InThis, Collector);
-}
-#endif // WITH_EDITOR
-
 void USoundCue::PostLoad()
 {
 	Super::PostLoad();
 
 	// Game doesn't care if there are NULL graph nodes
 #if WITH_EDITOR
-	if (GIsEditor)
+	if (GIsEditor )
 	{
-		// Deal with SoundNode types being removed - iterate in reverse as nodes may be removed
-		for (int32 idx=SoundCueGraph->Nodes.Num()-1; idx >= 0; --idx)
+		if (SoundCueGraph)
 		{
-			USoundCueGraphNode* Node = Cast<USoundCueGraphNode>(SoundCueGraph->Nodes[idx]);
-
-			if (Node && Node->SoundNode == NULL)
+			// Deal with SoundNode types being removed - iterate in reverse as nodes may be removed
+			for (int32 idx = SoundCueGraph->Nodes.Num() - 1; idx >= 0; --idx)
 			{
-				FBlueprintEditorUtils::RemoveNode(NULL, Node, true);
+				USoundCueGraphNode* Node = Cast<USoundCueGraphNode>(SoundCueGraph->Nodes[idx]);
+
+				if (Node && Node->SoundNode == NULL)
+				{
+					FBlueprintEditorUtils::RemoveNode(NULL, Node, true);
+				}
 			}
 		}
+		else
+		{
+			// we should have a soundcuegraph unless we are contained in a package which is missing editor only data
+			check( GetOutermost()->PackageFlags & PKG_FilterEditorOnly );
+		}
 
 		// Always load all sound waves in the editor
 		for (USoundNode* SoundNode : AllNodes)
diff --git a/Engine/Source/Runtime/Engine/Private/SoundNode.cpp b/Engine/Source/Runtime/Engine/Private/SoundNode.cpp
--- a/Engine/Source/Runtime/Engine/Private/SoundNode.cpp
+++ b/Engine/Source/Runtime/Engine/Private/SoundNode.cpp
@@ -17,17 +17,30 @@ USoundNode::USoundNode(const FObjectInitializer& ObjectInitializer)
 {
 }
 
-#if WITH_EDITOR
+
 void USoundNode::Serialize(FArchive& Ar)
 {
 	Super::Serialize(Ar);
 
-	if (!Ar.IsCooking())
+	if (Ar.UE4Ver() >= VER_UE4_COOKED_ASSETS_IN_EDITOR_SUPPORT)
+	{
+		FStripDataFlags StripFlags(Ar);
+#if WITH_EDITORONLY_DATA
+		if (!StripFlags.IsEditorDataStripped())
+		{
+			Ar << GraphNode;
+		}
+#endif
+	}
+#if WITH_EDITOR
+	else
 	{
 		Ar << GraphNode;
 	}
+#endif
 }
 
+#if WITH_EDITOR
 void USoundNode::AddReferencedObjects(UObject* InThis, FReferenceCollector& Collector)
 {
 	USoundNode* This = CastChecked<USoundNode>(InThis);
diff --git a/Engine/Source/Runtime/Engine/Private/StaticMesh.cpp b/Engine/Source/Runtime/Engine/Private/StaticMesh.cpp
--- a/Engine/Source/Runtime/Engine/Private/StaticMesh.cpp
+++ b/Engine/Source/Runtime/Engine/Private/StaticMesh.cpp
@@ -622,7 +622,8 @@ void FStaticMeshRenderData::Serialize(FArchive& Ar, UStaticMesh* Owner, bool bCo
 
 	// Note: this is all derived data, native versioning is not needed, but be sure to bump STATICMESH_DERIVEDDATA_VER when modifying!
 #if WITH_EDITOR
-	if (Ar.IsSaving())
+	const bool bHasEditorData = !(Owner->GetOutermost()->PackageFlags & PKG_FilterEditorOnly);
+	if (Ar.IsSaving() && bHasEditorData)
 	{
 		ResolveSectionInfo(Owner);
 	}
@@ -1157,6 +1158,13 @@ static FString BuildDistanceFieldDerivedDataKey(const FString& InMeshKey)
 
 void FStaticMeshRenderData::Cache(UStaticMesh* Owner, const FStaticMeshLODSettings& LODSettings)
 {
+	if (Owner->GetOutermost()->PackageFlags & PKG_FilterEditorOnly)
+	{
+		// Don't chache for cooked packages
+		return;
+	}
+
+
 	int32 T0 = FPlatformTime::Cycles();
 	int32 NumLODs = Owner->SourceModels.Num();
 	const FStaticMeshLODGroup& LODGroup = LODSettings.GetLODGroup(Owner->LODGroup);
@@ -1698,6 +1706,13 @@ static FStaticMeshRenderData& GetPlatformStaticMeshRenderData(UStaticMesh* Mesh,
 	const FStaticMeshLODSettings& PlatformLODSettings = Platform->GetStaticMeshLODSettings();
 	FString PlatformDerivedDataKey = BuildStaticMeshDerivedDataKey(Mesh, PlatformLODSettings.GetLODGroup(Mesh->LODGroup));
 	FStaticMeshRenderData* PlatformRenderData = Mesh->RenderData;
+
+	if (Mesh->GetOutermost()->PackageFlags & PKG_FilterEditorOnly)
+	{
+		check(PlatformRenderData);
+		return *PlatformRenderData;
+	}
+
 	while (PlatformRenderData && PlatformRenderData->DerivedDataKey != PlatformDerivedDataKey)
 	{
 		PlatformRenderData = PlatformRenderData->NextCachedRenderData;
@@ -1932,66 +1947,69 @@ void UStaticMesh::PostLoad()
 	Super::PostLoad();
 
 #if WITH_EDITOR
-	// Needs to happen before 'CacheDerivedData'
-	if ( GetLinkerUE4Version() < VER_UE4_BUILD_SCALE_VECTOR )
+	if (!(GetOutermost()->PackageFlags & PKG_FilterEditorOnly))
 	{
-		int32 NumLODs = SourceModels.Num();
-		for (int32 LODIndex = 0; LODIndex < NumLODs; ++LODIndex)
+		// Needs to happen before 'CacheDerivedData'
+		if (GetLinkerUE4Version() < VER_UE4_BUILD_SCALE_VECTOR)
 		{
-			FStaticMeshSourceModel& SrcModel = SourceModels[LODIndex];
-			SrcModel.BuildSettings.BuildScale3D = FVector( SrcModel.BuildSettings.BuildScale_DEPRECATED );
+			int32 NumLODs = SourceModels.Num();
+			for (int32 LODIndex = 0; LODIndex < NumLODs; ++LODIndex)
+			{
+				FStaticMeshSourceModel& SrcModel = SourceModels[LODIndex];
+				SrcModel.BuildSettings.BuildScale3D = FVector(SrcModel.BuildSettings.BuildScale_DEPRECATED);
+			}
 		}
-	}
 
-	if( GetLinkerUE4Version() < VER_UE4_LIGHTMAP_MESH_BUILD_SETTINGS )
-	{
-		for( int32 i = 0; i < SourceModels.Num(); i++ )
+		if (GetLinkerUE4Version() < VER_UE4_LIGHTMAP_MESH_BUILD_SETTINGS)
 		{
-			SourceModels[i].BuildSettings.bGenerateLightmapUVs = false;
+			for (int32 i = 0; i < SourceModels.Num(); i++)
+			{
+				SourceModels[i].BuildSettings.bGenerateLightmapUVs = false;
+			}
 		}
-	}
 
-	if (GetLinkerUE4Version() < VER_UE4_MIKKTSPACE_IS_DEFAULT)
-	{
-		for (int32 i = 0; i < SourceModels.Num(); ++i)
+		if (GetLinkerUE4Version() < VER_UE4_MIKKTSPACE_IS_DEFAULT)
 		{
-			SourceModels[i].BuildSettings.bUseMikkTSpace = true;
+			for (int32 i = 0; i < SourceModels.Num(); ++i)
+			{
+				SourceModels[i].BuildSettings.bUseMikkTSpace = true;
+			}
 		}
-	}
-
-	if (GetLinkerUE4Version() < VER_UE4_BUILD_MESH_ADJ_BUFFER_FLAG_EXPOSED)
-	{
-		FRawMesh TempRawMesh;
-		uint32 TotalIndexCount = 0;
 
-		for (int32 i = 0; i < SourceModels.Num(); ++i)
+		if (GetLinkerUE4Version() < VER_UE4_BUILD_MESH_ADJ_BUFFER_FLAG_EXPOSED)
 		{
-			FRawMeshBulkData* RawMeshBulkData = SourceModels[i].RawMeshBulkData;
-			if (RawMeshBulkData)
+			FRawMesh TempRawMesh;
+			uint32 TotalIndexCount = 0;
+
+			for (int32 i = 0; i < SourceModels.Num(); ++i)
+			{
+				FRawMeshBulkData* RawMeshBulkData = SourceModels[i].RawMeshBulkData;
+				if (RawMeshBulkData)
+				{
+					RawMeshBulkData->LoadRawMesh(TempRawMesh);
+					TotalIndexCount += TempRawMesh.WedgeIndices.Num();
+				}
+			}
+
+			for (int32 i = 0; i < SourceModels.Num(); ++i)
 			{
-				RawMeshBulkData->LoadRawMesh(TempRawMesh);
-				TotalIndexCount += TempRawMesh.WedgeIndices.Num();
+				SourceModels[i].BuildSettings.bBuildAdjacencyBuffer = (TotalIndexCount < 50000);
 			}
 		}
 
-		for (int32 i = 0; i < SourceModels.Num(); ++i)
+		CacheDerivedData();
+		
+		// Only required in an editor build as other builds process this in a different place
+		if (bRequiresLODDistanceConversion)
 		{
-			SourceModels[i].BuildSettings.bBuildAdjacencyBuffer = (TotalIndexCount < 50000);
+			// Convert distances to Display Factors
+			ConvertLegacyLODDistance();
 		}
-	}
-
-	CacheDerivedData();
-
-	// Only required in an editor build as other builds process this in a different place
-	if(bRequiresLODDistanceConversion)
-	{
-		// Convert distances to Display Factors
-		ConvertLegacyLODDistance();
-	}
 
-	if(RenderData && GStaticMeshesThatNeedMaterialFixup.Get(this))
-	{
-		FixupZeroTriangleSections();
+		if (RenderData && GStaticMeshesThatNeedMaterialFixup.Get(this))
+		{
+			FixupZeroTriangleSections();
+		}
 	}
 #endif // #if WITH_EDITOR
 
diff --git a/Engine/Source/Runtime/Engine/Private/TextureDerivedData.cpp b/Engine/Source/Runtime/Engine/Private/TextureDerivedData.cpp
--- a/Engine/Source/Runtime/Engine/Private/TextureDerivedData.cpp
+++ b/Engine/Source/Runtime/Engine/Private/TextureDerivedData.cpp
@@ -1464,7 +1464,7 @@ void UTexture::BeginCachePlatformData()
 void UTexture::BeginCacheForCookedPlatformData( const ITargetPlatform *TargetPlatform )
 {
 	TMap<FString, FTexturePlatformData*>* CookedPlatformDataPtr = GetCookedPlatformData();
-	if (CookedPlatformDataPtr)
+	if (CookedPlatformDataPtr && !(GetOutermost()->PackageFlags & PKG_FilterEditorOnly))
 	{
 		TMap<FString,FTexturePlatformData*>& CookedPlatformData = *CookedPlatformDataPtr;
 		
@@ -1719,12 +1719,15 @@ void UTexture::FinishCachePlatformData()
 			}
 
 #if DO_CHECK
-			FString DerivedDataKey;
-			FTextureBuildSettings BuildSettings;
-			GetBuildSettingsForRunningPlatform(*this, BuildSettings);
-			GetTextureDerivedDataKey(*this, BuildSettings, DerivedDataKey);
+			if (!(GetOutermost()->PackageFlags & PKG_FilterEditorOnly))
+			{
+				FString DerivedDataKey;
+				FTextureBuildSettings BuildSettings;
+				GetBuildSettingsForRunningPlatform(*this, BuildSettings);
+				GetTextureDerivedDataKey(*this, BuildSettings, DerivedDataKey);
 
-			check( RunningPlatformData->DerivedDataKey == DerivedDataKey );
+				check(RunningPlatformData->DerivedDataKey == DerivedDataKey);
+			}
 #endif
 		}
 	}
@@ -1828,35 +1831,53 @@ void UTexture::SerializeCookedPlatformData(FArchive& Ar)
 	{
 		if (!Ar.CookingTarget()->IsServerOnly())
 		{
-			TMap<FString, FTexturePlatformData*> *CookedPlatformDataPtr = GetCookedPlatformData();
-			if ( CookedPlatformDataPtr == NULL )
-				return;
-
-
-
 			FTextureBuildSettings BuildSettings;
-			TArray<FName> PlatformFormats;
+			GetTextureBuildSettings(*this, Ar.CookingTarget()->GetTextureLODSettings(), BuildSettings);
+
 			TArray<FTexturePlatformData*> PlatformDataToSerialize;
 
-			GetTextureBuildSettings(*this, Ar.CookingTarget()->GetTextureLODSettings(), BuildSettings);
-			Ar.CookingTarget()->GetTextureFormats(this, PlatformFormats);
-			for (int32 FormatIndex = 0; FormatIndex < PlatformFormats.Num(); FormatIndex++)
+			if (GetOutermost()->bIsCookedForEditor)
 			{
-				FString DerivedDataKey;
-				BuildSettings.TextureFormatName = PlatformFormats[FormatIndex];
-				GetTextureDerivedDataKey(*this, BuildSettings, DerivedDataKey);
-
-				FTexturePlatformData *PlatformDataPtr = (*CookedPlatformDataPtr).FindRef(DerivedDataKey);
+				// For cooked packages, simply grab the current platform data and serialize it
+				FTexturePlatformData** RunningPlatformDataPtr = GetRunningPlatformData();
+				if (RunningPlatformDataPtr == nullptr)
+				{
+					return;
+				}
+				FTexturePlatformData* RunningPlatformData = *RunningPlatformDataPtr;
+				if (RunningPlatformData == nullptr)
+				{
+					return;
+				}
+				PlatformDataToSerialize.Add(RunningPlatformData);
+			}
+			else
+			{
+				TMap<FString, FTexturePlatformData*> *CookedPlatformDataPtr = GetCookedPlatformData();
+				if (CookedPlatformDataPtr == NULL)
+					return;
 				
-				if (PlatformDataPtr == NULL)
+				TArray<FName> PlatformFormats;
+
+				Ar.CookingTarget()->GetTextureFormats(this, PlatformFormats);
+				for (int32 FormatIndex = 0; FormatIndex < PlatformFormats.Num(); FormatIndex++)
 				{
-					PlatformDataPtr = new FTexturePlatformData();
-					PlatformDataPtr->Cache(*this, BuildSettings, ETextureCacheFlags::InlineMips | ETextureCacheFlags::Async);
-					
-					CookedPlatformDataPtr->Add( DerivedDataKey, PlatformDataPtr );
-					
+					FString DerivedDataKey;
+					BuildSettings.TextureFormatName = PlatformFormats[FormatIndex];
+					GetTextureDerivedDataKey(*this, BuildSettings, DerivedDataKey);
+
+					FTexturePlatformData *PlatformDataPtr = (*CookedPlatformDataPtr).FindRef(DerivedDataKey);
+
+					if (PlatformDataPtr == NULL)
+					{
+						PlatformDataPtr = new FTexturePlatformData();
+						PlatformDataPtr->Cache(*this, BuildSettings, ETextureCacheFlags::InlineMips | ETextureCacheFlags::Async);
+
+						CookedPlatformDataPtr->Add(DerivedDataKey, PlatformDataPtr);
+
+					}
+					PlatformDataToSerialize.Add(PlatformDataPtr);
 				}
-				PlatformDataToSerialize.Add(PlatformDataPtr);
 			}
 
 			for (int32 i = 0; i < PlatformDataToSerialize.Num(); ++i)
