#UE4 Bumping Asset Registry cache version after CL#2632031 changed AssetData's serialize function.

diff --git a/Engine/Source/Runtime/AssetRegistry/Private/AssetDataGatherer.cpp b/Engine/Source/Runtime/AssetRegistry/Private/AssetDataGatherer.cpp
--- a/Engine/Source/Runtime/AssetRegistry/Private/AssetDataGatherer.cpp
+++ b/Engine/Source/Runtime/AssetRegistry/Private/AssetDataGatherer.cpp
@@ -3,7 +3,7 @@
 #include "AssetRegistryPCH.h"
 
 #define MAX_FILES_TO_PROCESS_BEFORE_FLUSH 250
-#define CACHE_SERIALIZATION_VERSION 5
+#define CACHE_SERIALIZATION_VERSION 6
 
 FAssetDataGatherer::FAssetDataGatherer(const TArray<FString>& InPaths, bool bInIsSynchronous, bool bInLoadAndSaveCache)
 	: StopTaskCounter( 0 )
