Updated the asset registry gathering to use the directory iteration functions that provide file stat data

The discovery phase now uses the fact that the stat data for a file is available up-front to perform the check against the cached timestamp immediately. On Windows, and for the cases where most or all of the data is cached, this can significantly improve the asset registry gathering time as it avoids performing a second scan of the file-system for the majority of the files.

diff --git a/Engine/Source/Runtime/AssetRegistry/Private/AssetDataGatherer.cpp b/Engine/Source/Runtime/AssetRegistry/Private/AssetDataGatherer.cpp
--- a/Engine/Source/Runtime/AssetRegistry/Private/AssetDataGatherer.cpp
+++ b/Engine/Source/Runtime/AssetRegistry/Private/AssetDataGatherer.cpp
@@ -94,17 +94,10 @@ uint32 FAssetDataGatherer::Run()
 
 		{
 			FScopeLock CritSectionLock(&WorkerThreadCriticalSection);
-			if ( LocalAssetResults.Num() )
-			{
-				AssetResults.Append(LocalAssetResults);
-			}
-
-			if ( LocalDependencyResults.Num() )
-			{
-				DependencyResults.Append(LocalDependencyResults);
-			}
+			AssetResults.Append(MoveTemp(LocalAssetResults));
+			DependencyResults.Append(MoveTemp(LocalDependencyResults));
 
-			if ( FilesToSearch.Num() )
+			if (FilesToSearch.Num() > 0)
 			{
 				if (SearchStartTime == 0)
 				{
@@ -112,13 +105,8 @@ uint32 FAssetDataGatherer::Run()
 				}
 
 				const int32 NumFilesToProcess = FMath::Min<int32>(MAX_FILES_TO_PROCESS_BEFORE_FLUSH, FilesToSearch.Num());
-
-				for (int32 FileIdx = 0; FileIdx < NumFilesToProcess; ++FileIdx)
-				{
-					LocalFilesToSearch.Add(FilesToSearch[FileIdx]);
-				}
-
-				FilesToSearch.RemoveAt(0, NumFilesToProcess);
+				LocalFilesToSearch.Append(FilesToSearch.GetData(), NumFilesToProcess);
+				FilesToSearch.RemoveAt(0, NumFilesToProcess, false);
 			}
 			else if (SearchStartTime != 0)
 			{
@@ -127,99 +115,61 @@ uint32 FAssetDataGatherer::Run()
 			}
 		}
 
-		if ( LocalAssetResults.Num() )
-		{
-			LocalAssetResults.Empty();
-		}
+		LocalAssetResults.Reset();
+		LocalDependencyResults.Reset();
 
-		if ( LocalDependencyResults.Num() )
+		if (LocalFilesToSearch.Num() > 0)
 		{
-			LocalDependencyResults.Empty();
-		}
-
-		if ( LocalFilesToSearch.Num() )
-		{
-			for (int32 FileIdx = 0; FileIdx < LocalFilesToSearch.Num(); ++FileIdx)
+			for (const FString& AssetFilename : LocalFilesToSearch)
 			{
-				const FString& AssetFile = LocalFilesToSearch[FileIdx];
-
-				if ( StopTaskCounter.GetValue() != 0 )
+				if (StopTaskCounter.GetValue() != 0)
 				{
 					// We have been asked to stop, so don't read any more files
 					break;
 				}
 
-				bool bLoadedFromCache = false;
-				if ( bLoadAndSaveCache )
+				TArray<FBackgroundAssetData*> AssetDataFromFile;
+				FPackageDependencyData DependencyData;
+				if (ReadAssetFile(AssetFilename, AssetDataFromFile, DependencyData))
 				{
-					const FName PackageName = FName(*FPackageName::FilenameToLongPackageName(AssetFile));
-					FDiskCachedAssetData** DiskCachedAssetDataPtr = DiskCachedAssetDataMap.Find(PackageName);
-					FDiskCachedAssetData* DiskCachedAssetData = NULL;
-					if ( DiskCachedAssetDataPtr && *DiskCachedAssetDataPtr )
-					{
-						const FDateTime& FileTimestamp = IFileManager::Get().GetTimeStamp(*AssetFile);
-						const FDateTime& CachedTimestamp = (*DiskCachedAssetDataPtr)->Timestamp;
-						if ( FileTimestamp == CachedTimestamp )
-						{
-							DiskCachedAssetData = *DiskCachedAssetDataPtr;
-						}
-					}
+					LocalAssetResults.Append(AssetDataFromFile);
+					LocalDependencyResults.Add(DependencyData);
 
-					if ( DiskCachedAssetData )
+					bool bCachePackage = bLoadAndSaveCache;
+					if (bCachePackage)
 					{
-						for ( auto CacheIt = DiskCachedAssetData->AssetDataList.CreateConstIterator(); CacheIt; ++CacheIt )
+						// Don't store info on cooked packages
+						for (const auto& AssetData : AssetDataFromFile)
 						{
-							LocalAssetResults.Add(new FAssetDataWrapper(*CacheIt));
+							if (AssetData->IsCooked())
+							{
+								bCachePackage = false;
+								break;
+							}
 						}
-
-						LocalDependencyResults.Add(DiskCachedAssetData->DependencyData);
-
-						NewCachedAssetDataMap.Add(PackageName, DiskCachedAssetData);
-						bLoadedFromCache = true;
 					}
-				}
 
-				if ( !bLoadedFromCache )
-				{
-					TArray<FBackgroundAssetData*> AssetDataFromFile;
-					FPackageDependencyData DependencyData;
-					if ( ReadAssetFile(AssetFile, AssetDataFromFile, DependencyData) )
+					if (bCachePackage)
 					{
-						LocalAssetResults.Append(AssetDataFromFile);
-						LocalDependencyResults.Add(DependencyData);
+						// Update the cache
+						const FName PackageName = FName(*FPackageName::FilenameToLongPackageName(AssetFilename));
+						const FDateTime FileTimestamp = IFileManager::Get().GetTimeStamp(*AssetFilename);
 
-						if ( bLoadAndSaveCache )
+						FDiskCachedAssetData* NewData = new FDiskCachedAssetData(PackageName, FileTimestamp);
+						NewData->AssetDataList.Reserve(AssetDataFromFile.Num());
+						for (const FBackgroundAssetData* BackgroundAssetData : AssetDataFromFile)
 						{
-							// don't store info on cooked packages
-							bool bIsCooked = false;
-							for (const auto& AssetData : AssetDataFromFile)
-							{
-								bIsCooked |= AssetData->IsCooked();
-								if (bIsCooked)
-								{
-									break;
-								}
-							}
-							if (bIsCooked == false)
-							{
-								// Update the cache
-								const FName PackageName = FName(*FPackageName::FilenameToLongPackageName(AssetFile));
-								const FDateTime& FileTimestamp = IFileManager::Get().GetTimeStamp(*AssetFile);
-								FDiskCachedAssetData* NewData = new FDiskCachedAssetData(PackageName, FileTimestamp);
-								for (auto AssetIt = AssetDataFromFile.CreateConstIterator(); AssetIt; ++AssetIt)
-								{
-									NewData->AssetDataList.Add((*AssetIt)->ToAssetData());
-								}
-								NewData->DependencyData = DependencyData;
-								NewCachedAssetData.Add(NewData);
-								NewCachedAssetDataMap.Add(PackageName, NewData);
-							}
+							NewData->AssetDataList.Add(BackgroundAssetData->ToAssetData());
 						}
+						NewData->DependencyData = DependencyData;
+
+						NewCachedAssetData.Add(NewData);
+						NewCachedAssetDataMap.Add(PackageName, NewData);
 					}
 				}
 			}
 
-			LocalFilesToSearch.Empty();
+			LocalFilesToSearch.Reset();
 		}
 		else
 		{
@@ -284,17 +234,17 @@ bool FAssetDataGatherer::GetAndTrimSearchResults(TArray<IGatheredAssetData*>& Ou
 {
 	FScopeLock CritSectionLock(&WorkerThreadCriticalSection);
 
-	OutAssetResults.Append(AssetResults);
-	AssetResults.Empty();
+	OutAssetResults.Append(MoveTemp(AssetResults));
+	AssetResults.Reset();
 
-	OutPathResults.Append(DiscoveredPaths);
-	DiscoveredPaths.Empty();
+	OutPathResults.Append(MoveTemp(DiscoveredPaths));
+	DiscoveredPaths.Reset();
 
-	OutDependencyResults.Append(DependencyResults);
-	DependencyResults.Empty();
+	OutDependencyResults.Append(MoveTemp(DependencyResults));
+	DependencyResults.Reset();
 
-	OutSearchTimes.Append(SearchTimes);
-	SearchTimes.Empty();
+	OutSearchTimes.Append(MoveTemp(SearchTimes));
+	SearchTimes.Reset();
 
 	OutNumFilesToSearch = FilesToSearch.Num();
 	OutNumPathsToSearch = PathsToSearch.Num();
@@ -311,9 +261,8 @@ void FAssetDataGatherer::AddPathToSearch(const FString& Path)
 void FAssetDataGatherer::AddFilesToSearch(const TArray<FString>& Files)
 {
 	TArray<FString> FilesToAdd;
-	for (int32 FilenameIdx = 0; FilenameIdx < Files.Num(); FilenameIdx++)
+	for (const FString& Filename : Files)
 	{
-		FString Filename(Files[FilenameIdx]);
 		if ( IsValidPackageFileToRead(Filename) )
 		{
 			// Add the path to this asset into the list of discovered paths
@@ -353,56 +302,108 @@ void FAssetDataGatherer::PrioritizeSearchPath(const FString& PathToPrioritize)
 
 void FAssetDataGatherer::DiscoverFilesToSearch()
 {
-	if( PathsToSearch.Num() > 0 )
+	const double DiscoverStartTime = FPlatformTime::Seconds();
+
+	TArray<FString> LocalFilesToSearch;
+	TSet<FString> LocalDiscoveredPathsSet;
+
+	int32 NumCachedPackages = 0;
+	TArray<IGatheredAssetData*> LocalAssetResults;
+	TArray<FPackageDependencyData> LocalDependencyResults;
+
+	auto OnPackageDiscovered = [&](const TCHAR* InPackageFilename, const FFileStatData& InPackageStatData) -> bool
 	{
-		TArray<FString> DiscoveredFilesToSearch;
-		TSet<FString> LocalDiscoveredPathsSet;
+		const FString PackageFilenameStr = InPackageFilename;
+		if (IsValidPackageFileToRead(PackageFilenameStr))
+		{
+			const FString LongPackageNameStr = FPackageName::FilenameToLongPackageName(PackageFilenameStr);
+
+			bool bLoadedFromCache = false;
+			if (bLoadAndSaveCache)
+			{
+				const FName PackageName = FName(*LongPackageNameStr);
+
+				FDiskCachedAssetData** DiskCachedAssetDataPtr = DiskCachedAssetDataMap.Find(PackageName);
+				FDiskCachedAssetData* DiskCachedAssetData = nullptr;
+				if (DiskCachedAssetDataPtr && *DiskCachedAssetDataPtr)
+				{
+					const FDateTime& CachedTimestamp = (*DiskCachedAssetDataPtr)->Timestamp;
+					if (InPackageStatData.ModificationTime == CachedTimestamp)
+					{
+						DiskCachedAssetData = *DiskCachedAssetDataPtr;
+					}
+				}
+
+				if (DiskCachedAssetData)
+				{
+					bLoadedFromCache = true;
+
+					++NumCachedPackages;
+
+					LocalAssetResults.Reserve(LocalAssetResults.Num() + DiskCachedAssetData->AssetDataList.Num());
+					for (const FAssetData& AssetData : DiskCachedAssetData->AssetDataList)
+					{
+						LocalAssetResults.Add(new FAssetDataWrapper(AssetData));
+					}
+
+					LocalDependencyResults.Add(DiskCachedAssetData->DependencyData);
+
+					NewCachedAssetDataMap.Add(PackageName, DiskCachedAssetData);
+				}
+			}
+
+			if (!bLoadedFromCache)
+			{
+				LocalFilesToSearch.Add(PackageFilenameStr);
+			}
 
-		TArray<FString> CopyOfPathsToSearch;
+			LocalDiscoveredPathsSet.Add(FPackageName::GetLongPackagePath(LongPackageNameStr));
+		}
+
+		return true;
+	};
+
+	if (PathsToSearch.Num() > 0)
+	{
+		TArray<FString> LocalPathsToSearch;
 		{
 			FScopeLock CritSectionLock(&WorkerThreadCriticalSection);
-			CopyOfPathsToSearch = PathsToSearch;
+
+			LocalPathsToSearch = PathsToSearch;
 
 			// Remove all of the existing paths from the list, since we'll process them all below.  New paths may be
 			// added to the original list on a different thread as we go along, but those new paths won't be processed
 			// during this call of DisoverFilesToSearch().  But we'll get them on the next call!
-			PathsToSearch.Empty();
+			PathsToSearch.Reset();
 			bIsDiscoveringFiles = true;
 		}
 
 		// Iterate over any paths that we have remaining to scan
-		for ( int32 PathIdx=0; PathIdx < CopyOfPathsToSearch.Num(); ++PathIdx )
+		for (const FString& PathToSearch : LocalPathsToSearch)
 		{
-			const FString& Path = CopyOfPathsToSearch[PathIdx];
-
 			// Convert the package path to a filename with no extension (directory)
-			const FString FilePath = FPackageName::LongPackageNameToFilename(Path);
-
-			// Gather the package files in that directory and subdirectories
-			TArray<FString> Filenames;
-			FPackageName::FindPackagesInDirectory(Filenames, FilePath);
+			const FString FilePath = FPackageName::LongPackageNameToFilename(PathToSearch);
 
-			for (int32 FilenameIdx = 0; FilenameIdx < Filenames.Num(); FilenameIdx++)
-			{
-				FString Filename(Filenames[FilenameIdx]);
-				if ( IsValidPackageFileToRead(Filename) )
-				{
-					// Add the path to this asset into the list of discovered paths
-					const FString LongPackageName = FPackageName::FilenameToLongPackageName(Filename);
-					LocalDiscoveredPathsSet.Add( FPackageName::GetLongPackagePath(LongPackageName) );
-					DiscoveredFilesToSearch.Add(Filename);
-				}
-			}
+			// Discover any packages that we need to scan
+			// This will also populate our local arrays with any valid cached data
+			FPackageName::IteratePackagesInDirectory(FilePath, FPackageName::FPackageNameStatVisitor(OnPackageDiscovered));
 		}
 
-		// Turn the set into an array here before the critical section below
-		TArray<FString> LocalDiscoveredPathsArray = LocalDiscoveredPathsSet.Array();
+		UE_LOG(LogAssetRegistry, Verbose, TEXT("Discovery took %0.6f seconds (found %d cached packages, with %d to re-scan)"), FPlatformTime::Seconds() - DiscoverStartTime, NumCachedPackages, LocalFilesToSearch.Num());
 
+		// Move all our discovered files and data into the final lists
+		TArray<FString> LocalDiscoveredPathsArray = LocalDiscoveredPathsSet.Array();
 		{
-			// Place all the discovered files into the files to search list
 			FScopeLock CritSectionLock(&WorkerThreadCriticalSection);
-			FilesToSearch.Append(DiscoveredFilesToSearch);
-			DiscoveredPaths.Append(LocalDiscoveredPathsArray);
+
+			// Place all the discovered files into the files to search list
+			FilesToSearch.Append(MoveTemp(LocalFilesToSearch));
+			DiscoveredPaths.Append(MoveTemp(LocalDiscoveredPathsArray));
+
+			// Place all the discovered cached data into the final lists for consumption
+			AssetResults.Append(MoveTemp(LocalAssetResults));
+			DependencyResults.Append(MoveTemp(LocalDependencyResults));
+
 			bIsDiscoveringFiles = false;
 		}
 	}
diff --git a/Engine/Source/Runtime/CoreUObject/Private/Misc/PackageName.cpp b/Engine/Source/Runtime/CoreUObject/Private/Misc/PackageName.cpp
--- a/Engine/Source/Runtime/CoreUObject/Private/Misc/PackageName.cpp
+++ b/Engine/Source/Runtime/CoreUObject/Private/Misc/PackageName.cpp
@@ -853,6 +853,55 @@ bool FPackageName::FindPackagesInDirectory( TArray<FString>& OutPackages, const
 	return OutPackages.Num() > PreviousPackagesCount;
 }
 
+void FPackageName::IteratePackagesInDirectory(const FString& RootDir, const FPackageNameVisitor& Callback)
+{
+	class FPackageVisitor : public IPlatformFile::FDirectoryVisitor
+	{
+	public:
+		const FPackageNameVisitor& Callback;
+		explicit FPackageVisitor(const FPackageNameVisitor& InCallback)
+			: Callback(InCallback)
+		{
+		}
+		virtual bool Visit(const TCHAR* FilenameOrDirectory, bool bIsDirectory) override
+		{
+			bool Result = true;
+			if (!bIsDirectory && IsPackageFilename(FilenameOrDirectory))
+			{
+				Result = Callback(FilenameOrDirectory);
+			}
+			return Result;
+		}
+	};
+
+	FPackageVisitor PackageVisitor(Callback);
+	IFileManager::Get().IterateDirectoryRecursively(*RootDir, PackageVisitor);
+}
+
+void FPackageName::IteratePackagesInDirectory(const FString& RootDir, const FPackageNameStatVisitor& Callback)
+{
+	class FPackageStatVisitor : public IPlatformFile::FDirectoryStatVisitor
+	{
+	public:
+		const FPackageNameStatVisitor& Callback;
+		explicit FPackageStatVisitor(const FPackageNameStatVisitor& InCallback)
+			: Callback(InCallback)
+		{
+		}
+		virtual bool Visit(const TCHAR* FilenameOrDirectory, const FFileStatData& StatData) override
+		{
+			bool Result = true;
+			if (!StatData.bIsDirectory && IsPackageFilename(FilenameOrDirectory))
+			{
+				Result = Callback(FilenameOrDirectory, StatData);
+			}
+			return Result;
+		}
+	};
+
+	FPackageStatVisitor PackageVisitor(Callback);
+	IFileManager::Get().IterateDirectoryStatRecursively(*RootDir, PackageVisitor);
+}
 
 void FPackageName::QueryRootContentPaths( TArray<FString>& OutRootContentPaths )
 {
diff --git a/Engine/Source/Runtime/CoreUObject/Public/Misc/PackageName.h b/Engine/Source/Runtime/CoreUObject/Public/Misc/PackageName.h
--- a/Engine/Source/Runtime/CoreUObject/Public/Misc/PackageName.h
+++ b/Engine/Source/Runtime/CoreUObject/Public/Misc/PackageName.h
@@ -273,6 +273,17 @@ public:
 	 */
 	static bool FindPackagesInDirectory(TArray<FString>& OutPackages, const FString& RootDir);
 
+	/**
+	 * This will recurse over a directory structure looking for packages.
+	 * 
+	 * @param	RootDirectory		The root of the directory structure to recurse through
+	 * @param	Visitor				Visitor to call for each package file found (takes the package filename, and optionally the stat data for the file - returns true to continue iterating)
+	 */
+	typedef TFunctionRef<bool(const TCHAR*)> FPackageNameVisitor;
+	typedef TFunctionRef<bool(const TCHAR*, const FFileStatData&)> FPackageNameStatVisitor;
+	static void IteratePackagesInDirectory(const FString& RootDir, const FPackageNameVisitor& Visitor);
+	static void IteratePackagesInDirectory(const FString& RootDir, const FPackageNameStatVisitor& Visitor);
+
 	/** Event that is triggered when a new content path is mounted */
 	DECLARE_MULTICAST_DELEGATE_TwoParams( FOnContentPathMountedEvent, const FString& /* Asset path */, const FString& /* ContentPath */ );
 	static FOnContentPathMountedEvent& OnContentPathMounted()
