Fixed asset registry loading code after asset reigstry format change.

diff --git a/Engine/Source/Editor/UnrealEd/Private/CookOnTheFlyServer.cpp b/Engine/Source/Editor/UnrealEd/Private/CookOnTheFlyServer.cpp
--- a/Engine/Source/Editor/UnrealEd/Private/CookOnTheFlyServer.cpp
+++ b/Engine/Source/Editor/UnrealEd/Private/CookOnTheFlyServer.cpp
@@ -18,7 +18,11 @@
 #include "IPlatformFileSandboxWrapper.h"
 #include "Messaging.h"
 #include "NetworkFileSystem.h"
+
 #include "AssetRegistryModule.h"
+#include "AssetData.h"
+#include "DependsNode.h"
+
 #include "UnrealEdMessages.h"
 #include "GameDelegates.h"
 #include "PhysicsPublic.h"
@@ -4701,9 +4705,56 @@ void UCookOnTheFlyServer::HandleNetworkFileServerRecompileShaders(const FShaderR
 
 bool UCookOnTheFlyServer::GetAllPackagesFromAssetRegistry( const FString& AssetRegistryPath, TArray<FName>& OutPackageNames ) const
 {
+
+	static const FName AssetRegistryName("AssetRegistry");
+	FAssetRegistryModule& AssetRegistryModule = FModuleManager::LoadModuleChecked<FAssetRegistryModule>(AssetRegistryName);
+	IAssetRegistry& AssetRegistry = AssetRegistryModule.Get();
+
+
 	FArrayReader SerializedAssetData;
 	if (FFileHelper::LoadFileToArray(SerializedAssetData, *AssetRegistryPath))
 	{
+		TMap<FName, FAssetData*> RegistryDataMap;
+		TArray<FDependsNode*> Dependencies;
+		AssetRegistry.LoadRegistryData(SerializedAssetData, RegistryDataMap, Dependencies);
+
+
+		for (const auto& RegistryData : RegistryDataMap)
+		{
+			const auto& NewAssetData = RegistryData.Value;
+			FName CachedPackageFileFName = GetCachedStandardPackageFileFName(NewAssetData->ObjectPath);
+			if (CachedPackageFileFName != NAME_None)
+			{
+				OutPackageNames.Add(CachedPackageFileFName);
+			}
+			else
+			{
+				UE_LOG(LogCook, Warning, TEXT("Could not resolve package %s from %s"), *NewAssetData->ObjectPath.ToString(), *AssetRegistryPath);
+			}
+		}
+		
+
+		// clean up memory
+		for (auto RegistryData : RegistryDataMap)
+		{
+			delete RegistryData.Value;
+		}
+		for (auto DependencyData : Dependencies)
+		{
+			delete DependencyData;
+		}
+		return true;
+	}
+
+	return false;
+	/*
+
+
+	FArrayReader SerializedAssetData;
+	if (FFileHelper::LoadFileToArray(SerializedAssetData, *AssetRegistryPath))
+	{
+
+
 		int32 LocalNumAssets = 0;
 		SerializedAssetData << LocalNumAssets;
 
@@ -4735,7 +4786,8 @@ bool UCookOnTheFlyServer::GetAllPackagesFromAssetRegistry( const FString& AssetR
 
 		return true;
 	}
-	return false;
+	return false;*/
+
 }
 
 #undef LOCTEXT_NAMESPACE
diff --git a/Engine/Source/Runtime/AssetRegistry/Private/DependsNode.h b/Engine/Source/Runtime/AssetRegistry/Private/DependsNode.h
deleted file mode 100644
diff --git a/Engine/Source/Runtime/AssetRegistry/Public/DependsNode.h b/Engine/Source/Runtime/AssetRegistry/Public/DependsNode.h
new file mode 100644
--- /dev/null
+++ b/Engine/Source/Runtime/AssetRegistry/Public/DependsNode.h
@@ -0,0 +1,124 @@
+// Copyright 1998-2015 Epic Games, Inc. All Rights Reserved.
+
+#pragma once
+
+/** Implementation of IDependsNode */
+class FDependsNode
+{
+public:
+	FDependsNode() { PackageName = NAME_None; }
+	FDependsNode(FName InPackageName);
+
+	/** Prints the dependencies and referencers for this node to the log */
+	void PrintNode() const;
+	/** Prints the dependencies for this node to the log */
+	void PrintDependencies() const;
+	/** Prints the referencers to this node to the log */
+	void PrintReferencers() const;
+	/** Gets the list of dependencies for this node */
+	void GetDependencies(TArray<FDependsNode*>& OutDependencies, EAssetRegistryDependencyType::Type InDependencyType = EAssetRegistryDependencyType::All) const;
+	/** Gets the list of dependency names for this node */
+	void GetDependencies(TArray<FName>& OutDependencies, EAssetRegistryDependencyType::Type InDependencyType = EAssetRegistryDependencyType::All) const;
+	/** Gets the list of referencers to this node */
+	void GetReferencers(TArray<FDependsNode*>& OutReferencers) const;
+	/** Gets the name of the package that this node represents */
+	FName GetPackageName() const { return PackageName; }
+	/** Sets the name of the package that this node represents */
+	void SetPackageName(FName InName) { PackageName = InName; }
+	/** Add a dependency to this node */
+	void AddDependency(FDependsNode* InDependency, EAssetRegistryDependencyType::Type InDependencyType = EAssetRegistryDependencyType::Hard)
+	{ 
+		if (InDependencyType == EAssetRegistryDependencyType::Hard)
+		{
+			HardDependencies.Add(InDependency);
+		}
+		else
+		{
+			SoftDependencies.Add(InDependency);
+		}
+	}
+	/** Add a referencer to this node */
+	void AddReferencer(FDependsNode* InReferencer) { Referencers.Add(InReferencer); }
+	/** Remove a dependency from this node */
+	void RemoveDependency(FDependsNode* InDependency) { HardDependencies.Remove(InDependency); SoftDependencies.Remove(InDependency); }
+	/** Remove a referencer from this node */
+	void RemoveReferencer(FDependsNode* InReferencer) { Referencers.Remove(InReferencer); }
+	/** Clear all dependency records from this node */
+	void ClearDependencies() { HardDependencies.Empty(); SoftDependencies.Empty(); }
+
+	/** Iterate over all the dependencies of this node, filtered by the supplied type parameter, and call the supplied lambda
+		parameter on the record */
+	template <class T>
+	void IterateOverDependencies(T InCallback, EAssetRegistryDependencyType::Type InDependencyType = EAssetRegistryDependencyType::All)
+	{
+		if (InDependencyType & EAssetRegistryDependencyType::Hard)
+		{
+			for (auto Dependency : HardDependencies)
+			{
+				if (Dependency)
+				{
+					InCallback(Dependency, EAssetRegistryDependencyType::Hard);
+				}
+			}
+		}
+
+		if (InDependencyType & EAssetRegistryDependencyType::Soft)
+		{
+			for (auto Dependency : SoftDependencies)
+			{
+				if (Dependency)
+				{
+					InCallback(Dependency, EAssetRegistryDependencyType::Soft);
+				}
+			}
+		}
+	}
+
+	/** Iterate over all the dependencies of this node, filtered by the supplied type parameter, and call the supplied lambda
+	parameter on the record */
+	template <class T>
+	void IterateOverDependencies(T InCallback, EAssetRegistryDependencyType::Type InDependencyType = EAssetRegistryDependencyType::All) const
+	{
+		if (InDependencyType & EAssetRegistryDependencyType::Hard)
+		{
+			for (auto Dependency : HardDependencies)
+			{
+				InCallback(Dependency, EAssetRegistryDependencyType::Hard);
+			}
+		}
+
+		if (InDependencyType & EAssetRegistryDependencyType::Soft)
+		{
+			for (auto Dependency : SoftDependencies)
+			{
+				InCallback(Dependency, EAssetRegistryDependencyType::Soft);
+			}
+		}
+	}
+
+	/** Iterate over all the referencers of this node and call the supplied lambda
+		parameter on the record */
+	template <class T>
+	void IterateOverReferencers(T InCallback)
+	{
+		for (auto Referencer : Referencers)
+		{
+			InCallback(Referencer);
+		}
+	}
+
+private:
+	/** Recursively prints dependencies of the node starting with the specified indent. VisitedNodes should be an empty set at first which is populated recursively. */
+	void PrintDependenciesRecursive(const FString& Indent, TSet<const FDependsNode*>& VisitedNodes) const;
+	/** Recursively prints referencers to the node starting with the specified indent. VisitedNodes should be an empty set at first which is populated recursively. */
+	void PrintReferencersRecursive(const FString& Indent, TSet<const FDependsNode*>& VisitedNodes) const;
+
+	/** The name of the package this node represents */
+	FName PackageName;
+	/** The list of hard dependencies for this node */
+	TSet<FDependsNode*> HardDependencies;
+	/** The list of soft dependencies for this node */
+	TSet<FDependsNode*> SoftDependencies;
+	/** The list of referencers to this node */
+	TSet<FDependsNode*> Referencers;
+};
\ No newline at end of file
