Fixed FAssetDataGatherer::PrioritizeSearchPath not working with the root paths

FPackageName::IsValidLongPackageName would fail for the root paths as the mount path roots all end with a trailing slash, however FPackageName::IsValidLongPackageName doesn't accept paths that end in a trailing slash (as it's expected to be working with names, not paths).

This change updates FAssetDataGatherer::PrioritizeSearchPath to use FPackageName::TryConvertLongPackageNameToFilename instead (which will gracefully fail on an invalid path), and SPathView makes sure that it appends a trailing slash to the path it gives to FAssetDataGatherer::PrioritizeSearchPath.

diff --git a/Engine/Source/Editor/ContentBrowser/Private/SPathView.cpp b/Engine/Source/Editor/ContentBrowser/Private/SPathView.cpp
--- a/Engine/Source/Editor/ContentBrowser/Private/SPathView.cpp
+++ b/Engine/Source/Editor/ContentBrowser/Private/SPathView.cpp
@@ -853,7 +853,7 @@ void SPathView::TreeSelectionChanged( TSharedPtr< FTreeItem > TreeItem, ESelectI
 	{
 		// Prioritize the asset registry scan for the selected path
 		FAssetRegistryModule& AssetRegistryModule = FModuleManager::LoadModuleChecked<FAssetRegistryModule>(TEXT("AssetRegistry"));
-		AssetRegistryModule.Get().PrioritizeSearchPath(TreeItem->FolderPath);
+		AssetRegistryModule.Get().PrioritizeSearchPath(TreeItem->FolderPath / TEXT(""));
 	}
 }
 
diff --git a/Engine/Source/Runtime/AssetRegistry/Private/AssetDataGatherer.cpp b/Engine/Source/Runtime/AssetRegistry/Private/AssetDataGatherer.cpp
--- a/Engine/Source/Runtime/AssetRegistry/Private/AssetDataGatherer.cpp
+++ b/Engine/Source/Runtime/AssetRegistry/Private/AssetDataGatherer.cpp
@@ -278,11 +278,9 @@ void FAssetDataGatherer::AddFilesToSearch(const TArray<FString>& Files)
 
 void FAssetDataGatherer::PrioritizeSearchPath(const FString& PathToPrioritize)
 {
-	const bool bIncludeReadOnlyRoots = true;
-	if ( FPackageName::IsValidLongPackageName(PathToPrioritize, bIncludeReadOnlyRoots) )
+	FString FilenamePathToPrioritize;
+	if ( FPackageName::TryConvertLongPackageNameToFilename(PathToPrioritize, FilenamePathToPrioritize) )
 	{
-		const FString FilenamePathToPrioritize = FPackageName::LongPackageNameToFilename(PathToPrioritize);
-
 		// Critical section. This code needs to be as fast as possible since it is in a critical section!
 		// Swap all priority files to the top of the list
 		{
