Improved the asset gathering and path prioritization workflow for very large projects

Previously, the file discovery phase of the asset gathering process would block the entire process until it had found all of the files it needed to test.

This can take quite some time for projects with a large amount of assets, so this change adjusts the process so that discovered paths will be reported back to the asset registry while the file discovery is on-going. This allows you to prioritize the scanning of certain folders before the initial discovery phase has finished.

To account for the fact that the asset registry will be reporting process of an unknown range, the asset discovery indicator has been updated to display a marquee bar until the initial file discovery has finished.

It will now show you the following messages:
    - "Initializing Asset Discovery" - Loading the cached asset registry data and waiting for the file scan to start. The progress bar will marquee while this is happening.
    - "Discovering Asset Files" - Shows an increasing count of the number of asset files that have been discovered on disk. The progress bar will marquee while this is happening.
    - "Discovering Asset Data" - Shows a decreasing count of assets that need to have their asset data either read from disk, or copied from the cache file. The progress bar will show the percentage of the assets processed by the asset registry compared to the total number of assets to load and process.
    - "Processing Asset Data" - Shows a decreasing count of assets that the asset registry still needs to process the gathered asset data for. The progress bar will show the percentage of the assets processed by the asset registry compared to the total number of assets to load and process.

One side effect of having the discovery phase scan go wide (rather than deep - to allow you to prioritize the scanning of certain folders) is that it may now report back folders that don't contain any assets. This essentially addresses a long time complaint from users, whereby the Content Browser wouldn't show empty folders.

This change also has the cached asset registry file write out for every 50,000 files that are read from disk. This avoids you needing to perform repeated work if you restart the editor before the entire scan has finished.

diff --git a/Engine/Source/Editor/EditorStyle/Private/SlateEditorStyle.cpp b/Engine/Source/Editor/EditorStyle/Private/SlateEditorStyle.cpp
--- a/Engine/Source/Editor/EditorStyle/Private/SlateEditorStyle.cpp
+++ b/Engine/Source/Editor/EditorStyle/Private/SlateEditorStyle.cpp
@@ -5548,7 +5548,6 @@ void FSlateEditorStyle::FStyle::SetupContentBrowserStyle()
 
 		// Misc
 		Set( "ContentBrowser.ThumbnailShadow", new BOX_BRUSH( "ContentBrowser/ThumbnailShadow" , FMargin( 4.0f / 64.0f ) ) );
-		Set( "AssetDiscoveryIndicator.DiscovertingAssetsFont", TTF_CORE_FONT( "Fonts/Roboto-Regular", 12 ) );
 		Set( "ContentBrowser.ColumnViewAssetIcon", new IMAGE_BRUSH( "Icons/doc_16x", Icon16x16 ) );
 		Set( "ContentBrowser.ColumnViewFolderIcon", new IMAGE_BRUSH( "Icons/FolderClosed", FVector2D(18, 16) ) );
 		Set( "ContentBrowser.ColumnViewDeveloperFolderIcon", new IMAGE_BRUSH( "Icons/FolderDeveloper", FVector2D(18, 16) ) );
@@ -5601,6 +5600,9 @@ void FSlateEditorStyle::FStyle::SetupContentBrowserStyle()
 		Set( "ContentBrowser.Private.Small", new IMAGE_BRUSH( "ContentBrowser/Content_Private_16x", Icon16x16 ) );
 		Set( "ContentBrowser.Private.Large", new IMAGE_BRUSH( "ContentBrowser/Content_Private_64x", Icon64x64 ) );
 		Set( "ContentBrowser.CollectionStatus", new IMAGE_BRUSH( "/Icons/CollectionStatus_8x", Icon8x8 ) );
+
+		Set( "AssetDiscoveryIndicator.MainStatusFont", TTF_CORE_FONT( "Fonts/Roboto-Regular", 12 ) );
+		Set( "AssetDiscoveryIndicator.SubStatusFont", TTF_CORE_FONT( "Fonts/Roboto-Regular", 9 ) );
 	}
 #endif // #if WITH_EDITOR
 }
diff --git a/Engine/Source/Editor/EditorWidgets/Private/SAssetDiscoveryIndicator.cpp b/Engine/Source/Editor/EditorWidgets/Private/SAssetDiscoveryIndicator.cpp
--- a/Engine/Source/Editor/EditorWidgets/Private/SAssetDiscoveryIndicator.cpp
+++ b/Engine/Source/Editor/EditorWidgets/Private/SAssetDiscoveryIndicator.cpp
@@ -4,19 +4,29 @@
 #include "AssetRegistryModule.h"
 #include "SlateBasics.h"
 
+#define LOCTEXT_NAMESPACE "AssetDiscoveryIndicator"
+
+namespace AssetDiscoveryIndicatorConstants
+{
+	static const FMargin Padding(12, 4);
+	static const FMargin SubStatusTextPadding(6, 2, 6, 0);
+}
+
 SAssetDiscoveryIndicator::~SAssetDiscoveryIndicator()
 {
 	if ( FModuleManager::Get().IsModuleLoaded(TEXT("AssetRegistry")) )
 	{
 		FAssetRegistryModule& AssetRegistryModule = FModuleManager::GetModuleChecked<FAssetRegistryModule>(TEXT("AssetRegistry"));
 		AssetRegistryModule.Get().OnFileLoadProgressUpdated().RemoveAll( this );
+		AssetRegistryModule.Get().OnFilesLoaded().RemoveAll( this );
 	}
 }
 
 void SAssetDiscoveryIndicator::Construct( const FArguments& InArgs )
 {
 	FAssetRegistryModule& AssetRegistryModule = FModuleManager::LoadModuleChecked<FAssetRegistryModule>(TEXT("AssetRegistry"));
-	AssetRegistryModule.Get().OnFileLoadProgressUpdated().AddRaw( this, &SAssetDiscoveryIndicator::OnAssetRegistryFileLoadProgress );
+	AssetRegistryModule.Get().OnFileLoadProgressUpdated().AddSP( this, &SAssetDiscoveryIndicator::OnAssetRegistryFileLoadProgress );
+	AssetRegistryModule.Get().OnFilesLoaded().AddSP( this, &SAssetDiscoveryIndicator::OnAssetRegistryFilesLoaded );
 
 	ScaleMode = InArgs._ScaleMode;
 
@@ -26,10 +36,13 @@ void SAssetDiscoveryIndicator::Construct( const FArguments& InArgs )
 	FadeCurve = FadeAnimation.AddCurve(4.75f, 0.75f);
 	FadeAnimation.AddCurve(5.5f, 1.f); // Add some space at the end to cause a short delay before fading out
 
+	MainStatusText = LOCTEXT("InitializingAssetDiscovery", "Initializing Asset Discovery");
+	StatusTextWrapWidth = 0.0f;
+
 	if ( AssetRegistryModule.Get().IsLoadingAssets() )
 	{
-		// Loading assets, show the progress
-		Progress = 0.f;
+		// Loading assets, marquee while discovering package files
+		Progress = TOptional<float>();
 
 		if ( InArgs._FadeIn )
 		{
@@ -54,7 +67,7 @@ void SAssetDiscoveryIndicator::Construct( const FArguments& InArgs )
 		.HAlign(HAlign_Center)
 		[
 			SNew(SBorder)
-			.BorderImage( FEditorStyle::GetBrush("Menu.Background") )
+			.BorderImage(FEditorStyle::GetBrush("Menu.Background"))
 			.BorderBackgroundColor(this, &SAssetDiscoveryIndicator::GetBorderBackgroundColor)
 			.ColorAndOpacity(this, &SAssetDiscoveryIndicator::GetIndicatorColorAndOpacity)
 			.DesiredSizeScale(this, &SAssetDiscoveryIndicator::GetIndicatorDesiredSizeScale)
@@ -66,46 +79,149 @@ void SAssetDiscoveryIndicator::Construct( const FArguments& InArgs )
 				// Text
 				+SVerticalBox::Slot()
 				.AutoHeight()
-				.Padding(12, 4)
+				.Padding(AssetDiscoveryIndicatorConstants::Padding)
 				[
-					SNew(STextBlock)
-					.Font( FEditorStyle::GetFontStyle("AssetDiscoveryIndicator.DiscovertingAssetsFont") )
-					.Text( NSLOCTEXT("AssetDiscoveryIndicator", "DiscoveringAssets", "Discovering Assets") )
-					.AutoWrapText(true)
-					.Justification(ETextJustify::Center)
+					SNew(SVerticalBox)
+
+					+SVerticalBox::Slot()
+					.AutoHeight()
+					[
+						SNew(STextBlock)
+						.Font(FEditorStyle::GetFontStyle("AssetDiscoveryIndicator.MainStatusFont"))
+						.Text(this, &SAssetDiscoveryIndicator::GetMainStatusText)
+						.WrapTextAt(this, &SAssetDiscoveryIndicator::GetStatusTextWrapWidth)
+						.Justification(ETextJustify::Center)
+					]
+
+					+SVerticalBox::Slot()
+					.AutoHeight()
+					[
+						SNew(SBox)
+						.Padding(AssetDiscoveryIndicatorConstants::SubStatusTextPadding)
+						.Visibility(this, &SAssetDiscoveryIndicator::GetSubStatusTextVisibility)
+						[
+							SNew(STextBlock)
+							.Font(FEditorStyle::GetFontStyle("AssetDiscoveryIndicator.SubStatusFont"))
+							.Text(this, &SAssetDiscoveryIndicator::GetSubStatusText)
+							.WrapTextAt(this, &SAssetDiscoveryIndicator::GetStatusTextWrapWidth)
+							.Justification(ETextJustify::Center)
+						]
+					]
 				]
 
 				// Progress bar
 				+SVerticalBox::Slot()
 				.AutoHeight()
-				.Padding(12, 4)
+				.Padding(AssetDiscoveryIndicatorConstants::Padding)
 				[
 					SNew(SProgressBar)
-					.Percent( this, &SAssetDiscoveryIndicator::GetProgress )
+					.Percent(this, &SAssetDiscoveryIndicator::GetProgress)
 				]
 			]
 		]
 	];
 }
 
-void SAssetDiscoveryIndicator::OnAssetRegistryFileLoadProgress(int32 NumDiscoveredAssets, int32 TotalAssets)
+void SAssetDiscoveryIndicator::Tick( const FGeometry& AllottedGeometry, const double InCurrentTime, const float InDeltaTime )
+{
+	// Calculate the wrapping width based on our complete alloted width
+	// We do this rather than auto-wrap because the size of the text changes, and auto-wrapping prevents the text block from being able to grow if the text shrinks
+	StatusTextWrapWidth = AllottedGeometry.GetLocalSize().X - (AssetDiscoveryIndicatorConstants::Padding.Left + AssetDiscoveryIndicatorConstants::Padding.Right);
+}
+
+void SAssetDiscoveryIndicator::OnAssetRegistryFileLoadProgress(const IAssetRegistry::FFileLoadProgressUpdateData& ProgressUpdateData)
+{
+	if (ProgressUpdateData.bIsDiscoveringAssetFiles)
+	{
+		// Marquee while we're discovering asset files as we can't yet show an accurate percentage
+		Progress = TOptional<float>();
+		MainStatusText = LOCTEXT("DiscoveringAssetFiles", "Discovering Asset Files");
+		SubStatusText = FText::Format(LOCTEXT("XFilesFoundFmt", "{0} files found"), FText::AsNumber(ProgressUpdateData.NumTotalAssets));
+	}
+	else
+	{
+		if (ProgressUpdateData.NumTotalAssets > 0)
+		{
+			Progress = ProgressUpdateData.NumAssetsProcessedByAssetRegistry / (float)ProgressUpdateData.NumTotalAssets;
+		}
+		else
+		{
+			Progress = 0.0f;
+		}
+
+		if (ProgressUpdateData.NumAssetsPendingDataLoad > 0)
+		{
+			MainStatusText = LOCTEXT("DiscoveringAssetData", "Discovering Asset Data");
+			SubStatusText = FText::Format(LOCTEXT("XAssetsRemainingFmt", "{0} assets remaining"), FText::AsNumber(ProgressUpdateData.NumAssetsPendingDataLoad));
+		}
+		else
+		{
+			const int32 NumAssetsLeftToProcess = ProgressUpdateData.NumTotalAssets - ProgressUpdateData.NumAssetsProcessedByAssetRegistry;
+			if (NumAssetsLeftToProcess == 0)
+			{
+				OnAssetRegistryFilesLoaded();
+			}
+			else
+			{
+				MainStatusText = LOCTEXT("ProcessingAssetData", "Processing Asset Data");
+				SubStatusText = FText::Format(LOCTEXT("XAssetsRemainingFmt", "{0} assets remaining"), FText::AsNumber(NumAssetsLeftToProcess));
+			}
+		}
+	}
+}
+
+void SAssetDiscoveryIndicator::OnAssetRegistryFilesLoaded()
 {
-	if ( Progress < 1.f )
+	if (!FadeAnimation.IsAtStart())
 	{
-		Progress = NumDiscoveredAssets / (float)TotalAssets;
+		MainStatusText = LOCTEXT("FinishedAssetDiscovery", "Finished Asset Discovery");
+		SubStatusText = FText::GetEmpty();
 
-		if ( Progress >= 1.f )
+		if (FadeAnimation.IsPlaying())
 		{
+			if (FadeAnimation.IsForward())
+			{
+				// Still fading in - reverse so we fade back out
+				FadeAnimation.Reverse();
+			}
+			else
+			{
+				// Do nothing - already fading out
+			}
+		}
+		else
+		{
+			// Play the fade out animation
 			FadeAnimation.PlayReverse(this->AsShared());
 		}
 	}
 }
 
+FText SAssetDiscoveryIndicator::GetMainStatusText() const
+{
+	return MainStatusText;
+}
+
+FText SAssetDiscoveryIndicator::GetSubStatusText() const
+{
+	return SubStatusText;
+}
+
 TOptional<float> SAssetDiscoveryIndicator::GetProgress() const
 {
 	return Progress;
 }
 
+EVisibility SAssetDiscoveryIndicator::GetSubStatusTextVisibility() const
+{
+	return (SubStatusText.IsEmpty()) ? EVisibility::Collapsed : EVisibility::Visible;
+}
+
+float SAssetDiscoveryIndicator::GetStatusTextWrapWidth() const
+{
+	return StatusTextWrapWidth;
+}
+
 FSlateColor SAssetDiscoveryIndicator::GetBorderBackgroundColor() const
 {
 	return FSlateColor(FLinearColor(1, 1, 1, 0.8f * FadeCurve.GetLerp()));
@@ -133,4 +249,6 @@ FVector2D SAssetDiscoveryIndicator::GetIndicatorDesiredSizeScale() const
 EVisibility SAssetDiscoveryIndicator::GetIndicatorVisibility() const
 {
 	return FadeAnimation.IsAtStart() ? EVisibility::Collapsed : EVisibility::HitTestInvisible;
-}
\ No newline at end of file
+}
+
+#undef LOCTEXT_NAMESPACE
diff --git a/Engine/Source/Editor/EditorWidgets/Private/SAssetDiscoveryIndicator.h b/Engine/Source/Editor/EditorWidgets/Private/SAssetDiscoveryIndicator.h
--- a/Engine/Source/Editor/EditorWidgets/Private/SAssetDiscoveryIndicator.h
+++ b/Engine/Source/Editor/EditorWidgets/Private/SAssetDiscoveryIndicator.h
@@ -2,6 +2,8 @@
 
 #pragma once
 
+#include "IAssetRegistry.h"
+
 /** An indicator for the progress of the asset registry background search */
 class SAssetDiscoveryIndicator : public SCompoundWidget
 {
@@ -28,12 +30,30 @@ public:
 	/** Constructs this widget with InArgs */
 	void Construct( const FArguments& InArgs );
 
+	// SCompoundWidget interface
+	virtual void Tick( const FGeometry& AllottedGeometry, const double InCurrentTime, const float InDeltaTime ) override;
+
+	/** Handles updating the progress from the asset registry */
+	void OnAssetRegistryFileLoadProgress(const IAssetRegistry::FFileLoadProgressUpdateData& ProgressUpdateData);
+
 	/** Handles updating the progress from the asset registry */
-	void OnAssetRegistryFileLoadProgress(int32 NumDiscoveredAssets, int32 TotalAssets);
+	void OnAssetRegistryFilesLoaded();
+
+	/** Gets the main status text */
+	FText GetMainStatusText() const;
+
+	/** Gets the sub status text */
+	FText GetSubStatusText() const;
 
 	/** Gets the progress bar fraction */
 	TOptional<float> GetProgress() const;
 
+	/** Gets the sub status text visbility */
+	EVisibility GetSubStatusTextVisibility() const;
+
+	/** Get the current wrap point for the status text */
+	float GetStatusTextWrapWidth() const;
+
 	/** Gets the background's opacity */
 	FSlateColor GetBorderBackgroundColor() const;
 
@@ -47,8 +67,17 @@ public:
 	EVisibility GetIndicatorVisibility() const;
 
 private:
+	/** The main status text */
+	FText MainStatusText;
+
+	/** The sub status text (if any) */
+	FText SubStatusText;
+
 	/** The asset registry's asset discovery progress as a percentage */
-	float Progress;
+	TOptional<float> Progress;
+
+	/** The current wrap point for the status text */
+	float StatusTextWrapWidth;
 
 	/** The way the indicator will scale in/out before/after displaying progress */
 	EAssetDiscoveryIndicatorScaleMode::Type ScaleMode;
@@ -57,4 +86,4 @@ private:
 	FCurveSequence FadeAnimation;
 	FCurveHandle FadeCurve;
 	FCurveHandle ScaleCurve;
-};
\ No newline at end of file
+};
diff --git a/Engine/Source/Editor/Kismet/Private/FindInBlueprints.cpp b/Engine/Source/Editor/Kismet/Private/FindInBlueprints.cpp
--- a/Engine/Source/Editor/Kismet/Private/FindInBlueprints.cpp
+++ b/Engine/Source/Editor/Kismet/Private/FindInBlueprints.cpp
@@ -1041,7 +1041,7 @@ void SFindInBlueprints::Construct( const FArguments& InArgs, TSharedPtr<FBluepri
 				.Padding(0, 2)
 				[
 					SNew(STextBlock)
-					.Font( FEditorStyle::GetFontStyle("AssetDiscoveryIndicator.DiscovertingAssetsFont") )
+					.Font( FEditorStyle::GetFontStyle("AssetDiscoveryIndicator.MainStatusFont") )
 					.Text( LOCTEXT("SearchResults", "Searching...") )
 					.Visibility(this, &SFindInBlueprints::GetSearchbarVisiblity)
 				]
diff --git a/Engine/Source/Runtime/AssetRegistry/Private/AssetDataGatherer.cpp b/Engine/Source/Runtime/AssetRegistry/Private/AssetDataGatherer.cpp
--- a/Engine/Source/Runtime/AssetRegistry/Private/AssetDataGatherer.cpp
+++ b/Engine/Source/Runtime/AssetRegistry/Private/AssetDataGatherer.cpp
@@ -2,37 +2,394 @@
 
 #include "AssetRegistryPCH.h"
 
-#define MAX_FILES_TO_PROCESS_BEFORE_FLUSH 250
-#define CACHE_SERIALIZATION_VERSION 6
+namespace AssetDataGathererConstants
+{
+	static const int32 CacheSerializationVersion = 6;
+	static const int32 MaxFilesToDiscoverBeforeFlush = 2500;
+	static const int32 MaxFilesToGatherBeforeFlush = 250;
+	static const int32 MaxFilesToProcessBeforeCacheWrite = 50000;
+}
+
+
+namespace
+{
+	class FLambdaDirectoryStatVisitor : public IPlatformFile::FDirectoryStatVisitor
+	{
+	public:
+		typedef TFunctionRef<bool(const TCHAR*, const FFileStatData&)> FLambdaRef;
+		FLambdaRef Callback;
+		explicit FLambdaDirectoryStatVisitor(FLambdaRef InCallback)
+			: Callback(MoveTemp(InCallback))
+		{
+		}
+		virtual bool Visit(const TCHAR* FilenameOrDirectory, const FFileStatData& StatData) override
+		{
+			return Callback(FilenameOrDirectory, StatData);
+		}
+	};
+
+	bool IsValidPackageFileToRead(const FString& Filename)
+	{
+		FString LongPackageName;
+		if (FPackageName::TryConvertFilenameToLongPackageName(Filename, LongPackageName))
+		{
+			// Make sure the path does not contain invalid characters. These packages will not be successfully loaded or read later.
+			for (const TCHAR PackageChar : LongPackageName)
+			{
+				for (const TCHAR InvalidChar : INVALID_LONGPACKAGE_CHARACTERS)
+				{
+					if (PackageChar == InvalidChar)
+					{
+						return false;
+					}
+				}
+			}
+
+			return true;
+		}
+
+		return false;
+	}
+}
+
+
+FAssetDataDiscovery::FAssetDataDiscovery(const TArray<FString>& InPaths, bool bInIsSynchronous)
+	: bIsSynchronous(bInIsSynchronous)
+	, bIsDiscoveringFiles(false)
+	, StopTaskCounter(0)
+	, Thread(nullptr)
+{
+	DirectoriesToSearch.Reserve(InPaths.Num());
+	for (const FString& Path : InPaths)
+	{
+		// Convert the package path to a filename with no extension (directory)
+		DirectoriesToSearch.Add(FPackageName::LongPackageNameToFilename(Path));
+	}
+
+	if (bIsSynchronous)
+	{
+		Run();
+	}
+	else
+	{
+		Thread = FRunnableThread::Create(this, TEXT("FAssetDataDiscovery"), 0, TPri_BelowNormal);
+	}
+}
+
+FAssetDataDiscovery::~FAssetDataDiscovery()
+{
+}
+
+bool FAssetDataDiscovery::Init()
+{
+	return true;
+}
+
+uint32 FAssetDataDiscovery::Run()
+{
+	double DiscoverStartTime = FPlatformTime::Seconds();
+	int32 NumDiscoveredFiles = 0;
+
+	FString LocalFilenamePathToPrioritize;
+
+	TSet<FString> LocalDiscoveredPathsSet;
+	TArray<FString> LocalDiscoveredDirectories;
+
+	TArray<FDiscoveredPackageFile> LocalPriorityFilesToSearch;
+	TArray<FDiscoveredPackageFile> LocalNonPriorityFilesToSearch;
+
+	// This set contains the folders that we should hide by default unless they contain assets
+	TSet<FString> PathsToHideIfEmpty;
+	PathsToHideIfEmpty.Add(TEXT("/Game/Collections"));
+
+	auto FlushLocalResultsIfRequired = [&]()
+	{
+		if (LocalPriorityFilesToSearch.Num() > 0 || LocalNonPriorityFilesToSearch.Num() > 0 || LocalDiscoveredPathsSet.Num() > 0)
+		{
+			TArray<FString> LocalDiscoveredPathsArray = LocalDiscoveredPathsSet.Array();
+
+			{
+				FScopeLock CritSectionLock(&WorkerThreadCriticalSection);
+
+				// Place all the discovered files into the files to search list
+				DiscoveredPaths.Append(MoveTemp(LocalDiscoveredPathsArray));
+
+				PriorityDiscoveredFiles.Append(MoveTemp(LocalPriorityFilesToSearch));
+				NonPriorityDiscoveredFiles.Append(MoveTemp(LocalNonPriorityFilesToSearch));
+			}
+		}
+
+		LocalDiscoveredPathsSet.Reset();
+
+		LocalPriorityFilesToSearch.Reset();
+		LocalNonPriorityFilesToSearch.Reset();
+	};
+
+	auto IsPriorityFile = [&](const FString& InPackageFilename) -> bool
+	{
+		return !bIsSynchronous && !LocalFilenamePathToPrioritize.IsEmpty() && InPackageFilename.StartsWith(LocalFilenamePathToPrioritize);
+	};
+
+	auto OnIterateDirectoryItem = [&](const TCHAR* InPackageFilename, const FFileStatData& InPackageStatData) -> bool
+	{
+		if (StopTaskCounter.GetValue() != 0)
+		{
+			// Requested to stop - break out of the directory iteration
+			return false;
+		}
+
+		const FString PackageFilenameStr = InPackageFilename;
+
+		if (InPackageStatData.bIsDirectory)
+		{
+			LocalDiscoveredDirectories.Add(PackageFilenameStr / TEXT(""));
+
+			FString PackagePath;
+			if (FPackageName::TryConvertFilenameToLongPackageName(PackageFilenameStr, PackagePath) && !PathsToHideIfEmpty.Contains(PackagePath))
+			{
+				LocalDiscoveredPathsSet.Add(PackagePath);
+			}
+		}
+		else if (FPackageName::IsPackageFilename(PackageFilenameStr))
+		{
+			if (IsValidPackageFileToRead(PackageFilenameStr))
+			{
+				const FString LongPackageNameStr = FPackageName::FilenameToLongPackageName(PackageFilenameStr);
+
+				if (IsPriorityFile(PackageFilenameStr))
+				{
+					LocalPriorityFilesToSearch.Add(FDiscoveredPackageFile(PackageFilenameStr, InPackageStatData.ModificationTime));
+				}
+				else
+				{
+					LocalNonPriorityFilesToSearch.Add(FDiscoveredPackageFile(PackageFilenameStr, InPackageStatData.ModificationTime));
+				}
+
+				LocalDiscoveredPathsSet.Add(FPackageName::GetLongPackagePath(LongPackageNameStr));
+
+				++NumDiscoveredFiles;
+
+				// Flush the data if we've processed enough
+				if (!bIsSynchronous && (LocalPriorityFilesToSearch.Num() + LocalNonPriorityFilesToSearch.Num()) >= AssetDataGathererConstants::MaxFilesToDiscoverBeforeFlush)
+				{
+					FlushLocalResultsIfRequired();
+				}
+			}
+		}
+
+		return true;
+	};
+
+	bool bIsIdle = true;
+
+	while (StopTaskCounter.GetValue() == 0)
+	{
+		FString LocalDirectoryToSearch;
+		{
+			FScopeLock CritSectionLock(&WorkerThreadCriticalSection);
+
+			if (DirectoriesToSearch.Num() > 0)
+			{
+				bIsDiscoveringFiles = true;
+
+				LocalFilenamePathToPrioritize = FilenamePathToPrioritize;
+
+				// Pop off the first path to search
+				LocalDirectoryToSearch = DirectoriesToSearch[0];
+				DirectoriesToSearch.RemoveAt(0, 1, false);
+			}
+		}
+
+		if (LocalDirectoryToSearch.Len() > 0)
+		{
+			if (bIsIdle)
+			{
+				bIsIdle = false;
+
+				// About to start work - reset these
+				DiscoverStartTime = FPlatformTime::Seconds();
+				NumDiscoveredFiles = 0;
+			}
+
+			// Iterate the current search directory
+			FLambdaDirectoryStatVisitor Visitor(OnIterateDirectoryItem);
+			IFileManager::Get().IterateDirectoryStat(*LocalDirectoryToSearch, Visitor);
+
+			{
+				FScopeLock CritSectionLock(&WorkerThreadCriticalSection);
+
+				// Push back any newly discovered sub-directories
+				if (LocalDiscoveredDirectories.Num() > 0)
+				{
+					// Use LocalDiscoveredDirectories as scratch space, then move it back out - this puts the directories we just 
+					// discovered at the start of the list for the next iteration, which can help with disk locality
+					LocalDiscoveredDirectories.Append(MoveTemp(DirectoriesToSearch));
+					DirectoriesToSearch = MoveTemp(LocalDiscoveredDirectories);
+				}
+				LocalDiscoveredDirectories.Reset();
+
+				if (!bIsSynchronous)
+				{
+					FlushLocalResultsIfRequired();
+					SortPathsByPriority(1);
+				}
+			}
+		}
+		else
+		{
+			if (!bIsIdle)
+			{
+				bIsIdle = true;
+
+				{
+					FScopeLock CritSectionLock(&WorkerThreadCriticalSection);
+					bIsDiscoveringFiles = false;
+				}
+
+				UE_LOG(LogAssetRegistry, Verbose, TEXT("Discovery took %0.6f seconds and found %d files to process"), FPlatformTime::Seconds() - DiscoverStartTime, NumDiscoveredFiles);
+			}
+
+			// Ran out of things to do... if we have any pending results, flush those now
+			FlushLocalResultsIfRequired();
+
+			if (bIsSynchronous)
+			{
+				// This is synchronous. Since our work is done, we should safely exit
+				Stop();
+			}
+			else
+			{
+				// No work to do. Sleep for a little and try again later.
+				FPlatformProcess::Sleep(0.1);
+			}
+		}
+	}
+
+	return 0;
+}
+
+void FAssetDataDiscovery::Stop()
+{
+	StopTaskCounter.Increment();
+}
+
+void FAssetDataDiscovery::Exit()
+{
+}
+
+void FAssetDataDiscovery::EnsureCompletion()
+{
+	{
+		FScopeLock CritSectionLock(&WorkerThreadCriticalSection);
+		DirectoriesToSearch.Empty();
+	}
+
+	Stop();
+	Thread->WaitForCompletion();
+    delete Thread;
+    Thread = nullptr;
+}
+
+bool FAssetDataDiscovery::GetAndTrimSearchResults(TArray<FString>& OutDiscoveredPaths, TArray<FDiscoveredPackageFile>& OutDiscoveredFiles, int32& OutNumPathsToSearch)
+{
+	FScopeLock CritSectionLock(&WorkerThreadCriticalSection);
+
+	OutDiscoveredPaths.Append(MoveTemp(DiscoveredPaths));
+	DiscoveredPaths.Reset();
+
+	if (PriorityDiscoveredFiles.Num() > 0)
+	{
+		// Use PriorityDiscoveredFiles as scratch space, then move it back out - this puts the priority files at the start of the final list
+		PriorityDiscoveredFiles.Append(MoveTemp(OutDiscoveredFiles));
+		PriorityDiscoveredFiles.Append(MoveTemp(NonPriorityDiscoveredFiles));
+		OutDiscoveredFiles = MoveTemp(PriorityDiscoveredFiles);
+	}
+	else
+	{
+		OutDiscoveredFiles.Append(MoveTemp(NonPriorityDiscoveredFiles));
+	}
+	PriorityDiscoveredFiles.Reset();
+	NonPriorityDiscoveredFiles.Reset();
+
+	OutNumPathsToSearch = DirectoriesToSearch.Num();
+
+	return bIsDiscoveringFiles;
+}
+
+void FAssetDataDiscovery::AddPathToSearch(const FString& Path)
+{
+	FScopeLock CritSectionLock(&WorkerThreadCriticalSection);
+
+	// Convert the package path to a filename with no extension (directory)
+	DirectoriesToSearch.Add(FPackageName::LongPackageNameToFilename(Path));
+}
+
+void FAssetDataDiscovery::PrioritizeSearchPath(const FString& PathToPrioritize)
+{
+	FString LocalFilenamePathToPrioritize;
+	if (FPackageName::TryConvertLongPackageNameToFilename(PathToPrioritize, LocalFilenamePathToPrioritize))
+	{
+		FScopeLock CritSectionLock(&WorkerThreadCriticalSection);
+
+		FilenamePathToPrioritize = LocalFilenamePathToPrioritize;
+		SortPathsByPriority(INDEX_NONE);
+	}
+}
+
+void FAssetDataDiscovery::SortPathsByPriority(const int32 MaxNumToSort)
+{
+	FScopeLock CritSectionLock(&WorkerThreadCriticalSection);
+
+	// Critical section. This code needs to be as fast as possible since it is in a critical section!
+	// Swap all priority files to the top of the list
+	if (FilenamePathToPrioritize.Len() > 0)
+	{
+		int32 LowestNonPriorityPathIdx = 0;
+		int32 NumSorted = 0;
+		const int32 NumToSort = (MaxNumToSort == INDEX_NONE) ? DirectoriesToSearch.Num() : FMath::Min(DirectoriesToSearch.Num(), MaxNumToSort);
+		for (int32 DirIdx = 0; DirIdx < DirectoriesToSearch.Num(); ++DirIdx)
+		{
+			if (DirectoriesToSearch[DirIdx].StartsWith(FilenamePathToPrioritize))
+			{
+				DirectoriesToSearch.Swap(DirIdx, LowestNonPriorityPathIdx);
+				LowestNonPriorityPathIdx++;
+
+				if (++NumSorted >= NumToSort)
+				{
+					break;
+				}
+			}
+		}
+	}
+}
+
 
 FAssetDataGatherer::FAssetDataGatherer(const TArray<FString>& InPaths, bool bInIsSynchronous, bool bInLoadAndSaveCache)
 	: StopTaskCounter( 0 )
 	, bIsSynchronous( bInIsSynchronous )
 	, bIsDiscoveringFiles( false )
 	, SearchStartTime( 0 )
+	, NumPathsToSearchAtLastSyncPoint( InPaths.Num() )
 	, bLoadAndSaveCache( bInLoadAndSaveCache )
-	, bSavedCacheAfterInitialDiscovery( false )
-	, DiskCachedAssetDataBuffer( NULL )
-	, Thread(NULL)
+	, bFinishedInitialDiscovery( false )
+	, DiskCachedAssetDataBuffer( nullptr )
+	, Thread(nullptr)
 {
-	const FString AllIllegalCharacters = INVALID_LONGPACKAGE_CHARACTERS;
-	for ( int32 CharIdx = 0; CharIdx < AllIllegalCharacters.Len() ; ++CharIdx )
-	{
-		InvalidAssetFileCharacters.Add(AllIllegalCharacters[CharIdx]);
-	}
-
-	PathsToSearch = InPaths;
-
 	bGatherDependsData = GIsEditor && !FParse::Param( FCommandLine::Get(), TEXT("NoDependsGathering") );
 
 	CacheFilename = FPaths::GameIntermediateDir() / TEXT("CachedAssetRegistry.bin");
 
 	if ( bIsSynchronous )
 	{
+		// Run the package file discovery synchronously
+		FAssetDataDiscovery PackageFileDiscovery(InPaths, bIsSynchronous);
+		PackageFileDiscovery.GetAndTrimSearchResults(DiscoveredPaths, FilesToSearch, NumPathsToSearchAtLastSyncPoint);
+
 		Run();
 	}
 	else
 	{
+		BackgroundPackageFileDiscovery = MakeShareable(new FAssetDataDiscovery(InPaths, bIsSynchronous));
 		Thread = FRunnableThread::Create(this, TEXT("FAssetDataGatherer"), 0, TPri_BelowNormal);
 	}
 }
@@ -45,7 +402,7 @@ FAssetDataGatherer::~FAssetDataGatherer()
 	if ( DiskCachedAssetDataBuffer )
 	{
 		delete[] DiskCachedAssetDataBuffer;
-		DiskCachedAssetDataBuffer = NULL;
+		DiskCachedAssetDataBuffer = nullptr;
 	}
 
 	for ( auto CacheIt = NewCachedAssetData.CreateConstIterator(); CacheIt; ++CacheIt )
@@ -62,7 +419,7 @@ bool FAssetDataGatherer::Init()
 
 uint32 FAssetDataGatherer::Run()
 {
-	int32 CacheSerializationVersion = CACHE_SERIALIZATION_VERSION;
+	int32 CacheSerializationVersion = AssetDataGathererConstants::CacheSerializationVersion;
 	
 	static const bool bUsingWorldAssets = FAssetRegistry::IsUsingWorldAssets();
 	if ( bUsingWorldAssets )
@@ -82,18 +439,38 @@ uint32 FAssetDataGatherer::Run()
 		}
 	}
 
-	TArray<FString> LocalFilesToSearch;
+	TArray<FDiscoveredPackageFile> LocalFilesToSearch;
 	TArray<IGatheredAssetData*> LocalAssetResults;
 	TArray<FPackageDependencyData> LocalDependencyResults;
 
+	const double InitialScanStartTime = FPlatformTime::Seconds();
+	int32 NumCachedFiles = 0;
+	int32 NumUncachedFiles = 0;
+
+	int32 NumFilesProcessedSinceLastCacheSave = 0;
+	auto WriteAssetCacheFile = [&]()
+	{
+		FNameTableArchiveWriter CachedAssetDataWriter(CacheSerializationVersion);
+		SerializeCache(CachedAssetDataWriter);
+		CachedAssetDataWriter.SaveToFile(*CacheFilename);
+
+		NumFilesProcessedSinceLastCacheSave = 0;
+	};
+
 	while ( StopTaskCounter.GetValue() == 0 )
 	{
-		// Check to see if there are any paths that need scanning for files.  On the first iteration, there will always
-		// be work to do here.  Later, if new paths are added on the fly, we'll also process those.
-		DiscoverFilesToSearch();
+		bool LocalIsDiscoveringFiles = false;
 
 		{
 			FScopeLock CritSectionLock(&WorkerThreadCriticalSection);
+
+			// Grab any new package files from the background directory scan
+			if (BackgroundPackageFileDiscovery.IsValid())
+			{
+				bIsDiscoveringFiles = BackgroundPackageFileDiscovery->GetAndTrimSearchResults(DiscoveredPaths, FilesToSearch, NumPathsToSearchAtLastSyncPoint);
+				LocalIsDiscoveringFiles = bIsDiscoveringFiles;
+			}
+
 			AssetResults.Append(MoveTemp(LocalAssetResults));
 			DependencyResults.Append(MoveTemp(LocalDependencyResults));
 
@@ -104,11 +481,11 @@ uint32 FAssetDataGatherer::Run()
 					SearchStartTime = FPlatformTime::Seconds();
 				}
 
-				const int32 NumFilesToProcess = FMath::Min<int32>(MAX_FILES_TO_PROCESS_BEFORE_FLUSH, FilesToSearch.Num());
+				const int32 NumFilesToProcess = FMath::Min<int32>(AssetDataGathererConstants::MaxFilesToGatherBeforeFlush, FilesToSearch.Num());
 				LocalFilesToSearch.Append(FilesToSearch.GetData(), NumFilesToProcess);
 				FilesToSearch.RemoveAt(0, NumFilesToProcess, false);
 			}
-			else if (SearchStartTime != 0)
+			else if (SearchStartTime != 0 && !LocalIsDiscoveringFiles)
 			{
 				SearchTimes.Add(FPlatformTime::Seconds() - SearchStartTime);
 				SearchStartTime = 0;
@@ -120,7 +497,7 @@ uint32 FAssetDataGatherer::Run()
 
 		if (LocalFilesToSearch.Num() > 0)
 		{
-			for (const FString& AssetFilename : LocalFilesToSearch)
+			for (const FDiscoveredPackageFile& AssetFileData : LocalFilesToSearch)
 			{
 				if (StopTaskCounter.GetValue() != 0)
 				{
@@ -128,48 +505,95 @@ uint32 FAssetDataGatherer::Run()
 					break;
 				}
 
-				TArray<FBackgroundAssetData*> AssetDataFromFile;
-				FPackageDependencyData DependencyData;
-				if (ReadAssetFile(AssetFilename, AssetDataFromFile, DependencyData))
+				const FName PackageName = FName(*FPackageName::FilenameToLongPackageName(AssetFileData.PackageFilename));
+
+				bool bLoadedFromCache = false;
+				if (bLoadAndSaveCache)
 				{
-					LocalAssetResults.Append(AssetDataFromFile);
-					LocalDependencyResults.Add(DependencyData);
+					FDiskCachedAssetData** DiskCachedAssetDataPtr = DiskCachedAssetDataMap.Find(PackageName);
+					FDiskCachedAssetData* DiskCachedAssetData = nullptr;
+					if (DiskCachedAssetDataPtr && *DiskCachedAssetDataPtr)
+					{
+						const FDateTime& CachedTimestamp = (*DiskCachedAssetDataPtr)->Timestamp;
+						if (AssetFileData.PackageTimestamp == CachedTimestamp)
+						{
+							DiskCachedAssetData = *DiskCachedAssetDataPtr;
+						}
+					}
 
-					bool bCachePackage = bLoadAndSaveCache;
-					if (bCachePackage)
+					if (DiskCachedAssetData)
 					{
-						// Don't store info on cooked packages
-						for (const auto& AssetData : AssetDataFromFile)
+						bLoadedFromCache = true;
+
+						++NumCachedFiles;
+
+						LocalAssetResults.Reserve(LocalAssetResults.Num() + DiskCachedAssetData->AssetDataList.Num());
+						for (const FAssetData& AssetData : DiskCachedAssetData->AssetDataList)
 						{
-							if (AssetData->IsCooked())
-							{
-								bCachePackage = false;
-								break;
-							}
+							LocalAssetResults.Add(new FAssetDataWrapper(AssetData));
 						}
+
+						LocalDependencyResults.Add(DiskCachedAssetData->DependencyData);
+
+						NewCachedAssetDataMap.Add(PackageName, DiskCachedAssetData);
 					}
+				}
 
-					if (bCachePackage)
+				if (!bLoadedFromCache)
+				{
+					TArray<FBackgroundAssetData*> AssetDataFromFile;
+					FPackageDependencyData DependencyData;
+					if (ReadAssetFile(AssetFileData.PackageFilename, AssetDataFromFile, DependencyData))
 					{
-						// Update the cache
-						const FName PackageName = FName(*FPackageName::FilenameToLongPackageName(AssetFilename));
-						const FDateTime FileTimestamp = IFileManager::Get().GetTimeStamp(*AssetFilename);
+						++NumUncachedFiles;
+
+						LocalAssetResults.Append(AssetDataFromFile);
+						LocalDependencyResults.Add(DependencyData);
 
-						FDiskCachedAssetData* NewData = new FDiskCachedAssetData(PackageName, FileTimestamp);
-						NewData->AssetDataList.Reserve(AssetDataFromFile.Num());
-						for (const FBackgroundAssetData* BackgroundAssetData : AssetDataFromFile)
+						bool bCachePackage = bLoadAndSaveCache;
+						if (bCachePackage)
 						{
-							NewData->AssetDataList.Add(BackgroundAssetData->ToAssetData());
+							// Don't store info on cooked packages
+							for (const auto& AssetData : AssetDataFromFile)
+							{
+								if (AssetData->IsCooked())
+								{
+									bCachePackage = false;
+									break;
+								}
+							}
 						}
-						NewData->DependencyData = DependencyData;
 
-						NewCachedAssetData.Add(NewData);
-						NewCachedAssetDataMap.Add(PackageName, NewData);
+						if (bCachePackage)
+						{
+							++NumFilesProcessedSinceLastCacheSave;
+
+							// Update the cache
+							FDiskCachedAssetData* NewData = new FDiskCachedAssetData(PackageName, AssetFileData.PackageTimestamp);
+							NewData->AssetDataList.Reserve(AssetDataFromFile.Num());
+							for (const FBackgroundAssetData* BackgroundAssetData : AssetDataFromFile)
+							{
+								NewData->AssetDataList.Add(BackgroundAssetData->ToAssetData());
+							}
+							NewData->DependencyData = DependencyData;
+
+							NewCachedAssetData.Add(NewData);
+							NewCachedAssetDataMap.Add(PackageName, NewData);
+						}
 					}
 				}
 			}
 
 			LocalFilesToSearch.Reset();
+
+			if (bLoadAndSaveCache)
+			{
+				// Save off the cache files if we're processed enough data since the last save
+				if (NumFilesProcessedSinceLastCacheSave >= AssetDataGathererConstants::MaxFilesToProcessBeforeCacheWrite)
+				{
+					WriteAssetCacheFile();
+				}
+			}
 		}
 		else
 		{
@@ -180,14 +604,17 @@ uint32 FAssetDataGatherer::Run()
 			}
 			else
 			{
-				// If we are caching discovered assets and this is the first time we had no work to do, save off the cache now in case the user terminates unexpectedly
-				if (bLoadAndSaveCache && !bSavedCacheAfterInitialDiscovery)
+				if (!LocalIsDiscoveringFiles && !bFinishedInitialDiscovery)
 				{
-					FNameTableArchiveWriter CachedAssetDataWriter(CacheSerializationVersion);
-					SerializeCache(CachedAssetDataWriter);
-					CachedAssetDataWriter.SaveToFile(*CacheFilename);
+					bFinishedInitialDiscovery = true;
 
-					bSavedCacheAfterInitialDiscovery = true;
+					UE_LOG(LogAssetRegistry, Verbose, TEXT("Initial scan took %0.6f seconds (found %d cached assets, and loaded %d)"), FPlatformTime::Seconds() - InitialScanStartTime, NumCachedFiles, NumUncachedFiles);
+
+					// If we are caching discovered assets and this is the first time we had no work to do, save off the cache now in case the user terminates unexpectedly
+					if (bLoadAndSaveCache)
+					{
+						WriteAssetCacheFile();
+					}
 				}
 
 				// No work to do. Sleep for a little and try again later.
@@ -198,9 +625,7 @@ uint32 FAssetDataGatherer::Run()
 
 	if ( bLoadAndSaveCache )
 	{
-		FNameTableArchiveWriter CachedAssetData(CacheSerializationVersion);
-		SerializeCache(CachedAssetData);
-		CachedAssetData.SaveToFile(*CacheFilename);
+		WriteAssetCacheFile();
 	}
 
 	return 0;
@@ -208,29 +633,37 @@ uint32 FAssetDataGatherer::Run()
 
 void FAssetDataGatherer::Stop()
 {
+	if (BackgroundPackageFileDiscovery.IsValid())
+	{
+		BackgroundPackageFileDiscovery->Stop();
+	}
+
 	StopTaskCounter.Increment();
 }
 
 void FAssetDataGatherer::Exit()
-{
-    
+{   
 }
 
 void FAssetDataGatherer::EnsureCompletion()
 {
+	if (BackgroundPackageFileDiscovery.IsValid())
+	{
+		BackgroundPackageFileDiscovery->EnsureCompletion();
+	}
+
 	{
 		FScopeLock CritSectionLock(&WorkerThreadCriticalSection);
 		FilesToSearch.Empty();
-		PathsToSearch.Empty();
 	}
 
 	Stop();
 	Thread->WaitForCompletion();
     delete Thread;
-    Thread = NULL;
+    Thread = nullptr;
 }
 
-bool FAssetDataGatherer::GetAndTrimSearchResults(TArray<IGatheredAssetData*>& OutAssetResults, TArray<FString>& OutPathResults, TArray<FPackageDependencyData>& OutDependencyResults, TArray<double>& OutSearchTimes, int32& OutNumFilesToSearch, int32& OutNumPathsToSearch)
+bool FAssetDataGatherer::GetAndTrimSearchResults(TArray<IGatheredAssetData*>& OutAssetResults, TArray<FString>& OutPathResults, TArray<FPackageDependencyData>& OutDependencyResults, TArray<double>& OutSearchTimes, int32& OutNumFilesToSearch, int32& OutNumPathsToSearch, bool& OutIsDiscoveringFiles)
 {
 	FScopeLock CritSectionLock(&WorkerThreadCriticalSection);
 
@@ -247,15 +680,18 @@ bool FAssetDataGatherer::GetAndTrimSearchResults(TArray<IGatheredAssetData*>& Ou
 	SearchTimes.Reset();
 
 	OutNumFilesToSearch = FilesToSearch.Num();
-	OutNumPathsToSearch = PathsToSearch.Num();
+	OutNumPathsToSearch = NumPathsToSearchAtLastSyncPoint;
+	OutIsDiscoveringFiles = bIsDiscoveringFiles;
 
 	return (SearchStartTime > 0 || bIsDiscoveringFiles);
 }
 
 void FAssetDataGatherer::AddPathToSearch(const FString& Path)
 {
-	FScopeLock CritSectionLock(&WorkerThreadCriticalSection);
-	PathsToSearch.Add(Path);
+	if (BackgroundPackageFileDiscovery.IsValid())
+	{
+		BackgroundPackageFileDiscovery->AddPathToSearch(Path);
+	}
 }
 
 void FAssetDataGatherer::AddFilesToSearch(const TArray<FString>& Files)
@@ -278,153 +714,46 @@ void FAssetDataGatherer::AddFilesToSearch(const TArray<FString>& Files)
 
 void FAssetDataGatherer::PrioritizeSearchPath(const FString& PathToPrioritize)
 {
-	FString FilenamePathToPrioritize;
-	if ( FPackageName::TryConvertLongPackageNameToFilename(PathToPrioritize, FilenamePathToPrioritize) )
+	if (BackgroundPackageFileDiscovery.IsValid())
 	{
-		// Critical section. This code needs to be as fast as possible since it is in a critical section!
-		// Swap all priority files to the top of the list
-		{
-			FScopeLock CritSectionLock(&WorkerThreadCriticalSection);
-			int32 LowestNonPriorityFileIdx = 0;
-			for ( int32 FilenameIdx = 0; FilenameIdx < FilesToSearch.Num(); ++FilenameIdx )
-			{
-				if ( FilesToSearch[FilenameIdx].StartsWith(FilenamePathToPrioritize) )
-				{
-					FilesToSearch.Swap(FilenameIdx, LowestNonPriorityFileIdx);
-					LowestNonPriorityFileIdx++;
-				}
-			}
-		}
+		BackgroundPackageFileDiscovery->PrioritizeSearchPath(PathToPrioritize);
 	}
-}
-
-void FAssetDataGatherer::DiscoverFilesToSearch()
-{
-	const double DiscoverStartTime = FPlatformTime::Seconds();
-
-	TArray<FString> LocalFilesToSearch;
-	TSet<FString> LocalDiscoveredPathsSet;
-
-	int32 NumCachedPackages = 0;
-	TArray<IGatheredAssetData*> LocalAssetResults;
-	TArray<FPackageDependencyData> LocalDependencyResults;
-
-	auto OnPackageDiscovered = [&](const TCHAR* InPackageFilename, const FFileStatData& InPackageStatData) -> bool
-	{
-		const FString PackageFilenameStr = InPackageFilename;
-		if (IsValidPackageFileToRead(PackageFilenameStr))
-		{
-			const FString LongPackageNameStr = FPackageName::FilenameToLongPackageName(PackageFilenameStr);
-
-			bool bLoadedFromCache = false;
-			if (bLoadAndSaveCache)
-			{
-				const FName PackageName = FName(*LongPackageNameStr);
-
-				FDiskCachedAssetData** DiskCachedAssetDataPtr = DiskCachedAssetDataMap.Find(PackageName);
-				FDiskCachedAssetData* DiskCachedAssetData = nullptr;
-				if (DiskCachedAssetDataPtr && *DiskCachedAssetDataPtr)
-				{
-					const FDateTime& CachedTimestamp = (*DiskCachedAssetDataPtr)->Timestamp;
-					if (InPackageStatData.ModificationTime == CachedTimestamp)
-					{
-						DiskCachedAssetData = *DiskCachedAssetDataPtr;
-					}
-				}
-
-				if (DiskCachedAssetData)
-				{
-					bLoadedFromCache = true;
-
-					++NumCachedPackages;
-
-					LocalAssetResults.Reserve(LocalAssetResults.Num() + DiskCachedAssetData->AssetDataList.Num());
-					for (const FAssetData& AssetData : DiskCachedAssetData->AssetDataList)
-					{
-						LocalAssetResults.Add(new FAssetDataWrapper(AssetData));
-					}
-
-					LocalDependencyResults.Add(DiskCachedAssetData->DependencyData);
-
-					NewCachedAssetDataMap.Add(PackageName, DiskCachedAssetData);
-				}
-			}
-
-			if (!bLoadedFromCache)
-			{
-				LocalFilesToSearch.Add(PackageFilenameStr);
-			}
-
-			LocalDiscoveredPathsSet.Add(FPackageName::GetLongPackagePath(LongPackageNameStr));
-		}
 
-		return true;
-	};
-
-	if (PathsToSearch.Num() > 0)
+	FString LocalFilenamePathToPrioritize;
+	if (FPackageName::TryConvertLongPackageNameToFilename(PathToPrioritize, LocalFilenamePathToPrioritize))
 	{
-		TArray<FString> LocalPathsToSearch;
-		{
-			FScopeLock CritSectionLock(&WorkerThreadCriticalSection);
-
-			LocalPathsToSearch = PathsToSearch;
-
-			// Remove all of the existing paths from the list, since we'll process them all below.  New paths may be
-			// added to the original list on a different thread as we go along, but those new paths won't be processed
-			// during this call of DisoverFilesToSearch().  But we'll get them on the next call!
-			PathsToSearch.Reset();
-			bIsDiscoveringFiles = true;
-		}
-
-		// Iterate over any paths that we have remaining to scan
-		for (const FString& PathToSearch : LocalPathsToSearch)
-		{
-			// Convert the package path to a filename with no extension (directory)
-			const FString FilePath = FPackageName::LongPackageNameToFilename(PathToSearch);
-
-			// Discover any packages that we need to scan
-			// This will also populate our local arrays with any valid cached data
-			FPackageName::IteratePackagesInDirectory(FilePath, FPackageName::FPackageNameStatVisitor(OnPackageDiscovered));
-		}
-
-		UE_LOG(LogAssetRegistry, Verbose, TEXT("Discovery took %0.6f seconds (found %d cached packages, with %d to re-scan)"), FPlatformTime::Seconds() - DiscoverStartTime, NumCachedPackages, LocalFilesToSearch.Num());
-
-		// Move all our discovered files and data into the final lists
-		TArray<FString> LocalDiscoveredPathsArray = LocalDiscoveredPathsSet.Array();
-		{
-			FScopeLock CritSectionLock(&WorkerThreadCriticalSection);
-
-			// Place all the discovered files into the files to search list
-			FilesToSearch.Append(MoveTemp(LocalFilesToSearch));
-			DiscoveredPaths.Append(MoveTemp(LocalDiscoveredPathsArray));
-
-			// Place all the discovered cached data into the final lists for consumption
-			AssetResults.Append(MoveTemp(LocalAssetResults));
-			DependencyResults.Append(MoveTemp(LocalDependencyResults));
+		FScopeLock CritSectionLock(&WorkerThreadCriticalSection);
 
-			bIsDiscoveringFiles = false;
-		}
+		FilenamePathToPrioritize = LocalFilenamePathToPrioritize;
+		SortPathsByPriority(INDEX_NONE);
 	}
 }
 
-bool FAssetDataGatherer::IsValidPackageFileToRead(const FString& Filename) const
+void FAssetDataGatherer::SortPathsByPriority(const int32 MaxNumToSort)
 {
-	FString LongPackageName;
-	if ( FPackageName::TryConvertFilenameToLongPackageName(Filename, LongPackageName) )
+	FScopeLock CritSectionLock(&WorkerThreadCriticalSection);
+
+	// Critical section. This code needs to be as fast as possible since it is in a critical section!
+	// Swap all priority files to the top of the list
+	if (FilenamePathToPrioritize.Len() > 0)
 	{
-		// Make sure the path does not contain invalid characters. These packages will not be successfully loaded or read later.
-		for ( int32 CharIdx = 0; CharIdx < LongPackageName.Len() ; ++CharIdx )
+		int32 LowestNonPriorityFileIdx = 0;
+		int32 NumSorted = 0;
+		const int32 NumToSort = (MaxNumToSort == INDEX_NONE) ? FilesToSearch.Num() : FMath::Min(FilesToSearch.Num(), MaxNumToSort);
+		for (int32 FilenameIdx = 0; FilenameIdx < FilesToSearch.Num(); ++FilenameIdx)
 		{
-			if ( InvalidAssetFileCharacters.Contains(LongPackageName[CharIdx]) )
+			if (FilesToSearch[FilenameIdx].PackageFilename.StartsWith(FilenamePathToPrioritize))
 			{
-				return false;
+				FilesToSearch.Swap(FilenameIdx, LowestNonPriorityFileIdx);
+				LowestNonPriorityFileIdx++;
+
+				if (++NumSorted >= NumToSort)
+				{
+					break;
+				}
 			}
 		}
-
-		return true;
 	}
-
-	return false;
 }
 
 FBackgroundAssetData* FAssetDataGatherer::CreateAssetDataFromLinkerTables(const FString& AssetFilename, uint32 InPackageFlags, const FObjectExport& AssetExport, const TArray<FObjectImport>& ImportMap, const TArray<FObjectExport>& ExportMap) const
diff --git a/Engine/Source/Runtime/AssetRegistry/Private/AssetDataGatherer.h b/Engine/Source/Runtime/AssetRegistry/Private/AssetDataGatherer.h
--- a/Engine/Source/Runtime/AssetRegistry/Private/AssetDataGatherer.h
+++ b/Engine/Source/Runtime/AssetRegistry/Private/AssetDataGatherer.h
@@ -2,6 +2,129 @@
 
 #pragma once
 
+/**
+ * Minimal amount of information needed about a discovered asset file
+ */
+struct FDiscoveredPackageFile
+{
+	explicit FDiscoveredPackageFile(FString InPackageFilename)
+		: PackageFilename(MoveTemp(InPackageFilename))
+		, PackageTimestamp(IFileManager::Get().GetTimeStamp(*PackageFilename))
+	{
+	}
+
+	FDiscoveredPackageFile(FString InPackageFilename, FDateTime InPackageTimestamp)
+		: PackageFilename(MoveTemp(InPackageFilename))
+		, PackageTimestamp(MoveTemp(InPackageTimestamp))
+	{
+	}
+
+	FDiscoveredPackageFile(const FDiscoveredPackageFile& Other)
+		: PackageFilename(Other.PackageFilename)
+		, PackageTimestamp(Other.PackageTimestamp)
+	{
+	}
+
+	FDiscoveredPackageFile(FDiscoveredPackageFile&& Other)
+		: PackageFilename(MoveTemp(Other.PackageFilename))
+		, PackageTimestamp(MoveTemp(Other.PackageTimestamp))
+	{
+	}
+
+	FDiscoveredPackageFile& operator=(const FDiscoveredPackageFile& Other)
+	{
+		if (this != &Other)
+		{
+			PackageFilename = Other.PackageFilename;
+			PackageTimestamp = Other.PackageTimestamp;
+		}
+		return *this;
+	}
+
+	FDiscoveredPackageFile& operator=(FDiscoveredPackageFile&& Other)
+	{
+		if (this != &Other)
+		{
+			PackageFilename = MoveTemp(Other.PackageFilename);
+			PackageTimestamp = MoveTemp(Other.PackageTimestamp);
+		}
+		return *this;
+	}
+
+	/** The name of the package file on disk */
+	FString PackageFilename;
+
+	/** The modification timestamp of the package file (when it was discovered) */
+	FDateTime PackageTimestamp;
+};
+
+
+/**
+ * Async task for discovering files that FAssetDataGatherer should search
+ */
+class FAssetDataDiscovery : public FRunnable
+{
+public:
+	/** Constructor */
+	FAssetDataDiscovery(const TArray<FString>& InPaths, bool bInIsSynchronous);
+	virtual ~FAssetDataDiscovery();
+
+	// FRunnable implementation
+	virtual bool Init() override;
+	virtual uint32 Run() override;
+	virtual void Stop() override;
+	virtual void Exit() override;
+
+	/** Signals to end the thread and waits for it to close before returning */
+	void EnsureCompletion();
+
+	/** Gets search results from the file discovery */
+	bool GetAndTrimSearchResults(TArray<FString>& OutDiscoveredPaths, TArray<FDiscoveredPackageFile>& OutDiscoveredFiles, int32& OutNumPathsToSearch);
+
+	/** Adds a root path to the search queue. Only works when searching asynchronously */
+	void AddPathToSearch(const FString& Path);
+
+	/** If assets are currently being asynchronously scanned in the specified path, this will cause them to be scanned before other assets. */
+	void PrioritizeSearchPath(const FString& PathToPrioritize);
+
+private:
+	/** Sort the paths so that items belonging to the current priority path is processed first */
+	void SortPathsByPriority(const int32 MaxNumToSort);
+
+	/** A critical section to protect data transfer to other threads */
+	FCriticalSection WorkerThreadCriticalSection;
+
+	/** The current path to prioritize. Files and folders under this location will be bumped to the top of the processing list as they are discovered */
+	FString FilenamePathToPrioritize;
+
+	/** True if this gather request is synchronous (i.e, IsRunningCommandlet()) */
+	bool bIsSynchronous;
+
+	/** True if in the process of discovering files in PathsToSearch */
+	bool bIsDiscoveringFiles;
+
+	/** The directories in which to discover assets and paths */
+	// IMPORTANT: This variable may be modified by from a different thread via a call to AddPathToSearch(), so access 
+	//            to this array should be handled very carefully.
+	TArray<FString> DirectoriesToSearch;
+
+	/** The paths found during the search. It is not threadsafe to directly access this array */
+	TArray<FString> DiscoveredPaths;
+
+	/** List of priority files that need to be processed by the gatherer. It is not threadsafe to directly access this array */
+	TArray<FDiscoveredPackageFile> PriorityDiscoveredFiles;
+
+	/** List of non-priority files that need to be processed by the gatherer. It is not threadsafe to directly access this array */
+	TArray<FDiscoveredPackageFile> NonPriorityDiscoveredFiles;
+
+	/** > 0 if we've been asked to abort work in progress at the next opportunity */
+	FThreadSafeCounter StopTaskCounter;
+
+	/** Thread to run the discovery FRunnable on */
+	FRunnableThread* Thread;
+};
+
+
 /**
  * Async task for gathering asset data from from the file list in FAssetRegistry
  */
@@ -22,7 +145,7 @@ public:
 	void EnsureCompletion();
 
 	/** Gets search results from the data gatherer */
-	bool GetAndTrimSearchResults(TArray<IGatheredAssetData*>& OutAssetResults, TArray<FString>& OutPathResults, TArray<FPackageDependencyData>& OutDependencyResults, TArray<double>& OutSearchTimes, int32& OutNumFilesToSearch, int32& OutNumPathsToSearch);
+	bool GetAndTrimSearchResults(TArray<IGatheredAssetData*>& OutAssetResults, TArray<FString>& OutPathResults, TArray<FPackageDependencyData>& OutDependencyResults, TArray<double>& OutSearchTimes, int32& OutNumFilesToSearch, int32& OutNumPathsToSearch, bool& OutIsDiscoveringFiles);
 
 	/** Adds a root path to the search queue. Only works when searching asynchronously */
 	void AddPathToSearch(const FString& Path);
@@ -34,12 +157,8 @@ public:
 	void PrioritizeSearchPath(const FString& PathToPrioritize);
 
 private:
-	
-	/** This function is run on the gathering thread, unless synchronous */
-	void DiscoverFilesToSearch();
-
-	/** Returns true if this package file has only valid characters and can exist in a content root */
-	bool IsValidPackageFileToRead(const FString& Filename) const;
+	/** Sort the paths so that items belonging to the current priority path is processed first */
+	void SortPathsByPriority(const int32 MaxNumToSort);
 
 	/**
 	 * Reads FAssetData information out of a file
@@ -65,8 +184,11 @@ private:
 	/** A critical section to protect data transfer to the main thread */
 	FCriticalSection WorkerThreadCriticalSection;
 
+	/** The current path to prioritize. Files and folders under this location will be bumped to the top of the processing list as they are discovered */
+	FString FilenamePathToPrioritize;
+
 	/** List of files that need to be processed by the search. It is not threadsafe to directly access this array */
-	TArray<FString> FilesToSearch;
+	TArray<FDiscoveredPackageFile> FilesToSearch;
 
 	/** > 0 if we've been asked to abort work in progress at the next opportunity */
 	FThreadSafeCounter StopTaskCounter;
@@ -80,11 +202,6 @@ private:
 	/** The current search start time */
 	double SearchStartTime;
 
-	/** The input base paths in which to discover assets and paths */
-	// IMPORTANT: This variable may be modified by from a different thread via a call to AddPathToSearch(), so access 
-	//            to this array should be handled very carefully.
-	TArray<FString> PathsToSearch;
-
 	/** The asset data gathered from the searched files */
 	TArray<IGatheredAssetData*> AssetResults;
 
@@ -100,8 +217,11 @@ private:
 	/** True if dependency data should be gathered */
 	bool bGatherDependsData;
 
-	/** Set of characters that are invalid in packages, thus files containing them can not be loaded. */
-	TSet<TCHAR> InvalidAssetFileCharacters;
+	/** The cached value of the NumPathsToSearch returned by FAssetDataDiscovery the last time we synchronized with it */
+	int32 NumPathsToSearchAtLastSyncPoint;
+
+	/** Background package file discovery (when running async) */
+	TSharedPtr<FAssetDataDiscovery> BackgroundPackageFileDiscovery;
 
 	///////////////////////////////////////////////////////////////
 	// Asset discovery caching
@@ -110,8 +230,8 @@ private:
 	/** True if this gather request should both load and save the asset cache. Only one gatherer should do this at a time! */
 	bool bLoadAndSaveCache;
 
-	/** True if we have finished discovering our first wave of files. Used to preemptively save the discovered assets cache once a full scan has completed */
-	bool bSavedCacheAfterInitialDiscovery;
+	/** True if we have finished discovering our first wave of files */
+	bool bFinishedInitialDiscovery;
 
 	/** The name of the file that contains the timestamped cache of discovered assets */
 	FString CacheFilename;
@@ -128,7 +248,6 @@ private:
 	/** Map of PackageName to cached discovered assets that will be written to disk at shutdown */
 	TMap<FName, FDiskCachedAssetData*> NewCachedAssetDataMap;
 
-public:
 	/** Thread to run the cleanup FRunnable on */
 	FRunnableThread* Thread;
-};
\ No newline at end of file
+};
diff --git a/Engine/Source/Runtime/AssetRegistry/Private/AssetRegistry.cpp b/Engine/Source/Runtime/AssetRegistry/Private/AssetRegistry.cpp
--- a/Engine/Source/Runtime/AssetRegistry/Private/AssetRegistry.cpp
+++ b/Engine/Source/Runtime/AssetRegistry/Private/AssetRegistry.cpp
@@ -1362,13 +1362,14 @@ void FAssetRegistry::Tick(float DeltaTime)
 	double TickStartTime = FPlatformTime::Seconds();
 
 	// Gather results from the background search
-	int32 NumFilesToSearch = 0;
-	int32 NumPathsToSearch = 0;
 	bool bIsSearching = false;
 	TArray<double> SearchTimes;
+	int32 NumFilesToSearch = 0;
+	int32 NumPathsToSearch = 0;
+	bool bIsDiscoveringFiles = false;
 	if ( BackgroundAssetSearch.IsValid() )
-	{		
-		bIsSearching = BackgroundAssetSearch->GetAndTrimSearchResults(BackgroundAssetResults, BackgroundPathResults, BackgroundDependencyResults, SearchTimes, NumFilesToSearch, NumPathsToSearch);
+	{
+		bIsSearching = BackgroundAssetSearch->GetAndTrimSearchResults(BackgroundAssetResults, BackgroundPathResults, BackgroundDependencyResults, SearchTimes, NumFilesToSearch, NumPathsToSearch, bIsDiscoveringFiles);
 	}
 
 	// Report the search times
@@ -1384,6 +1385,7 @@ void FAssetRegistry::Tick(float DeltaTime)
 	}
 
 	// Process the asset results
+	const bool bHadAssetsToProcess = BackgroundAssetResults.Num() > 0;
 	if ( BackgroundAssetResults.Num() )
 	{
 		// Mark the first amortize time
@@ -1399,12 +1401,6 @@ void FAssetRegistry::Tick(float DeltaTime)
 			TotalAmortizeTime += FPlatformTime::Seconds() - AmortizeStartTime;
 			AmortizeStartTime = 0;
 		}
-
-		// Notify the status change
-		int32 NumDiscoveredAssets = CachedAssetsByObjectPath.Num();
-		// Total assets may not be accurate if some files to search are not UAsset files. This number may fluctuate a little as packages are searched.
-		int32 TotalAssets = NumDiscoveredAssets + BackgroundAssetResults.Num() + NumFilesToSearch;
-		FileLoadProgressUpdatedEvent.Broadcast(NumDiscoveredAssets, TotalAssets);
 	}
 
 	// Add dependencies
@@ -1413,6 +1409,18 @@ void FAssetRegistry::Tick(float DeltaTime)
 		DependencyDataGathered(TickStartTime, BackgroundDependencyResults);
 	}
 
+	// Notify the status change
+	if (bIsSearching || bHadAssetsToProcess)
+	{
+		const FFileLoadProgressUpdateData ProgressUpdateData(
+			CachedAssetsByObjectPath.Num() + BackgroundAssetResults.Num() + NumFilesToSearch,	// NumTotalAssets
+			CachedAssetsByObjectPath.Num(),														// NumAssetsProcessedByAssetRegistry
+			NumFilesToSearch,																	// NumAssetsPendingDataLoad
+			bIsDiscoveringFiles																	// bIsDiscoveringAssetFiles
+			);
+		FileLoadProgressUpdatedEvent.Broadcast(ProgressUpdateData);
+	}
+
 	// If completing an initial search, refresh the content browser
 	if ( NumFilesToSearch == 0 && NumPathsToSearch == 0 && !bIsSearching && BackgroundPathResults.Num() == 0 && BackgroundAssetResults.Num() == 0 && BackgroundDependencyResults.Num() == 0 )
 	{
@@ -1790,7 +1798,8 @@ void FAssetRegistry::ScanPathsSynchronous_Internal(const TArray<FString>& InPath
 		TArray<double> SearchTimes;
 		int32 NumFilesToSearch = 0;
 		int32 NumPathsToSearch = 0;
-		AssetSearch.GetAndTrimSearchResults(AssetResults, PathResults, DependencyResults, SearchTimes, NumFilesToSearch, NumPathsToSearch);
+		bool bIsDiscoveringFiles = false;
+		AssetSearch.GetAndTrimSearchResults(AssetResults, PathResults, DependencyResults, SearchTimes, NumFilesToSearch, NumPathsToSearch, bIsDiscoveringFiles);
 
 		// Cache the search results
 		const int32 NumResults = AssetResults.Num();
diff --git a/Engine/Source/Runtime/AssetRegistry/Public/IAssetRegistry.h b/Engine/Source/Runtime/AssetRegistry/Public/IAssetRegistry.h
--- a/Engine/Source/Runtime/AssetRegistry/Public/IAssetRegistry.h
+++ b/Engine/Source/Runtime/AssetRegistry/Public/IAssetRegistry.h
@@ -213,8 +213,25 @@ public:
 	DECLARE_EVENT( IAssetRegistry, FFilesLoadedEvent );
 	virtual FFilesLoadedEvent& OnFilesLoaded() = 0;
 
+	/** Payload data for a file progress update */
+	struct FFileLoadProgressUpdateData
+	{
+		FFileLoadProgressUpdateData(int32 InNumTotalAssets, int32 InNumAssetsProcessedByAssetRegistry, int32 InNumAssetsPendingDataLoad, bool InIsDiscoveringAssetFiles)
+			: NumTotalAssets(InNumTotalAssets)
+			, NumAssetsProcessedByAssetRegistry(InNumAssetsProcessedByAssetRegistry)
+			, NumAssetsPendingDataLoad(InNumAssetsPendingDataLoad)
+			, bIsDiscoveringAssetFiles(InIsDiscoveringAssetFiles)
+		{
+		}
+
+		int32 NumTotalAssets;
+		int32 NumAssetsProcessedByAssetRegistry;
+		int32 NumAssetsPendingDataLoad;
+		bool bIsDiscoveringAssetFiles;
+	};
+
 	/** Event to update the progress of the background file load */
-	DECLARE_EVENT_TwoParams( IAssetRegistry, FFileLoadProgressUpdatedEvent, int32 /*NumAssetsDiscovered*/, int32 /*TotalAssets*/ );
+	DECLARE_EVENT_OneParam( IAssetRegistry, FFileLoadProgressUpdatedEvent, const FFileLoadProgressUpdateData& /*ProgressUpdateData*/ );
 	virtual FFileLoadProgressUpdatedEvent& OnFileLoadProgressUpdated() = 0;
 
 	/** Returns true if the asset registry is currently loading files and does not yet know about all assets */
