#UE4 Moved DependsNode.h back to the private folder while it is being discussed as to how we should handle serialization of the asset registry. This was intended to be an internal data structure and should not be used outside of the AssetRegistry module.

diff --git a/Engine/Source/Editor/UnrealEd/Private/Commandlets/ChunkManifestGenerator.cpp b/Engine/Source/Editor/UnrealEd/Private/Commandlets/ChunkManifestGenerator.cpp
--- a/Engine/Source/Editor/UnrealEd/Private/Commandlets/ChunkManifestGenerator.cpp
+++ b/Engine/Source/Editor/UnrealEd/Private/Commandlets/ChunkManifestGenerator.cpp
@@ -310,10 +310,9 @@ bool FChunkManifestGenerator::LoadAssetRegistry(const FString& SandboxPath, cons
 		FArchive* AssetRegistryReader = &FileContents;
 
 		TMap<FName, FAssetData*> SavedAssetRegistryData;
-		TArray<FDependsNode*> DependencyData;
 		if (AssetRegistryReader)
 		{
-			AssetRegistry.LoadRegistryData(*AssetRegistryReader, SavedAssetRegistryData, DependencyData);
+			AssetRegistry.LoadRegistryData(*AssetRegistryReader, SavedAssetRegistryData);
 		}
 		for (auto& LoadedAssetData : AssetRegistryData)
 		{
diff --git a/Engine/Source/Editor/UnrealEd/Private/CookOnTheFlyServer.cpp b/Engine/Source/Editor/UnrealEd/Private/CookOnTheFlyServer.cpp
--- a/Engine/Source/Editor/UnrealEd/Private/CookOnTheFlyServer.cpp
+++ b/Engine/Source/Editor/UnrealEd/Private/CookOnTheFlyServer.cpp
@@ -21,7 +21,6 @@
 
 #include "AssetRegistryModule.h"
 #include "AssetData.h"
-#include "DependsNode.h"
 
 #include "UnrealEdMessages.h"
 #include "GameDelegates.h"
@@ -4746,8 +4745,7 @@ bool UCookOnTheFlyServer::GetAllPackagesFromAssetRegistry( const FString& AssetR
 	if (FFileHelper::LoadFileToArray(SerializedAssetData, *AssetRegistryPath))
 	{
 		TMap<FName, FAssetData*> RegistryDataMap;
-		TArray<FDependsNode*> Dependencies;
-		AssetRegistry.LoadRegistryData(SerializedAssetData, RegistryDataMap, Dependencies);
+		AssetRegistry.LoadRegistryData(SerializedAssetData, RegistryDataMap);
 
 
 		for (const auto& RegistryData : RegistryDataMap)
@@ -4770,10 +4768,6 @@ bool UCookOnTheFlyServer::GetAllPackagesFromAssetRegistry( const FString& AssetR
 		{
 			delete RegistryData.Value;
 		}
-		for (auto DependencyData : Dependencies)
-		{
-			delete DependencyData;
-		}
 		return true;
 	}
 
diff --git a/Engine/Source/Runtime/AssetRegistry/Private/AssetRegistry.cpp b/Engine/Source/Runtime/AssetRegistry/Private/AssetRegistry.cpp
--- a/Engine/Source/Runtime/AssetRegistry/Private/AssetRegistry.cpp
+++ b/Engine/Source/Runtime/AssetRegistry/Private/AssetRegistry.cpp
@@ -1675,7 +1675,7 @@ void FAssetRegistry::SaveRegistryData(FArchive& Ar, TMap<FName, FAssetData*>& Da
 	}
 }
 
-void FAssetRegistry::LoadRegistryData(FArchive& Ar, TMap<FName, FAssetData*>& Data, TArray<FDependsNode*>& OutDependencyData)
+void FAssetRegistry::LoadRegistryData(FArchive& Ar, TMap<FName, FAssetData*>& Data)
 {
 	auto Version = ReadRuntimeRegistryVersion(Ar);
 
@@ -1708,24 +1708,14 @@ void FAssetRegistry::LoadRegistryData(FArchive& Ar, TMap<FName, FAssetData*>& Da
 			Ar << DependsNodeCount;
 		}
 
-		OutDependencyData.Reserve(DependsNodeCount);
-
-		for (int32 DependsNodeIndex = 0; DependsNodeIndex < DependsNodeCount; DependsNodeIndex++)
+		// @todo - Dependency data is serialized locally then thrown away until we establish a proper way to expose it to external code.
+		if (Version == ERuntimeRegistryVersion::PreVersioning)
 		{
-			int32 AssetIndex = 0;
-
-			if (Version == ERuntimeRegistryVersion::PreVersioning)
+			for (int32 DependsNodeIndex = 0; DependsNodeIndex < DependsNodeCount; DependsNodeIndex++)
 			{
+				int32 AssetIndex = 0;
 				Ar << AssetIndex;
 			}
-			else
-			{
-				AssetIndex = OutDependencyData.Num();
-			}
-
-			auto PackageName = AssetIndexMap[AssetIndex];
-			auto NewNode = new FDependsNode(PackageName);
-			OutDependencyData.Add(NewNode);
 		}
 
 		for (int32 DependsNodeIndex = 0; DependsNodeIndex < DependsNodeCount; DependsNodeIndex++)
@@ -1742,30 +1732,22 @@ void FAssetRegistry::LoadRegistryData(FArchive& Ar, TMap<FName, FAssetData*>& Da
 			int32 NumReferencers = 0;
 			Ar << NumReferencers;
 
-			auto DependsNode = OutDependencyData[DependsNodeIndex];
-
 			for (int32 i = 0; i < NumHardDependencies; ++i)
 			{
 				int32 DependencyIndex = 0;
 				Ar << DependencyIndex;
-
-				DependsNode->AddDependency(OutDependencyData[DependencyIndex], EAssetRegistryDependencyType::Hard);
 			}
 
 			for (int32 i = 0; i < NumSoftDependencies; ++i)
 			{
 				int32 DependencyIndex = 0;
 				Ar << DependencyIndex;
-
-				DependsNode->AddDependency(OutDependencyData[DependencyIndex], EAssetRegistryDependencyType::Soft);
 			}
 
 			for (int32 i = 0; i < NumReferencers; ++i)
 			{
 				int32 Index = 0;
 				Ar << Index;
-
-				DependsNode->AddReferencer(OutDependencyData[Index]);
 			}
 		}
 	}
diff --git a/Engine/Source/Runtime/AssetRegistry/Private/AssetRegistry.h b/Engine/Source/Runtime/AssetRegistry/Private/AssetRegistry.h
--- a/Engine/Source/Runtime/AssetRegistry/Private/AssetRegistry.h
+++ b/Engine/Source/Runtime/AssetRegistry/Private/AssetRegistry.h
@@ -39,7 +39,7 @@ public:
 	virtual void PrioritizeSearchPath(const FString& PathToPrioritize) override;
 	virtual void Serialize(FArchive& Ar) override;
 	virtual void SaveRegistryData(FArchive& Ar, TMap<FName, FAssetData*>& Data, TArray<FName>* InMaps = nullptr) override;
-	virtual void LoadRegistryData(FArchive& Ar, TMap<FName, FAssetData*>& Data, TArray<FDependsNode*>& OutDependencyData) override;
+	virtual void LoadRegistryData(FArchive& Ar, TMap<FName, FAssetData*>& Data) override;
 
 	DECLARE_DERIVED_EVENT( FAssetRegistry, IAssetRegistry::FPathAddedEvent, FPathAddedEvent);
 	virtual FPathAddedEvent& OnPathAdded() override { return PathAddedEvent; }
diff --git a/Engine/Source/Runtime/AssetRegistry/Private/DependsNode.h b/Engine/Source/Runtime/AssetRegistry/Private/DependsNode.h
new file mode 100644
--- /dev/null
+++ b/Engine/Source/Runtime/AssetRegistry/Private/DependsNode.h
@@ -0,0 +1,124 @@
+// Copyright 1998-2015 Epic Games, Inc. All Rights Reserved.
+
+#pragma once
+
+/** Implementation of IDependsNode */
+class FDependsNode
+{
+public:
+	FDependsNode() { PackageName = NAME_None; }
+	FDependsNode(FName InPackageName);
+
+	/** Prints the dependencies and referencers for this node to the log */
+	void PrintNode() const;
+	/** Prints the dependencies for this node to the log */
+	void PrintDependencies() const;
+	/** Prints the referencers to this node to the log */
+	void PrintReferencers() const;
+	/** Gets the list of dependencies for this node */
+	void GetDependencies(TArray<FDependsNode*>& OutDependencies, EAssetRegistryDependencyType::Type InDependencyType = EAssetRegistryDependencyType::All) const;
+	/** Gets the list of dependency names for this node */
+	void GetDependencies(TArray<FName>& OutDependencies, EAssetRegistryDependencyType::Type InDependencyType = EAssetRegistryDependencyType::All) const;
+	/** Gets the list of referencers to this node */
+	void GetReferencers(TArray<FDependsNode*>& OutReferencers) const;
+	/** Gets the name of the package that this node represents */
+	FName GetPackageName() const { return PackageName; }
+	/** Sets the name of the package that this node represents */
+	void SetPackageName(FName InName) { PackageName = InName; }
+	/** Add a dependency to this node */
+	void AddDependency(FDependsNode* InDependency, EAssetRegistryDependencyType::Type InDependencyType = EAssetRegistryDependencyType::Hard)
+	{ 
+		if (InDependencyType == EAssetRegistryDependencyType::Hard)
+		{
+			HardDependencies.Add(InDependency);
+		}
+		else
+		{
+			SoftDependencies.Add(InDependency);
+		}
+	}
+	/** Add a referencer to this node */
+	void AddReferencer(FDependsNode* InReferencer) { Referencers.Add(InReferencer); }
+	/** Remove a dependency from this node */
+	void RemoveDependency(FDependsNode* InDependency) { HardDependencies.Remove(InDependency); SoftDependencies.Remove(InDependency); }
+	/** Remove a referencer from this node */
+	void RemoveReferencer(FDependsNode* InReferencer) { Referencers.Remove(InReferencer); }
+	/** Clear all dependency records from this node */
+	void ClearDependencies() { HardDependencies.Empty(); SoftDependencies.Empty(); }
+
+	/** Iterate over all the dependencies of this node, filtered by the supplied type parameter, and call the supplied lambda
+		parameter on the record */
+	template <class T>
+	void IterateOverDependencies(T InCallback, EAssetRegistryDependencyType::Type InDependencyType = EAssetRegistryDependencyType::All)
+	{
+		if (InDependencyType & EAssetRegistryDependencyType::Hard)
+		{
+			for (auto Dependency : HardDependencies)
+			{
+				if (Dependency)
+				{
+					InCallback(Dependency, EAssetRegistryDependencyType::Hard);
+				}
+			}
+		}
+
+		if (InDependencyType & EAssetRegistryDependencyType::Soft)
+		{
+			for (auto Dependency : SoftDependencies)
+			{
+				if (Dependency)
+				{
+					InCallback(Dependency, EAssetRegistryDependencyType::Soft);
+				}
+			}
+		}
+	}
+
+	/** Iterate over all the dependencies of this node, filtered by the supplied type parameter, and call the supplied lambda
+	parameter on the record */
+	template <class T>
+	void IterateOverDependencies(T InCallback, EAssetRegistryDependencyType::Type InDependencyType = EAssetRegistryDependencyType::All) const
+	{
+		if (InDependencyType & EAssetRegistryDependencyType::Hard)
+		{
+			for (auto Dependency : HardDependencies)
+			{
+				InCallback(Dependency, EAssetRegistryDependencyType::Hard);
+			}
+		}
+
+		if (InDependencyType & EAssetRegistryDependencyType::Soft)
+		{
+			for (auto Dependency : SoftDependencies)
+			{
+				InCallback(Dependency, EAssetRegistryDependencyType::Soft);
+			}
+		}
+	}
+
+	/** Iterate over all the referencers of this node and call the supplied lambda
+		parameter on the record */
+	template <class T>
+	void IterateOverReferencers(T InCallback)
+	{
+		for (auto Referencer : Referencers)
+		{
+			InCallback(Referencer);
+		}
+	}
+
+private:
+	/** Recursively prints dependencies of the node starting with the specified indent. VisitedNodes should be an empty set at first which is populated recursively. */
+	void PrintDependenciesRecursive(const FString& Indent, TSet<const FDependsNode*>& VisitedNodes) const;
+	/** Recursively prints referencers to the node starting with the specified indent. VisitedNodes should be an empty set at first which is populated recursively. */
+	void PrintReferencersRecursive(const FString& Indent, TSet<const FDependsNode*>& VisitedNodes) const;
+
+	/** The name of the package this node represents */
+	FName PackageName;
+	/** The list of hard dependencies for this node */
+	TSet<FDependsNode*> HardDependencies;
+	/** The list of soft dependencies for this node */
+	TSet<FDependsNode*> SoftDependencies;
+	/** The list of referencers to this node */
+	TSet<FDependsNode*> Referencers;
+};
\ No newline at end of file
diff --git a/Engine/Source/Runtime/AssetRegistry/Public/DependsNode.h b/Engine/Source/Runtime/AssetRegistry/Public/DependsNode.h
deleted file mode 100644
diff --git a/Engine/Source/Runtime/AssetRegistry/Public/IAssetRegistry.h b/Engine/Source/Runtime/AssetRegistry/Public/IAssetRegistry.h
--- a/Engine/Source/Runtime/AssetRegistry/Public/IAssetRegistry.h
+++ b/Engine/Source/Runtime/AssetRegistry/Public/IAssetRegistry.h
@@ -27,8 +27,6 @@ namespace EAssetAvailabilityProgressReportingType
 	};
 }
 
-class FDependsNode;
-
 class IAssetRegistry
 {
 public:
@@ -248,5 +246,5 @@ public:
 
 
 	/** Serialize registry data from a file */
-	virtual void LoadRegistryData(FArchive& Ar, TMap<FName, FAssetData*>& Data, TArray<FDependsNode*>& OutDependencyData) = 0;
+	virtual void LoadRegistryData(FArchive& Ar, TMap<FName, FAssetData*>& Data) = 0;
 };
