Improved the caching of FPathTree (as used by FAssetRegistry)

The CachePath function of FPathTree (formerly FPathTreeNode) used to always test each component part of a path to ensure that the path was available in the tree. This not only involved always performing an FName -> FString conversion for each path that was found by the asset registry, but it also resulted in repeated work being performed when the path was already fully cached.

The new FPathTree contains a map of full parent paths to full child paths (all using FName), which allows it to early out if it's already cached the path in question. This not only avoids the repeated work, but also avoids several FName -> FString -> FName conversions that used to happen between adding a path, and then retrieving it again later (the path would previously be rebuilt from its component FString parts, and then converted back into an FName - now it just uses the FName that we already have in the map).

These changes can dramatically improve the asset registry processing performance for large projects.

diff --git a/Engine/Source/Runtime/AssetRegistry/Private/AssetRegistry.cpp b/Engine/Source/Runtime/AssetRegistry/Private/AssetRegistry.cpp
--- a/Engine/Source/Runtime/AssetRegistry/Private/AssetRegistry.cpp
+++ b/Engine/Source/Runtime/AssetRegistry/Private/AssetRegistry.cpp
@@ -71,8 +71,7 @@ inline ERuntimeRegistryVersion ReadRuntimeRegistryVersion(FArchive& Ar)
 }
 
 FAssetRegistry::FAssetRegistry()
-	: PathTreeRoot(FString())
-	, PreallocatedAssetDataBuffer(NULL)
+	: PreallocatedAssetDataBuffer(NULL)
 	, PreallocatedDependsNodeDataBuffer(NULL)
 {
 	const double StartupStartTime = FPlatformTime::Seconds();
@@ -434,7 +433,7 @@ bool FAssetRegistry::GetAssets(const FARFilter& Filter, TArray<FAssetData>& OutA
 		// Add subpaths to all the input paths to the list
 		for ( int32 PathIdx = 0; PathIdx < NumFilterPackagePaths; ++PathIdx )
 		{
-			PathTreeRoot.GetSubPaths(Filter.PackagePaths[PathIdx].ToString(), FilterPackagePaths);
+			PathTreeRoot.GetSubPaths(Filter.PackagePaths[PathIdx], FilterPackagePaths);
 		}
 	}
 
@@ -948,7 +947,7 @@ void FAssetRegistry::GetDerivedClassNames(const TArray<FName>& ClassNames, const
 void FAssetRegistry::GetAllCachedPaths(TArray<FString>& OutPathList) const
 {
 	TSet<FName> PathList;
-	PathTreeRoot.GetSubPaths(TEXT(""), PathList);
+	PathTreeRoot.GetAllPaths(PathList);
 	
 	OutPathList.Empty(PathList.Num());
 	for ( auto PathIt = PathList.CreateConstIterator(); PathIt; ++PathIt )
@@ -960,7 +959,7 @@ void FAssetRegistry::GetAllCachedPaths(TArray<FString>& OutPathList) const
 void FAssetRegistry::GetSubPaths(const FString& InBasePath, TArray<FString>& OutPathList, bool bInRecurse) const
 {
 	TSet<FName> PathList;
-	PathTreeRoot.GetSubPaths(InBasePath, PathList, bInRecurse);
+	PathTreeRoot.GetSubPaths(FName(*InBasePath), PathList, bInRecurse);
 	
 	OutPathList.Empty(PathList.Num());
 	for ( auto PathIt = PathList.CreateConstIterator(); PathIt; ++PathIt )
@@ -1220,12 +1219,12 @@ void FAssetRegistry::PrioritizeAssetInstall(const FAssetData& AssetData) const
 
 bool FAssetRegistry::AddPath(const FString& PathToAdd)
 {
-	return AddAssetPath(PathToAdd);
+	return AddAssetPath(FName(*PathToAdd));
 }
 
 bool FAssetRegistry::RemovePath(const FString& PathToRemove)
 {
-	return RemoveAssetPath(PathToRemove);
+	return RemoveAssetPath(FName(*PathToRemove));
 }
 
 void FAssetRegistry::ScanPathsSynchronous(const TArray<FString>& InPaths, bool bForceRescan)
@@ -1248,7 +1247,8 @@ void FAssetRegistry::PrioritizeSearchPath(const FString& PathToPrioritize)
 		int32 LowestNonPriorityFileIdx = 0;
 		for (int32 ResultIdx = 0; ResultIdx < BackgroundAssetResults.Num(); ++ResultIdx)
 		{
-			if (BackgroundAssetResults[ResultIdx]->IsWithinSearchPath(PathToPrioritize))
+			IGatheredAssetData* BackgroundAssetResult = BackgroundAssetResults[ResultIdx];
+			if (BackgroundAssetResult && BackgroundAssetResult->IsWithinSearchPath(PathToPrioritize))
 			{
 				BackgroundAssetResults.Swap(ResultIdx, LowestNonPriorityFileIdx);
 				LowestNonPriorityFileIdx++;
@@ -1276,7 +1276,7 @@ void FAssetRegistry::AssetCreated(UObject* NewAsset)
 		RemoveEmptyPackage(NewPackage->GetFName());
 
 		// Add the path to the Path Tree, in case it wasn't already there
-		AddAssetPath( FPackageName::GetLongPackagePath(NewPackageName) );
+		AddAssetPath(*FPackageName::GetLongPackagePath(NewPackageName));
 
 		// Let subscribers know that the new asset was added to the registry
 		AssetAddedEvent.Broadcast(FAssetData(NewAsset));
@@ -1346,7 +1346,7 @@ void FAssetRegistry::AssetRenamed(const UObject* RenamedAsset, const FString& Ol
 		}
 
 		// Add the path to the Path Tree, in case it wasn't already there
-		AddAssetPath( FPackageName::GetLongPackagePath(NewPackageName) );
+		AddAssetPath(*FPackageName::GetLongPackagePath(NewPackageName));
 
 		AssetRenamedEvent.Broadcast(FAssetData(RenamedAsset), OldObjectPath);
 	}
@@ -1471,7 +1471,7 @@ void FAssetRegistry::Serialize(FArchive& Ar)
 
 			AddAssetData(NewAssetData);
 
-			AddAssetPath(NewAssetData->PackagePath.ToString());
+			AddAssetPath(NewAssetData->PackagePath);
 		}
 
 		int32 LocalNumDependsNodes = LocalNumAssets;
@@ -1847,7 +1847,7 @@ void FAssetRegistry::AssetSearchDataGathered(const double TickStartTime, TArray<
 		}
 
 		// Populate the path tree
-		AddAssetPath(Result.PackagePath.ToString());
+		AddAssetPath(Result.PackagePath);
 
 		// Delete the result that was originally created by an FPackageReader
 		delete BackgroundResult;
@@ -1876,7 +1876,7 @@ void FAssetRegistry::PathDataGathered(const double TickStartTime, TArray<FString
 	for (ResultIdx = 0; ResultIdx < PathResults.Num(); ++ResultIdx)
 	{
 		const FString& Path = PathResults[ResultIdx];
-		AddAssetPath(Path);
+		AddAssetPath(FName(*Path));
 		
 		// Check to see if we have run out of time in this tick
 		if ( !bFlushFullBuffer && (FPlatformTime::Seconds() - TickStartTime) > MaxSecondsPerFrame)
@@ -2031,24 +2031,24 @@ bool FAssetRegistry::RemoveEmptyPackage(FName PackageName)
 	return CachedEmptyPackages.Remove( PackageName ) > 0;
 }
 
-bool FAssetRegistry::AddAssetPath(const FString& PathToAdd)
+bool FAssetRegistry::AddAssetPath(FName PathToAdd)
 {
 	if ( PathTreeRoot.CachePath(PathToAdd) )
 	{		
-		PathAddedEvent.Broadcast(PathToAdd);
+		PathAddedEvent.Broadcast(PathToAdd.ToString());
 		return true;
 	}
 
 	return false;
 }
 
-bool FAssetRegistry::RemoveAssetPath(const FString& PathToRemove, bool bEvenIfAssetsStillExist)
+bool FAssetRegistry::RemoveAssetPath(FName PathToRemove, bool bEvenIfAssetsStillExist)
 {
 	if ( !bEvenIfAssetsStillExist )
 	{
 		// Check if there were assets in the specified folder. You can not remove paths that still contain assets
 		TArray<FAssetData> AssetsInPath;
-		GetAssetsByPath(FName(*PathToRemove), AssetsInPath, true);
+		GetAssetsByPath(PathToRemove, AssetsInPath, true);
 		if ( AssetsInPath.Num() > 0 )
 		{
 			// At least one asset still exists in the path. Fail the remove.
@@ -2056,9 +2056,9 @@ bool FAssetRegistry::RemoveAssetPath(const FString& PathToRemove, bool bEvenIfAs
 		}
 	}
 
-	if ( PathTreeRoot.RemoveFolder(PathToRemove) )
+	if ( PathTreeRoot.RemovePath(PathToRemove) )
 	{
-		PathRemovedEvent.Broadcast(PathToRemove);
+		PathRemovedEvent.Broadcast(PathToRemove.ToString());
 		return true;
 	}
 	else
@@ -2434,7 +2434,7 @@ void FAssetRegistry::OnContentPathDismounted(const FString& InAssetPath, const F
 	// Remove the root path
 	{
 		const bool bEvenIfAssetsStillExist = true;
-		RemoveAssetPath(AssetPath, bEvenIfAssetsStillExist);
+		RemoveAssetPath(FName(*AssetPath), bEvenIfAssetsStillExist);
 	}
 
 	// Stop listening for directory changes in this content path
diff --git a/Engine/Source/Runtime/AssetRegistry/Private/AssetRegistry.h b/Engine/Source/Runtime/AssetRegistry/Private/AssetRegistry.h
--- a/Engine/Source/Runtime/AssetRegistry/Private/AssetRegistry.h
+++ b/Engine/Source/Runtime/AssetRegistry/Private/AssetRegistry.h
@@ -109,10 +109,10 @@ private:
 	bool RemoveEmptyPackage(FName PackageName);
 
 	/** Adds a path to the cached paths tree. Returns true if the path was added to the tree, as opposed to already existing in the tree */
-	bool AddAssetPath(const FString& PathToAdd);
+	bool AddAssetPath(FName PathToAdd);
 
 	/** Removes a path to the cached paths tree. Returns true if successful. */
-	bool RemoveAssetPath(const FString& PathToAdd, bool bEvenIfAssetsStillExist = false);
+	bool RemoveAssetPath(FName PathToRemove, bool bEvenIfAssetsStillExist = false);
 
 	/** Helper function to return the name of an object, given the objects export text path */
 	FString ExportTextPathToObjectName(const FString& InExportTextPath) const;
@@ -201,7 +201,7 @@ private:
 	TMap<FName, TSet<FName>> CookWhitelistedTagsByClass;
 
 	/** The root node to the path tree */
-	FPathTreeNode PathTreeRoot;
+	FPathTree PathTreeRoot;
 
 	/** Async task that gathers asset information from disk */
 	TSharedPtr< class FAssetDataGatherer > BackgroundAssetSearch;
diff --git a/Engine/Source/Runtime/AssetRegistry/Private/PathTreeNode.cpp b/Engine/Source/Runtime/AssetRegistry/Private/PathTreeNode.cpp
--- a/Engine/Source/Runtime/AssetRegistry/Private/PathTreeNode.cpp
+++ b/Engine/Source/Runtime/AssetRegistry/Private/PathTreeNode.cpp
@@ -2,358 +2,149 @@
 
 #include "AssetRegistryPCH.h"
 
-bool IsUsingExperimentalFNameTreeNodes()
+bool FPathTree::CachePath(const FName& Path)
 {
-	static bool bUsingExperimentalFNameTreeNodes = FParse::Param(FCommandLine::Get(), TEXT("PathTreeFNames"));
-	return bUsingExperimentalFNameTreeNodes;
-}
-
-FPathTreeNode::FPathTreeNode(const FString& InFolderName)
-	: FolderName(InFolderName)
-	, FolderFName(NAME_None)
-{}
-
-FPathTreeNode::FPathTreeNode(FName InFolderName)
-	: FolderFName(InFolderName)
-{}
-
-FPathTreeNode::~FPathTreeNode()
-{
-	// Every node is responsible for deleting its children
-	for (int32 ChildIdx = 0; ChildIdx < Children.Num(); ++ChildIdx)
+	if (Path.IsNone())
 	{
-		delete Children[ChildIdx];
+		return false;
 	}
 
-	Children.Empty();
-}
-
-bool FPathTreeNode::CachePath(const FString& Path)
-{
-	if (IsUsingExperimentalFNameTreeNodes())
+	if (ParentPathToChildPaths.Contains(Path))
 	{
-		TArray<FString> PathElementStrings;
-		Path.ParseIntoArray(PathElementStrings, TEXT("/"), /*InCullEmpty=*/true);
-
-		TArray<FName> PathElements;
-		PathElements.Reserve(PathElementStrings.Num());
-		for (const auto& Element : PathElementStrings)
-		{
-			PathElements.Add(FName(*Element));
-		}
-
-		return CachePath_Recursive(PathElements);
+		// Already cached - nothing more to do
+		return false;
 	}
-	else
-	{
-		TArray<FString> PathElements;
-		Path.ParseIntoArray(PathElements, TEXT("/"), /*InCullEmpty=*/true);
 
-		return CachePath_Recursive(PathElements);
-	}
-}
+	FString PathStr = Path.ToString();
+	check(PathStr.Len() >= 2);	// Must be at least "/A"
+	check(PathStr[0] == '/');	// Must start with a "/"
 
-bool FPathTreeNode::RemoveFolder(const FString& Path)
-{
-	if (IsUsingExperimentalFNameTreeNodes())
+	// Walk each part of the path, adding known path entries if required
+	// This manipulates PathStr in-place to avoid making any string copies
+	FName LastPath;
+	TCHAR* PathCharPtr = &PathStr[1]; // Skip over the first / when scanning
+	for (;;)
 	{
-		TArray<FString> PathElementStrings;
-		Path.ParseIntoArray(PathElementStrings, TEXT("/"), /*InCullEmpty=*/true);
-
-		TArray<FName> PathElements;
-		PathElements.Reserve(PathElementStrings.Num());
-		for (const auto& Element : PathElementStrings)
+		const TCHAR PathChar = *PathCharPtr;
+		if (PathChar == '/' || PathChar == 0)
 		{
-			PathElements.Add(FName(*Element));
-		}
+			// We've found a path separator (or the end of the string), so process this part of the path
+			(*PathCharPtr) = 0;			// Null terminate this part of the string so we can create an FName from it
+			const FName CurrentPath = *PathStr;
+			(*PathCharPtr) = PathChar;	// Restore the original character now
 
-		return RemoveFolder_Recursive(PathElements);
-	}
-	else
-	{
-		TArray<FString> PathElements;
-		Path.ParseIntoArray(PathElements, TEXT("/"), /*InCullEmpty=*/true);
+			check(!CurrentPath.IsNone());	// Path parts cannot be empty
+			check(*(PathCharPtr-1) != '/'); // The previous character cannot be a /, as that would suggest a malformed path such as "/Game//MyAsset"
 
-		return RemoveFolder_Recursive(PathElements);
-	}
-}
-
-bool FPathTreeNode::GetSubPaths(const FString& BasePath, TSet<FName>& OutPaths, bool bRecurse) const
-{
-	if (IsUsingExperimentalFNameTreeNodes())
-	{
-		TArray<FString> PathElementStrings;
-		BasePath.ParseIntoArray(PathElementStrings, TEXT("/"), /*InCullEmpty=*/true);
-
-		TArray<FName> PathElements;
-		PathElements.Reserve(PathElementStrings.Num());
-		for (const auto& Element : PathElementStrings)
-		{
-			PathElements.Add(FName(*Element));
-		}
-
-		const FPathTreeNode* BasePathNode = FindNode_Recursive(PathElements);
-
-		if (BasePathNode)
-		{
-			// Found the base path, get the paths of all its children
-			for (int32 ChildIdx = 0; ChildIdx < BasePathNode->Children.Num(); ++ChildIdx)
+			if (!ParentPathToChildPaths.Contains(CurrentPath))
 			{
-				BasePathNode->Children[ChildIdx]->GetSubPaths_Recursive(BasePath, OutPaths, bRecurse);
+				ParentPathToChildPaths.Add(CurrentPath);
 			}
 
-			return true;
-		}
-		else
-		{
-			// Failed to find the base path
-			return false;
-		}
-	}
-	else
-	{
-		TArray<FString> PathElements;
-		BasePath.ParseIntoArray(PathElements, TEXT("/"), /*InCullEmpty=*/true);
-
-		const FPathTreeNode* BasePathNode = FindNode_Recursive(PathElements);
-
-		if ( BasePathNode )
-		{
-			// Found the base path, get the paths of all its children
-			for (int32 ChildIdx = 0; ChildIdx < BasePathNode->Children.Num(); ++ChildIdx)
+			if (!LastPath.IsNone())
 			{
-				BasePathNode->Children[ChildIdx]->GetSubPaths_Recursive(BasePath, OutPaths, bRecurse);
-			}
-		
-			return true;
-		}
-		else
-		{
-			// Failed to find the base path
-			return false;
-		}
-	}
-}
+				// Add us as a known child of our parent path
+				TSet<FName>& ChildPaths = ParentPathToChildPaths.FindChecked(LastPath);
+				ChildPaths.Add(CurrentPath);
 
-bool FPathTreeNode::CachePath_Recursive(TArray<FString>& PathElements)
-{
-	if ( PathElements.Num() > 0 )
-	{
-		// Pop the bottom element
-		FString ChildFolderName = PathElements[0];
-		PathElements.RemoveAt(0);
-
-		// Try to find a child which uses this folder name
-		FPathTreeNode* Child = NULL;
-		for (int32 ChildIdx = 0; ChildIdx < Children.Num(); ++ChildIdx)
-		{
-			if ( Children[ChildIdx]->FolderName == ChildFolderName )
-			{
-				Child = Children[ChildIdx];
-				break;
+				// Make sure we know how to find our parent again later on
+				ChildPathToParentPath.Add(CurrentPath, LastPath);
 			}
-		}
-
-		// If one was not found, create it
-		if ( !Child )
-		{
-			int32 ChildIdx = Children.Add(new FPathTreeNode(ChildFolderName));
-			Child = Children[ChildIdx];
 
-			if ( PathElements.Num() == 0 )
-			{
-				// We added final element to the tree, return true;
-				return true;
-			}
+			LastPath = CurrentPath;
 		}
 
-		if ( ensure(Child) )
+		if (PathChar == 0)
 		{
-			return Child->CachePath_Recursive(PathElements);
+			// End of the string
+			break;
 		}
+
+		++PathCharPtr;
 	}
 
-	return false;
+	return true;
 }
 
-bool FPathTreeNode::RemoveFolder_Recursive(TArray<FString>& PathElements)
+bool FPathTree::RemovePath(const FName& Path)
 {
-	if ( PathElements.Num() > 0 )
+	if (Path.IsNone())
 	{
-		// Pop the bottom element
-		FString ChildFolderName = PathElements[0];
-		PathElements.RemoveAt(0);
-
-		// Try to find a child which uses this folder name
-		for (int32 ChildIdx = 0; ChildIdx < Children.Num(); ++ChildIdx)
-		{
-			if ( Children[ChildIdx]->FolderName == ChildFolderName )
-			{
-				FPathTreeNode* Child = Children[ChildIdx];
-
-				if ( PathElements.Num() == 0 )
-				{
-					// This is the final child in the path, thus it is the folder to remove.
-					// Remove it now and return true;
-					Children.RemoveAt(ChildIdx);
-					delete Child;
-					return true;
-				}
-				else
-				{
-					// This is not the last child, so recurse below
-					return Child->RemoveFolder_Recursive(PathElements);
-				}
-			}
-		}
+		return false;
 	}
 
-	// The child was not found at some point while recursing, fail the remove.
-	return false;
-}
-
-const FPathTreeNode* FPathTreeNode::FindNode_Recursive(TArray<FString>& PathElements) const
-{
-	if (PathElements.Num() == 0)
+	if (!ParentPathToChildPaths.Contains(Path))
 	{
-		// Found the path node
-		return this;
+		// Doesn't exist - nothing more to do
+		return false;
 	}
-	else
-	{
-		// Must drill down further
-		// Pop the bottom element
-		FString ChildFolderName = PathElements[0];
-		PathElements.RemoveAt(0);
 
-		for (int32 ChildIdx = 0; ChildIdx < Children.Num(); ++ChildIdx)
-		{
-			if ( Children[ChildIdx]->FolderName == ChildFolderName )
-			{
-				return Children[ChildIdx]->FindNode_Recursive(PathElements);
-			}
-		}
+	// We also need to gather up and remove any children of this path
+	TSet<FName> SubPathsToRemove;
+	GetSubPaths(Path, SubPathsToRemove, /*bRecurse=*/true);
 
-		// Could not find the path
-		return NULL;
+	// Simply remove sub-paths from both maps
+	for (const FName& SubPathToRemove : SubPathsToRemove)
+	{
+		ParentPathToChildPaths.Remove(SubPathToRemove);
+		ChildPathToParentPath.Remove(SubPathToRemove);
 	}
-}
 
-bool FPathTreeNode::CachePath_Recursive(TArray<FName>& PathElements)
-{
-	if ( PathElements.Num() > 0 )
+	// We also need to remove ourself from our parent list before removing ourself from the maps
 	{
-		// Pop the bottom element
-		FName ChildFolderName = PathElements[0];
-		PathElements.RemoveAt(0);
-
-		// Try to find a child which uses this folder name
-		FPathTreeNode* Child = NULL;
-		for (int32 ChildIdx = 0; ChildIdx < Children.Num(); ++ChildIdx)
+		const FName* ParentPathPtr = ChildPathToParentPath.Find(Path);
+		if (ParentPathPtr)
 		{
-			if ( Children[ChildIdx]->FolderFName == ChildFolderName )
-			{
-				Child = Children[ChildIdx];
-				break;
-			}
-		}
-
-		// If one was not found, create it
-		if ( !Child )
-		{
-			int32 ChildIdx = Children.Add(new FPathTreeNode(ChildFolderName));
-			Child = Children[ChildIdx];
-
-			if ( PathElements.Num() == 0 )
-			{
-				// We added final element to the tree, return true;
-				return true;
-			}
-		}
-
-		if ( ensure(Child) )
-		{
-			return Child->CachePath_Recursive(PathElements);
+			TSet<FName>& ChildPaths = ParentPathToChildPaths.FindChecked(*ParentPathPtr);
+			ChildPaths.Remove(Path);
 		}
 	}
 
-	return false;
+	ParentPathToChildPaths.Remove(Path);
+	ChildPathToParentPath.Remove(Path);
+
+	return true;
 }
 
-bool FPathTreeNode::RemoveFolder_Recursive(TArray<FName>& PathElements)
+bool FPathTree::GetAllPaths(TSet<FName>& OutPaths) const
 {
-	if ( PathElements.Num() > 0 )
-	{
-		// Pop the bottom element
-		FName ChildFolderName = PathElements[0];
-		PathElements.RemoveAt(0);
+	OutPaths.Reset();
+	OutPaths.Reserve(ParentPathToChildPaths.Num());
 
-		// Try to find a child which uses this folder name
-		for (int32 ChildIdx = 0; ChildIdx < Children.Num(); ++ChildIdx)
-		{
-			if ( Children[ChildIdx]->FolderFName == ChildFolderName )
-			{
-				FPathTreeNode* Child = Children[ChildIdx];
-
-				if ( PathElements.Num() == 0 )
-				{
-					// This is the final child in the path, thus it is the folder to remove.
-					// Remove it now and return true;
-					Children.RemoveAt(ChildIdx);
-					delete Child;
-					return true;
-				}
-				else
-				{
-					// This is not the last child, so recurse below
-					return Child->RemoveFolder_Recursive(PathElements);
-				}
-			}
-		}
+	for (const auto& PathPair : ParentPathToChildPaths)
+	{
+		OutPaths.Add(PathPair.Key);
 	}
 
-	// The child was not found at some point while recursing, fail the remove.
-	return false;
+	return OutPaths.Num() > 0;
 }
 
-const FPathTreeNode* FPathTreeNode::FindNode_Recursive(TArray<FName>& PathElements) const
+bool FPathTree::GetSubPaths(const FName& BasePath, TSet<FName>& OutPaths, bool bRecurse) const
 {
-	if (PathElements.Num() == 0)
+	if (BasePath.IsNone())
 	{
-		// Found the path node
-		return this;
+		return false;
 	}
-	else
-	{
-		// Must drill down further
-		// Pop the bottom element
-		FName ChildFolderName = PathElements[0];
-		PathElements.RemoveAt(0);
 
-		for (int32 ChildIdx = 0; ChildIdx < Children.Num(); ++ChildIdx)
-		{
-			if ( Children[ChildIdx]->FolderFName == ChildFolderName )
-			{
-				return Children[ChildIdx]->FindNode_Recursive(PathElements);
-			}
-		}
-
-		// Could not find the path
-		return NULL;
+	const TSet<FName>* ChildPathsPtr = ParentPathToChildPaths.Find(BasePath);
+	if (!ChildPathsPtr)
+	{
+		return false;
 	}
-}
-
 
-void FPathTreeNode::GetSubPaths_Recursive(const FString& CurrentPath, TSet<FName>& OutPaths, bool bRecurse) const
-{
-	FString NewPath = IsUsingExperimentalFNameTreeNodes() ? (CurrentPath / FolderFName.ToString()) : (CurrentPath / FolderName);
-	OutPaths.Add(FName(*NewPath));
+	const int32 OutPathsOriginalNum = OutPaths.Num();
 
-	if(bRecurse)
+	for (const FName& ChildPath : *ChildPathsPtr)
 	{
-		for (int32 ChildIdx = 0; ChildIdx < Children.Num(); ++ChildIdx)
+		check(ParentPathToChildPaths.Contains(ChildPath)); // This failing is an integrity violation as this entry lists a child that we don't know about
+
+		OutPaths.Add(ChildPath);
+
+		if (bRecurse)
 		{
-			Children[ChildIdx]->GetSubPaths_Recursive(NewPath, OutPaths, bRecurse);
+			GetSubPaths(ChildPath, OutPaths, /*bRecurse=*/true);
 		}
 	}
+
+	return OutPaths.Num() > OutPathsOriginalNum;
 }
diff --git a/Engine/Source/Runtime/AssetRegistry/Private/PathTreeNode.h b/Engine/Source/Runtime/AssetRegistry/Private/PathTreeNode.h
--- a/Engine/Source/Runtime/AssetRegistry/Private/PathTreeNode.h
+++ b/Engine/Source/Runtime/AssetRegistry/Private/PathTreeNode.h
@@ -2,46 +2,25 @@
 
 #pragma once
 
-class FPathTreeNode
+class FPathTree
 {
 public:
-	FPathTreeNode(const FString& InFolderName);
-	FPathTreeNode(FName InFolderName);
-	~FPathTreeNode();
-
 	/** Adds the path to the tree relative to this node, creating nodes as needed. Returns true if the specified path was actually added (as opposed to already existed) */
-	bool CachePath(const FString& Path);
-
-	/** Removes the specified path in the tree relative to this node. Returns true if the folder was found and removed. */
-	bool RemoveFolder(const FString& Path);
-
-	/** Recursively gathers all child paths from the specified base path relative to this node */
-	bool GetSubPaths(const FString& BasePath, TSet<FName>& OutPaths, bool bRecurse = true) const;
+	bool CachePath(const FName& Path);
 
-private:
-	/** Helper function for CachePath. PathElements is the tokenized path delimited by "/" */
-	bool CachePath_Recursive(TArray<FString>& PathElements);
-
-	/** Helper function for RemoveFolder. PathElements is the tokenized path delimited by "/" */
-	bool RemoveFolder_Recursive(TArray<FString>& PathElements);
+	/** Removes the specified path in the tree relative to this node. Returns true if the path was found and removed. */
+	bool RemovePath(const FName& Path);
 
-	/** Finds a node specified by the supplied tokenized PathElements list */
-	const FPathTreeNode* FindNode_Recursive(TArray<FString>& PathElements) const;
+	/** Get all of the paths we know about */
+	bool GetAllPaths(TSet<FName>& OutPaths) const;
 
-	/** Helper function for recursive all subpath gathering. Currentpath is the string version of the parent node. */
-	void GetSubPaths_Recursive(const FString& CurrentPath, TSet<FName>& OutPaths, bool bRecurse = true) const;
+	/** Recursively gathers all child paths from the specified base path relative to this node */
+	bool GetSubPaths(const FName& BasePath, TSet<FName>& OutPaths, bool bRecurse = true) const;
 
 private:
-	// Experimental FName version of PathTreeNode functions. Only used when launching with -PathTreeFNames
-	bool CachePath_Recursive(TArray<FName>& PathElements);
-	bool RemoveFolder_Recursive(TArray<FName>& PathElements);
-	const FPathTreeNode* FindNode_Recursive(TArray<FName>& PathElements) const;
+	/** A one-to-many mapping between a parent path and its child paths. */
+	TMap<FName, TSet<FName>> ParentPathToChildPaths;
 
-private:
-	/** The short folder name in the tree (no path) */
-	FString FolderName;
-	/** Experimental FName folder */
-	FName FolderFName;
-	/** The child node list for this node */
-	TArray<FPathTreeNode*> Children;
-};
\ No newline at end of file
+	/** A one-to-one mapping between a child path and its parent path. Paths without a parent (root paths) will not appear in this map. */
+	TMap<FName, FName> ChildPathToParentPath;
+};
