Finished renaming FPathTreeNode to FPathTree

diff --git a/Engine/Source/Runtime/AssetRegistry/Private/AssetRegistry.cpp b/Engine/Source/Runtime/AssetRegistry/Private/AssetRegistry.cpp
--- a/Engine/Source/Runtime/AssetRegistry/Private/AssetRegistry.cpp
+++ b/Engine/Source/Runtime/AssetRegistry/Private/AssetRegistry.cpp
@@ -433,7 +433,7 @@ bool FAssetRegistry::GetAssets(const FARFilter& Filter, TArray<FAssetData>& OutA
 		// Add subpaths to all the input paths to the list
 		for ( int32 PathIdx = 0; PathIdx < NumFilterPackagePaths; ++PathIdx )
 		{
-			PathTreeRoot.GetSubPaths(Filter.PackagePaths[PathIdx], FilterPackagePaths);
+			CachedPathTree.GetSubPaths(Filter.PackagePaths[PathIdx], FilterPackagePaths);
 		}
 	}
 
@@ -947,7 +947,7 @@ void FAssetRegistry::GetDerivedClassNames(const TArray<FName>& ClassNames, const
 void FAssetRegistry::GetAllCachedPaths(TArray<FString>& OutPathList) const
 {
 	TSet<FName> PathList;
-	PathTreeRoot.GetAllPaths(PathList);
+	CachedPathTree.GetAllPaths(PathList);
 	
 	OutPathList.Empty(PathList.Num());
 	for ( auto PathIt = PathList.CreateConstIterator(); PathIt; ++PathIt )
@@ -959,7 +959,7 @@ void FAssetRegistry::GetAllCachedPaths(TArray<FString>& OutPathList) const
 void FAssetRegistry::GetSubPaths(const FString& InBasePath, TArray<FString>& OutPathList, bool bInRecurse) const
 {
 	TSet<FName> PathList;
-	PathTreeRoot.GetSubPaths(FName(*InBasePath), PathList, bInRecurse);
+	CachedPathTree.GetSubPaths(FName(*InBasePath), PathList, bInRecurse);
 	
 	OutPathList.Empty(PathList.Num());
 	for ( auto PathIt = PathList.CreateConstIterator(); PathIt; ++PathIt )
@@ -2033,7 +2033,7 @@ bool FAssetRegistry::RemoveEmptyPackage(FName PackageName)
 
 bool FAssetRegistry::AddAssetPath(FName PathToAdd)
 {
-	if ( PathTreeRoot.CachePath(PathToAdd) )
+	if (CachedPathTree.CachePath(PathToAdd))
 	{		
 		PathAddedEvent.Broadcast(PathToAdd.ToString());
 		return true;
@@ -2056,7 +2056,7 @@ bool FAssetRegistry::RemoveAssetPath(FName PathToRemove, bool bEvenIfAssetsStill
 		}
 	}
 
-	if ( PathTreeRoot.RemovePath(PathToRemove) )
+	if (CachedPathTree.RemovePath(PathToRemove))
 	{
 		PathRemovedEvent.Broadcast(PathToRemove.ToString());
 		return true;
diff --git a/Engine/Source/Runtime/AssetRegistry/Private/AssetRegistry.h b/Engine/Source/Runtime/AssetRegistry/Private/AssetRegistry.h
--- a/Engine/Source/Runtime/AssetRegistry/Private/AssetRegistry.h
+++ b/Engine/Source/Runtime/AssetRegistry/Private/AssetRegistry.h
@@ -200,8 +200,8 @@ private:
 	/** The map of classname to tag set of tags that are allowed in cooked builds */
 	TMap<FName, TSet<FName>> CookWhitelistedTagsByClass;
 
-	/** The root node to the path tree */
-	FPathTree PathTreeRoot;
+	/** The tree of known cached paths that assets may reside within */
+	FPathTree CachedPathTree;
 
 	/** Async task that gathers asset information from disk */
 	TSharedPtr< class FAssetDataGatherer > BackgroundAssetSearch;
diff --git a/Engine/Source/Runtime/AssetRegistry/Private/AssetRegistryPCH.h b/Engine/Source/Runtime/AssetRegistry/Private/AssetRegistryPCH.h
--- a/Engine/Source/Runtime/AssetRegistry/Private/AssetRegistryPCH.h
+++ b/Engine/Source/Runtime/AssetRegistry/Private/AssetRegistryPCH.h
@@ -11,7 +11,7 @@ DECLARE_LOG_CATEGORY_EXTERN(LogAssetRegistry, Log, All);
 #include "NameTableArchive.h"
 #include "IAssetRegistry.h"
 #include "DependsNode.h"
-#include "PathTreeNode.h"
+#include "PathTree.h"
 #include "BackgroundAssetData.h"
 #include "PackageDependencyData.h"
 #include "DiskCachedAssetData.h"
diff --git a/Engine/Source/Runtime/AssetRegistry/Private/PathTree.cpp b/Engine/Source/Runtime/AssetRegistry/Private/PathTree.cpp
new file mode 100644
--- /dev/null
+++ b/Engine/Source/Runtime/AssetRegistry/Private/PathTree.cpp
@@ -0,0 +1,150 @@
+// Copyright 1998-2015 Epic Games, Inc. All Rights Reserved.
+
+#include "AssetRegistryPCH.h"
+
+bool FPathTree::CachePath(const FName& Path)
+{
+	if (Path.IsNone())
+	{
+		return false;
+	}
+
+	if (ParentPathToChildPaths.Contains(Path))
+	{
+		// Already cached - nothing more to do
+		return false;
+	}
+
+	FString PathStr = Path.ToString();
+	check(PathStr.Len() >= 2);	// Must be at least "/A"
+	check(PathStr[0] == '/');	// Must start with a "/"
+
+	// Walk each part of the path, adding known path entries if required
+	// This manipulates PathStr in-place to avoid making any string copies
+	FName LastPath;
+	TCHAR* PathCharPtr = &PathStr[1]; // Skip over the first / when scanning
+	for (;;)
+	{
+		const TCHAR PathChar = *PathCharPtr;
+		if (PathChar == '/' || PathChar == 0)
+		{
+			// We've found a path separator (or the end of the string), so process this part of the path
+			(*PathCharPtr) = 0;			// Null terminate this part of the string so we can create an FName from it
+			const FName CurrentPath = *PathStr;
+			(*PathCharPtr) = PathChar;	// Restore the original character now
+
+			check(!CurrentPath.IsNone());	// Path parts cannot be empty
+			check(*(PathCharPtr-1) != '/'); // The previous character cannot be a /, as that would suggest a malformed path such as "/Game//MyAsset"
+
+			if (!ParentPathToChildPaths.Contains(CurrentPath))
+			{
+				ParentPathToChildPaths.Add(CurrentPath);
+			}
+
+			if (!LastPath.IsNone())
+			{
+				// Add us as a known child of our parent path
+				TSet<FName>& ChildPaths = ParentPathToChildPaths.FindChecked(LastPath);
+				ChildPaths.Add(CurrentPath);
+
+				// Make sure we know how to find our parent again later on
+				ChildPathToParentPath.Add(CurrentPath, LastPath);
+			}
+
+			LastPath = CurrentPath;
+		}
+
+		if (PathChar == 0)
+		{
+			// End of the string
+			break;
+		}
+
+		++PathCharPtr;
+	}
+
+	return true;
+}
+
+bool FPathTree::RemovePath(const FName& Path)
+{
+	if (Path.IsNone())
+	{
+		return false;
+	}
+
+	if (!ParentPathToChildPaths.Contains(Path))
+	{
+		// Doesn't exist - nothing more to do
+		return false;
+	}
+
+	// We also need to gather up and remove any children of this path
+	TSet<FName> SubPathsToRemove;
+	GetSubPaths(Path, SubPathsToRemove, /*bRecurse=*/true);
+
+	// Simply remove sub-paths from both maps
+	for (const FName& SubPathToRemove : SubPathsToRemove)
+	{
+		ParentPathToChildPaths.Remove(SubPathToRemove);
+		ChildPathToParentPath.Remove(SubPathToRemove);
+	}
+
+	// We also need to remove ourself from our parent list before removing ourself from the maps
+	{
+		const FName* ParentPathPtr = ChildPathToParentPath.Find(Path);
+		if (ParentPathPtr)
+		{
+			TSet<FName>& ChildPaths = ParentPathToChildPaths.FindChecked(*ParentPathPtr);
+			ChildPaths.Remove(Path);
+		}
+	}
+
+	ParentPathToChildPaths.Remove(Path);
+	ChildPathToParentPath.Remove(Path);
+
+	return true;
+}
+
+bool FPathTree::GetAllPaths(TSet<FName>& OutPaths) const
+{
+	OutPaths.Reset();
+	OutPaths.Reserve(ParentPathToChildPaths.Num());
+
+	for (const auto& PathPair : ParentPathToChildPaths)
+	{
+		OutPaths.Add(PathPair.Key);
+	}
+
+	return OutPaths.Num() > 0;
+}
+
+bool FPathTree::GetSubPaths(const FName& BasePath, TSet<FName>& OutPaths, bool bRecurse) const
+{
+	if (BasePath.IsNone())
+	{
+		return false;
+	}
+
+	const TSet<FName>* ChildPathsPtr = ParentPathToChildPaths.Find(BasePath);
+	if (!ChildPathsPtr)
+	{
+		return false;
+	}
+
+	const int32 OutPathsOriginalNum = OutPaths.Num();
+
+	for (const FName& ChildPath : *ChildPathsPtr)
+	{
+		check(ParentPathToChildPaths.Contains(ChildPath)); // This failing is an integrity violation as this entry lists a child that we don't know about
+
+		OutPaths.Add(ChildPath);
+
+		if (bRecurse)
+		{
+			GetSubPaths(ChildPath, OutPaths, /*bRecurse=*/true);
+		}
+	}
+
+	return OutPaths.Num() > OutPathsOriginalNum;
+}
diff --git a/Engine/Source/Runtime/AssetRegistry/Private/PathTree.h b/Engine/Source/Runtime/AssetRegistry/Private/PathTree.h
new file mode 100644
--- /dev/null
+++ b/Engine/Source/Runtime/AssetRegistry/Private/PathTree.h
@@ -0,0 +1,26 @@
+// Copyright 1998-2015 Epic Games, Inc. All Rights Reserved.
+
+#pragma once
+
+class FPathTree
+{
+public:
+	/** Adds the path to the tree relative to this node, creating nodes as needed. Returns true if the specified path was actually added (as opposed to already existed) */
+	bool CachePath(const FName& Path);
+
+	/** Removes the specified path in the tree relative to this node. Returns true if the path was found and removed. */
+	bool RemovePath(const FName& Path);
+
+	/** Get all of the paths we know about */
+	bool GetAllPaths(TSet<FName>& OutPaths) const;
+
+	/** Recursively gathers all child paths from the specified base path relative to this node */
+	bool GetSubPaths(const FName& BasePath, TSet<FName>& OutPaths, bool bRecurse = true) const;
+
+private:
+	/** A one-to-many mapping between a parent path and its child paths. */
+	TMap<FName, TSet<FName>> ParentPathToChildPaths;
+
+	/** A one-to-one mapping between a child path and its parent path. Paths without a parent (root paths) will not appear in this map. */
+	TMap<FName, FName> ChildPathToParentPath;
+};
diff --git a/Engine/Source/Runtime/AssetRegistry/Private/PathTreeNode.cpp b/Engine/Source/Runtime/AssetRegistry/Private/PathTreeNode.cpp
deleted file mode 100644
diff --git a/Engine/Source/Runtime/AssetRegistry/Private/PathTreeNode.h b/Engine/Source/Runtime/AssetRegistry/Private/PathTreeNode.h
deleted file mode 100644
