Fixed a case where FAssetRegistry::RemoveDependsNode could potentially delete an allocation from within an array allocation

Also replaced NULL usage with nullptr inside FAssetRegistry.

diff --git a/Engine/Source/Runtime/AssetRegistry/Private/AssetRegistry.cpp b/Engine/Source/Runtime/AssetRegistry/Private/AssetRegistry.cpp
--- a/Engine/Source/Runtime/AssetRegistry/Private/AssetRegistry.cpp
+++ b/Engine/Source/Runtime/AssetRegistry/Private/AssetRegistry.cpp
@@ -71,8 +71,8 @@ inline ERuntimeRegistryVersion ReadRuntimeRegistryVersion(FArchive& Ar)
 }
 
 FAssetRegistry::FAssetRegistry()
-	: PreallocatedAssetDataBuffer(NULL)
-	, PreallocatedDependsNodeDataBuffer(NULL)
+	: PreallocatedAssetDataBuffer(nullptr)
+	, PreallocatedDependsNodeDataBuffer(nullptr)
 {
 	const double StartupStartTime = FPlatformTime::Seconds();
 
@@ -246,7 +246,7 @@ FAssetRegistry::~FAssetRegistry()
 	}
 
 	// if we have preallocated all the FAssetData's in a single block, free it now, instead of one at a time
-	if (PreallocatedAssetDataBuffer != NULL)
+	if (PreallocatedAssetDataBuffer != nullptr)
 	{
 		delete [] PreallocatedAssetDataBuffer;
 		PreallocatedAssetDataBuffer = nullptr;
@@ -269,7 +269,7 @@ FAssetRegistry::~FAssetRegistry()
 	// Make sure we have deleted all our allocated FAssetData objects
 	ensure(NumAssets == 0);
 
-	if (PreallocatedDependsNodeDataBuffer != NULL)
+	if (PreallocatedDependsNodeDataBuffer != nullptr)
 	{
 		delete[] PreallocatedDependsNodeDataBuffer;
 		PreallocatedDependsNodeDataBuffer = nullptr;
@@ -592,7 +592,7 @@ bool FAssetRegistry::GetAssets(const FARFilter& Filter, TArray<FAssetData>& OutA
 		{
 			auto PackageAssets = CachedAssetsByPackageName.Find(*PackageIt);
 					
-			if (PackageAssets != NULL)
+			if (PackageAssets != nullptr)
 			{
 				PackageNameFilter->Append(*PackageAssets);
 			}
@@ -608,7 +608,7 @@ bool FAssetRegistry::GetAssets(const FARFilter& Filter, TArray<FAssetData>& OutA
 		{
 			auto PathAssets = CachedAssetsByPath.Find(*PathIt);
 
-			if (PathAssets != NULL)
+			if (PathAssets != nullptr)
 			{
 				PathFilter->Append(*PathAssets);
 			}
@@ -624,7 +624,7 @@ bool FAssetRegistry::GetAssets(const FARFilter& Filter, TArray<FAssetData>& OutA
 		{
 			auto ClassAssets = CachedAssetsByClass.Find(*ClassNameIt);
 
-			if (ClassAssets != NULL)
+			if (ClassAssets != nullptr)
 			{
 				ClassFilter->Append(*ClassAssets);
 			}
@@ -640,7 +640,7 @@ bool FAssetRegistry::GetAssets(const FARFilter& Filter, TArray<FAssetData>& OutA
 		{
 			const FAssetData*const* AssetDataPtr = CachedAssetsByObjectPath.Find(Filter.ObjectPaths[ObjectPathIdx]);
 
-			if ( AssetDataPtr != NULL )
+			if ( AssetDataPtr != nullptr )
 			{
 				ObjectPathsFilter->Add(*AssetDataPtr);
 			}
@@ -659,16 +659,16 @@ bool FAssetRegistry::GetAssets(const FARFilter& Filter, TArray<FAssetData>& OutA
 				
 			auto TagAssets = CachedAssetsByTag.Find(Tag);
 
-			if (TagAssets != NULL)
+			if (TagAssets != nullptr)
 			{
 				for (auto TagIt = (*TagAssets).CreateConstIterator(); TagIt; ++TagIt)
 				{
 					FAssetData* AssetData = *TagIt;
 
-					if ( AssetData != NULL )
+					if ( AssetData != nullptr )
 					{
 						const FString* TagValue = AssetData->TagsAndValues.Find(Tag);
-						if ( TagValue != NULL && *TagValue == Value )
+						if ( TagValue != nullptr && *TagValue == Value )
 						{
 							TagAndValuesFilter->Add(AssetData);
 						}
@@ -756,15 +756,15 @@ bool FAssetRegistry::GetAssets(const FARFilter& Filter, TArray<FAssetData>& OutA
 
 FAssetData FAssetRegistry::GetAssetByObjectPath( const FName ObjectPath ) const 
 {
-	UObject* Asset = FindObject<UObject>( NULL, *ObjectPath.ToString() );
+	UObject* Asset = FindObject<UObject>( nullptr, *ObjectPath.ToString() );
 
-	if ( Asset != NULL )
+	if ( Asset != nullptr )
 	{
 		return FAssetData( Asset );
 	}
 
 	const FAssetData* const * AssetData = CachedAssetsByObjectPath.Find( ObjectPath );
-	if ( AssetData == NULL || CachedEmptyPackages.Contains( (*AssetData)->PackageName ) )
+	if ( AssetData == nullptr || CachedEmptyPackages.Contains( (*AssetData)->PackageName ) )
 	{
 		return FAssetData();
 	}
@@ -792,7 +792,7 @@ bool FAssetRegistry::GetAllAssets(TArray<FAssetData>& OutAssetData) const
 	{
 		const FAssetData* AssetData = AssetDataIt.Value();
 
-		if (AssetData != NULL)
+		if (AssetData != nullptr)
 		{
 			// Make sure the asset's package was not loaded then the object was deleted/renamed
 			if ( !CachedEmptyPackages.Contains(AssetData->PackageName) )
@@ -814,13 +814,13 @@ bool FAssetRegistry::GetAllAssets(TArray<FAssetData>& OutAssetData) const
 bool FAssetRegistry::GetDependencies(FName PackageName, TArray<FName>& OutDependencies, EAssetRegistryDependencyType::Type InDependencyType, bool bResolveIniStringReferences) const
 {
 	const FDependsNode*const* NodePtr = CachedDependsNodes.Find(PackageName);
-	const FDependsNode* Node = NULL;
-	if (NodePtr != NULL )
+	const FDependsNode* Node = nullptr;
+	if (NodePtr != nullptr )
 	{
 		Node = *NodePtr;
 	}
 
-	if (Node != NULL)
+	if (Node != nullptr)
 	{
 		if (bResolveIniStringReferences)
 		{
@@ -854,13 +854,13 @@ bool FAssetRegistry::GetDependencies(FName PackageName, TArray<FName>& OutDepend
 bool FAssetRegistry::GetReferencers(FName PackageName, TArray<FName>& OutReferencers) const
 {
 	const FDependsNode*const* NodePtr = CachedDependsNodes.Find(PackageName);
-	const FDependsNode* Node = NULL;
-	if (NodePtr != NULL )
+	const FDependsNode* Node = nullptr;
+	if (NodePtr != nullptr )
 	{
 		Node = *NodePtr;
 	}
 
-	if (Node != NULL)
+	if (Node != nullptr)
 	{
 		TArray<FDependsNode*> DependencyNodes;
 		Node->GetReferencers(DependencyNodes);
@@ -909,7 +909,7 @@ bool FAssetRegistry::GetAncestorClassNames(FName ClassName, TArray<FName>& OutAn
 	FName* CurrentClassName = &ClassName;
 	const uint32 MaxInheritanceDepth = 65536;
 	uint32 CurrentInheritanceDepth = 0;
-	while ( CurrentInheritanceDepth < MaxInheritanceDepth && CurrentClassName != NULL )
+	while ( CurrentInheritanceDepth < MaxInheritanceDepth && CurrentClassName != nullptr )
 	{
 		CurrentClassName = InheritanceMap.Find(*CurrentClassName);
 
@@ -918,7 +918,7 @@ bool FAssetRegistry::GetAncestorClassNames(FName ClassName, TArray<FName>& OutAn
 			if ( *CurrentClassName == NAME_None )
 			{
 				// No parent, we are at the root
-				CurrentClassName = NULL;
+				CurrentClassName = nullptr;
 			}
 			else
 			{
@@ -1101,7 +1101,7 @@ void FAssetRegistry::RunAssetsThroughFilter(TArray<FAssetData>& AssetDataList, c
 				{
 					const FString* Value = AssetData.TagsAndValues.Find(FilterTagIt.Key());
 
-					if ( Value != NULL && (*Value) == FilterTagIt.Value() )
+					if ( Value != nullptr && (*Value) == FilterTagIt.Value() )
 					{
 						bPassesTags = true;
 						break;
@@ -1337,7 +1337,7 @@ void FAssetRegistry::AssetRenamed(const UObject* RenamedAsset, const FString& Ol
 		FString OldAssetName;
 		if ( OldObjectPath.Split(TEXT("."), &OldPackageName, &OldAssetName) )
 		{
-			UPackage* OldPackage = FindPackage(NULL, *OldPackageName);
+			UPackage* OldPackage = FindPackage(nullptr, *OldPackageName);
 
 			if ( UPackage::IsEmptyPackage(OldPackage) )
 			{
@@ -1828,14 +1828,14 @@ void FAssetRegistry::AssetSearchDataGathered(const double TickStartTime, TArray<
 		FAssetData Result = BackgroundResult->ToAssetData();
 
 		// Try to update any asset data that may already exist
-		FAssetData* AssetData = NULL;
+		FAssetData* AssetData = nullptr;
 		FAssetData** AssetDataPtr = CachedAssetsByObjectPath.Find(Result.ObjectPath);
-		if (AssetDataPtr != NULL)
+		if (AssetDataPtr != nullptr)
 		{
 			AssetData = *AssetDataPtr;
 		}
 
-		if ( AssetData != NULL )
+		if ( AssetData != nullptr )
 		{
 			// The asset exists in the cache, update it
 			UpdateAssetData(AssetData, Result);
@@ -1851,7 +1851,7 @@ void FAssetRegistry::AssetSearchDataGathered(const double TickStartTime, TArray<
 
 		// Delete the result that was originally created by an FPackageReader
 		delete BackgroundResult;
-		BackgroundResult = NULL;
+		BackgroundResult = nullptr;
 
 		// Check to see if we have run out of time in this tick
 		if ( !bFlushFullBuffer && (FPlatformTime::Seconds() - TickStartTime) > MaxSecondsPerFrame)
@@ -1934,7 +1934,7 @@ void FAssetRegistry::DependencyDataGathered(const double TickStartTime, TArray<F
 		{
 			FDependsNode* DependsNode = CreateOrFindDependsNode(NewDependsIt.Key);
 			
-			if (DependsNode != NULL)
+			if (DependsNode != nullptr)
 			{
 				Node->AddDependency(DependsNode, NewDependsIt.Value);
 				DependsNode->AddReferencer(Node);
@@ -2012,7 +2012,12 @@ bool FAssetRegistry::RemoveDependsNode( FName PackageName )
 			// Remove the node and delete it
 			CachedDependsNodes.Remove( PackageName );
 			NumDependsNodes--;
-			delete Node;
+			
+			// if the depends nodes were preallocated in a block, we can't delete them one at a time, only the whole chunk in the destructor
+			if (PreallocatedDependsNodeDataBuffer == nullptr)
+			{
+				delete Node;
+			}
 
 			return true;
 		}
@@ -2253,7 +2258,7 @@ bool FAssetRegistry::RemoveAssetData(FAssetData* AssetData)
 		RemoveDependsNode( AssetData->PackageName );
 
 		// if the assets were preallocated in a block, we can't delete them one at a time, only the whole chunk in the destructor
-		if (PreallocatedAssetDataBuffer == NULL)
+		if (PreallocatedAssetDataBuffer == nullptr)
 		{
 			delete AssetData;
 		}
