Fixed FAssetRegistry::DependencyDataGathered favoring soft asset dependencies

FAssetRegistry::DependencyDataGathered was potentially replacing hard asset dependencies with soft asset dependencies, but the hard asset dependencies need to take precedence.

diff --git a/Engine/Source/Runtime/AssetRegistry/Private/AssetRegistry.cpp b/Engine/Source/Runtime/AssetRegistry/Private/AssetRegistry.cpp
--- a/Engine/Source/Runtime/AssetRegistry/Private/AssetRegistry.cpp
+++ b/Engine/Source/Runtime/AssetRegistry/Private/AssetRegistry.cpp
@@ -1915,20 +1915,35 @@ void FAssetRegistry::DependencyDataGathered(const double TickStartTime, TArray<F
 		TMap<FName, EAssetRegistryDependencyType::Type> PackageDependencies;
 		for (int32 ImportIdx = 0; ImportIdx < Result.ImportMap.Num(); ++ImportIdx)
 		{
-			PackageDependencies.Add(Result.GetImportPackageName(ImportIdx), EAssetRegistryDependencyType::Hard);
+			const FName AssetReference = Result.GetImportPackageName(ImportIdx);
+
+			// Already processed?
+			if (PackageDependencies.Contains(AssetReference))
+			{
+				continue;
+			}
+
+			PackageDependencies.Add(AssetReference, EAssetRegistryDependencyType::Hard);
 		}
 
-		for (int32 StringAssetRefIdx = 0; StringAssetRefIdx < Result.StringAssetReferencesMap.Num(); ++StringAssetRefIdx)
+		for (const FString& StringAssetReference : Result.StringAssetReferencesMap)
 		{
-			if (FPackageName::IsShortPackageName(Result.StringAssetReferencesMap[StringAssetRefIdx]))
+			const FName AssetReference = *StringAssetReference;
+
+			// Already processed?
+			if (PackageDependencies.Contains(AssetReference))
+			{
+				continue;
+			}
+
+			if (FPackageName::IsShortPackageName(StringAssetReference))
 			{
 				UE_LOG(LogAssetRegistry, Warning, TEXT("Package with string asset reference with short asset path: %s. This is unsupported, can couse errors and be slow on loading. Please resave the package to fix this."), *Result.PackageName.ToString());
 			}
 
-			PackageDependencies.Add(*Result.StringAssetReferencesMap[StringAssetRefIdx]);
+			PackageDependencies.Add(AssetReference, EAssetRegistryDependencyType::Soft);
 		}
 
-
 		// Doubly-link all new dependencies for this package
 		for (auto NewDependsIt : PackageDependencies)
 		{
