#UE4 Consolidated CreateAssetDataFromLinkerTables and CreateAssetDataFromCookedPackage into FPackageReader::ReadAssetRegistryDataIfCookedPackage and established a separate list of cooked packages found with no asset data to fully load later. This way FAssetData is not used to transfer information to the main thread about non-assets (packages).

diff --git a/Engine/Source/Runtime/AssetRegistry/Private/AssetDataGatherer.cpp b/Engine/Source/Runtime/AssetRegistry/Private/AssetDataGatherer.cpp
--- a/Engine/Source/Runtime/AssetRegistry/Private/AssetDataGatherer.cpp
+++ b/Engine/Source/Runtime/AssetRegistry/Private/AssetDataGatherer.cpp
@@ -442,6 +442,7 @@ uint32 FAssetDataGatherer::Run()
 	TArray<FDiscoveredPackageFile> LocalFilesToSearch;
 	TArray<FAssetData*> LocalAssetResults;
 	TArray<FPackageDependencyData> LocalDependencyResults;
+	TArray<FString> LocalCookedPackageNamesWithoutAssetDataResults;
 
 	const double InitialScanStartTime = FPlatformTime::Seconds();
 	int32 NumCachedFiles = 0;
@@ -473,6 +474,7 @@ uint32 FAssetDataGatherer::Run()
 
 			AssetResults.Append(MoveTemp(LocalAssetResults));
 			DependencyResults.Append(MoveTemp(LocalDependencyResults));
+			CookedPackageNamesWithoutAssetDataResults.Append(MoveTemp(LocalCookedPackageNamesWithoutAssetDataResults));
 
 			if (FilesToSearch.Num() > 0)
 			{
@@ -494,6 +496,7 @@ uint32 FAssetDataGatherer::Run()
 
 		LocalAssetResults.Reset();
 		LocalDependencyResults.Reset();
+		LocalCookedPackageNamesWithoutAssetDataResults.Reset();
 
 		if (LocalFilesToSearch.Num() > 0)
 		{
@@ -543,14 +546,17 @@ uint32 FAssetDataGatherer::Run()
 				{
 					TArray<FAssetData*> AssetDataFromFile;
 					FPackageDependencyData DependencyData;
-					if (ReadAssetFile(AssetFileData.PackageFilename, AssetDataFromFile, DependencyData))
+					TArray<FString> CookedPackageNamesWithoutAssetData;
+					if (ReadAssetFile(AssetFileData.PackageFilename, AssetDataFromFile, DependencyData, CookedPackageNamesWithoutAssetData))
 					{
 						++NumUncachedFiles;
 
 						LocalAssetResults.Append(AssetDataFromFile);
 						LocalDependencyResults.Add(DependencyData);
+						LocalCookedPackageNamesWithoutAssetDataResults.Append(CookedPackageNamesWithoutAssetData);
 
-						bool bCachePackage = bLoadAndSaveCache;
+						// Don't store info on cooked packages
+						bool bCachePackage = bLoadAndSaveCache && LocalCookedPackageNamesWithoutAssetDataResults.Num() == 0;
 						if (bCachePackage)
 						{
 							// Don't store info on cooked packages
@@ -663,7 +669,7 @@ void FAssetDataGatherer::EnsureCompletion()
     Thread = nullptr;
 }
 
-bool FAssetDataGatherer::GetAndTrimSearchResults(TArray<FAssetData*>& OutAssetResults, TArray<FString>& OutPathResults, TArray<FPackageDependencyData>& OutDependencyResults, TArray<double>& OutSearchTimes, int32& OutNumFilesToSearch, int32& OutNumPathsToSearch, bool& OutIsDiscoveringFiles)
+bool FAssetDataGatherer::GetAndTrimSearchResults(TArray<FAssetData*>& OutAssetResults, TArray<FString>& OutPathResults, TArray<FPackageDependencyData>& OutDependencyResults, TArray<FString>& OutCookedPackageNamesWithoutAssetDataResults, TArray<double>& OutSearchTimes, int32& OutNumFilesToSearch, int32& OutNumPathsToSearch, bool& OutIsDiscoveringFiles)
 {
 	FScopeLock CritSectionLock(&WorkerThreadCriticalSection);
 
@@ -676,6 +682,9 @@ bool FAssetDataGatherer::GetAndTrimSearchResults(TArray<FAssetData*>& OutAssetRe
 	OutDependencyResults.Append(MoveTemp(DependencyResults));
 	DependencyResults.Reset();
 
+	OutCookedPackageNamesWithoutAssetDataResults.Append(MoveTemp(CookedPackageNamesWithoutAssetDataResults));
+	CookedPackageNamesWithoutAssetDataResults.Reset();
+
 	OutSearchTimes.Append(MoveTemp(SearchTimes));
 	SearchTimes.Reset();
 
@@ -756,66 +765,7 @@ void FAssetDataGatherer::SortPathsByPriority(const int32 MaxNumToSort)
 	}
 }
 
-FAssetData* FAssetDataGatherer::CreateAssetDataFromLinkerTables(const FString& AssetFilename, uint32 InPackageFlags, const FObjectExport& AssetExport, const TArray<FObjectImport>& ImportMap, const TArray<FObjectExport>& ExportMap) const
-{
-	const FString PackageName = FPackageName::FilenameToLongPackageName(AssetFilename);
-	const FString PackagePath = FPaths::GetPath(PackageName);
-	FString GroupNames; // Not used for anything
-	TMap<FName, FString> Tags; // Not used for anything
-	TArray<int32> ChunkIDs; // Not used for anything
-
-	// We need to get the class name from the import/export maps
-	FName ObjectClassName;
-	if (AssetExport.ClassIndex.IsNull())
-	{
-		ObjectClassName = UClass::StaticClass()->GetFName();
-	}
-	else if (AssetExport.ClassIndex.IsExport())
-	{
-		const FObjectExport& ClassExport = ExportMap[AssetExport.ClassIndex.ToExport()];
-		ObjectClassName = ClassExport.ObjectName;
-	}
-	else if (AssetExport.ClassIndex.IsImport())
-	{
-		const FObjectImport& ClassImport = ImportMap[AssetExport.ClassIndex.ToImport()];
-		ObjectClassName = ClassImport.ObjectName;
-	}
-	return new FAssetData(FName(*PackageName), FName(*PackagePath), FName(*GroupNames), AssetExport.ObjectName, ObjectClassName, Tags, ChunkIDs, InPackageFlags);
-}
-
-FAssetData* FAssetDataGatherer::CreateAssetDataFromCookedPackage(const FString& AssetFilename, uint32 InPackageFlags, FPackageReader& PackageReader) const
-{
-	FString PackageName = FPackageName::FilenameToLongPackageName(AssetFilename);
-	FAssetData* Result = nullptr;
-
-	// If the packaged is saved with the right version we have the information
-	// which of the objects in the export map as the asset.
-	// Otherwise we need to store a temp minimal data and then force load the asset
-	// to re-generate its registry data
-	if (PackageReader.UE4Ver() >= VER_UE4_COOKED_ASSETS_IN_EDITOR_SUPPORT)
-	{
-		TArray<FObjectImport> ImportMap;
-		TArray<FObjectExport> ExportMap;		
-		PackageReader.SerializeNameMap();
-		PackageReader.SerializeImportMap(ImportMap);
-		PackageReader.SerializeExportMap(ExportMap);
-		for (FObjectExport& Export : ExportMap)
-		{
-			if (Export.bIsAsset)
-			{
-				Result = CreateAssetDataFromLinkerTables(AssetFilename, InPackageFlags, Export, ImportMap, ExportMap);
-				break;
-			}
-		}
-	}	
-	if (!Result)
-	{
-		Result = new FAssetData(FName(*PackageName), FName(), FName(), FName(), FName(), TMap<FName, FString>(), TArray<int32>(), InPackageFlags);
-	}
-	return Result;
-}
-
-bool FAssetDataGatherer::ReadAssetFile(const FString& AssetFilename, TArray<FAssetData*>& AssetDataList, FPackageDependencyData& DependencyData ) const
+bool FAssetDataGatherer::ReadAssetFile(const FString& AssetFilename, TArray<FAssetData*>& AssetDataList, FPackageDependencyData& DependencyData, TArray<FString>& CookedPackageNamesWithoutAssetData ) const
 {
 	FPackageReader PackageReader;
 
@@ -824,15 +774,12 @@ bool FAssetDataGatherer::ReadAssetFile(const FString& AssetFilename, TArray<FAss
 		return false;
 	}
 
-	if (!!(PackageReader.GetPackageFlags() & PKG_FilterEditorOnly))
+	if ( PackageReader.ReadAssetRegistryDataIfCookedPackage(AssetDataList, CookedPackageNamesWithoutAssetData) )
 	{
-		// Try to reconstruct asset data from the cooked package by serializing its import and export tables
-		FAssetData* CookedPackageData = CreateAssetDataFromCookedPackage(AssetFilename, PackageReader.GetPackageFlags(), PackageReader);
-		AssetDataList.Add(CookedPackageData);
+		// Cooked data is special. No further data is found in these packages
 		return true;
 	}
 
-
 	if ( !PackageReader.ReadAssetRegistryData(AssetDataList) )
 	{
 		if ( !PackageReader.ReadAssetDataFromThumbnailCache(AssetDataList) )
diff --git a/Engine/Source/Runtime/AssetRegistry/Private/AssetDataGatherer.h b/Engine/Source/Runtime/AssetRegistry/Private/AssetDataGatherer.h
--- a/Engine/Source/Runtime/AssetRegistry/Private/AssetDataGatherer.h
+++ b/Engine/Source/Runtime/AssetRegistry/Private/AssetDataGatherer.h
@@ -145,7 +145,7 @@ public:
 	void EnsureCompletion();
 
 	/** Gets search results from the data gatherer */
-	bool GetAndTrimSearchResults(TArray<FAssetData*>& OutAssetResults, TArray<FString>& OutPathResults, TArray<FPackageDependencyData>& OutDependencyResults, TArray<double>& OutSearchTimes, int32& OutNumFilesToSearch, int32& OutNumPathsToSearch, bool& OutIsDiscoveringFiles);
+	bool GetAndTrimSearchResults(TArray<FAssetData*>& OutAssetResults, TArray<FString>& OutPathResults, TArray<FPackageDependencyData>& OutDependencyResults, TArray<FString>& OutCookedPackageNamesWithoutAssetDataResults, TArray<double>& OutSearchTimes, int32& OutNumFilesToSearch, int32& OutNumPathsToSearch, bool& OutIsDiscoveringFiles);
 
 	/** Adds a root path to the search queue. Only works when searching asynchronously */
 	void AddPathToSearch(const FString& Path);
@@ -169,17 +169,11 @@ private:
 	 *
 	 * @return true if the file was successfully read
 	 */
-	bool ReadAssetFile(const FString& AssetFilename, TArray<FAssetData*>& AssetDataList, FPackageDependencyData& DependencyData) const;
+	bool ReadAssetFile(const FString& AssetFilename, TArray<FAssetData*>& AssetDataList, FPackageDependencyData& DependencyData, TArray<FString>& CookedPackagesToLoadUponDiscovery) const;
 
 	/** Serializes the timestamped cache of discovered assets. Used for quick loading of data for assets that have not changed on disk */
 	void SerializeCache(FArchive& Ar);
 
-	/** Creates asset data reconstructing all the required info from linker tables */
-	FAssetData* CreateAssetDataFromLinkerTables(const FString& AssetFilename, uint32 InPackageFlags, const FObjectExport& AssetExport, const TArray<FObjectImport>& ImportMap, const TArray<FObjectExport>& ExportMap) const;
-
-	/** Creates asset data reconstructing all the required data from cooked package info */
-	FAssetData* CreateAssetDataFromCookedPackage(const FString& AssetFilename, uint32 InPackageFlags, FPackageReader& PackageReader) const;
-
 private:
 	/** A critical section to protect data transfer to the main thread */
 	FCriticalSection WorkerThreadCriticalSection;
@@ -208,6 +202,9 @@ private:
 	/** Dependency data for scanned packages */
 	TArray<FPackageDependencyData> DependencyResults;
 
+	/** A list of cooked packages that did not have asset data in them. These assets may still contain assets (if they were older for example) */
+	TArray<FString> CookedPackageNamesWithoutAssetDataResults;
+
 	/** All the search times since the last main thread tick */
 	TArray<double> SearchTimes;
 
diff --git a/Engine/Source/Runtime/AssetRegistry/Private/AssetRegistry.cpp b/Engine/Source/Runtime/AssetRegistry/Private/AssetRegistry.cpp
--- a/Engine/Source/Runtime/AssetRegistry/Private/AssetRegistry.cpp
+++ b/Engine/Source/Runtime/AssetRegistry/Private/AssetRegistry.cpp
@@ -1369,7 +1369,7 @@ void FAssetRegistry::Tick(float DeltaTime)
 	bool bIsDiscoveringFiles = false;
 	if ( BackgroundAssetSearch.IsValid() )
 	{
-		bIsSearching = BackgroundAssetSearch->GetAndTrimSearchResults(BackgroundAssetResults, BackgroundPathResults, BackgroundDependencyResults, SearchTimes, NumFilesToSearch, NumPathsToSearch, bIsDiscoveringFiles);
+		bIsSearching = BackgroundAssetSearch->GetAndTrimSearchResults(BackgroundAssetResults, BackgroundPathResults, BackgroundDependencyResults, BackgroundCookedPackageNamesWithoutAssetDataResults, SearchTimes, NumFilesToSearch, NumPathsToSearch, bIsDiscoveringFiles);
 	}
 
 	// Report the search times
@@ -1409,6 +1409,12 @@ void FAssetRegistry::Tick(float DeltaTime)
 		DependencyDataGathered(TickStartTime, BackgroundDependencyResults);
 	}
 
+	// Load cooked packages that do not have asset data
+	if ( BackgroundCookedPackageNamesWithoutAssetDataResults.Num() )
+	{
+		CookedPackageNamesWithoutAssetDataGathered(TickStartTime, BackgroundCookedPackageNamesWithoutAssetDataResults);
+	}
+
 	// Notify the status change
 	if (bIsSearching || bHadAssetsToProcess)
 	{
@@ -1422,7 +1428,7 @@ void FAssetRegistry::Tick(float DeltaTime)
 	}
 
 	// If completing an initial search, refresh the content browser
-	if ( NumFilesToSearch == 0 && NumPathsToSearch == 0 && !bIsSearching && BackgroundPathResults.Num() == 0 && BackgroundAssetResults.Num() == 0 && BackgroundDependencyResults.Num() == 0 )
+	if ( NumFilesToSearch == 0 && NumPathsToSearch == 0 && !bIsSearching && BackgroundPathResults.Num() == 0 && BackgroundAssetResults.Num() == 0 && BackgroundDependencyResults.Num() == 0 && BackgroundCookedPackageNamesWithoutAssetDataResults.Num() == 0 )
 	{
 		if ( !bInitialSearchCompleted )
 		{
@@ -1777,17 +1783,19 @@ void FAssetRegistry::ScanPathsSynchronous_Internal(const TArray<FString>& InPath
 		TArray<FAssetData*> AssetResults;
 		TArray<FString> PathResults;
 		TArray<FPackageDependencyData> DependencyResults;
+		TArray<FString> CookedPackageNamesWithoutAssetDataResults;
 		TArray<double> SearchTimes;
 		int32 NumFilesToSearch = 0;
 		int32 NumPathsToSearch = 0;
 		bool bIsDiscoveringFiles = false;
-		AssetSearch.GetAndTrimSearchResults(AssetResults, PathResults, DependencyResults, SearchTimes, NumFilesToSearch, NumPathsToSearch, bIsDiscoveringFiles);
+		AssetSearch.GetAndTrimSearchResults(AssetResults, PathResults, DependencyResults, CookedPackageNamesWithoutAssetDataResults, SearchTimes, NumFilesToSearch, NumPathsToSearch, bIsDiscoveringFiles);
 
 		// Cache the search results
 		const int32 NumResults = AssetResults.Num();
 		AssetSearchDataGathered(-1, AssetResults);
 		PathDataGathered(-1, PathResults);
 		DependencyDataGathered(-1, DependencyResults);
+		CookedPackageNamesWithoutAssetDataGathered(-1, CookedPackageNamesWithoutAssetDataResults);
 
 		// Log stats
 		const FString& Path = PathsToScan[0];
@@ -1816,15 +1824,6 @@ void FAssetRegistry::AssetSearchDataGathered(const double TickStartTime, TArray<
 	{
 		FAssetData*& BackgroundResult = AssetResults[AssetIndex];
 
-		// If this data is cooked and it we couldn't find any asset in its export table then try load the entire package 
-		if (!!(BackgroundResult->PackageFlags & PKG_FilterEditorOnly))
-		{
-			LoadPackage(nullptr, *BackgroundResult->PackageName.ToString(), 0);
-			delete BackgroundResult;
-			BackgroundResult = nullptr;
-			continue;
-		}
-
 		// Try to update any asset data that may already exist
 		FAssetData* AssetData = nullptr;
 		FAssetData** AssetDataPtr = CachedAssetsByObjectPath.Find(BackgroundResult->ObjectPath);
@@ -1968,6 +1967,34 @@ void FAssetRegistry::DependencyDataGathered(const double TickStartTime, TArray<F
 	DependsResults.RemoveAt(0, ResultIdx);
 }
 
+void FAssetRegistry::CookedPackageNamesWithoutAssetDataGathered(const double TickStartTime, TArray<FString>& CookedPackageNamesWithoutAssetDataResults)
+{
+	const bool bFlushFullBuffer = TickStartTime < 0;
+
+	// Add the found assets
+	int32 PackageNameIndex = 0;
+	for (PackageNameIndex = 0; PackageNameIndex < CookedPackageNamesWithoutAssetDataResults.Num(); ++PackageNameIndex)
+	{
+		// If this data is cooked and it we couldn't find any asset in its export table then try load the entire package 
+		const FString& BackgroundResult = CookedPackageNamesWithoutAssetDataResults[PackageNameIndex];
+		LoadPackage(nullptr, *BackgroundResult, 0);
+
+		// Check to see if we have run out of time in this tick
+		if (!bFlushFullBuffer && (FPlatformTime::Seconds() - TickStartTime) > MaxSecondsPerFrame)
+		{
+			// Increment the index to properly trim the buffer below
+			++PackageNameIndex;
+			break;
+		}
+	}
+
+	// Trim the results array
+	if (PackageNameIndex > 0)
+	{
+		CookedPackageNamesWithoutAssetDataResults.RemoveAt(0, PackageNameIndex);
+	}
+}
+
 FDependsNode* FAssetRegistry::FindDependsNode(FName ObjectName)
 {
 	FDependsNode** FoundNode = CachedDependsNodes.Find(ObjectName);
diff --git a/Engine/Source/Runtime/AssetRegistry/Private/AssetRegistry.h b/Engine/Source/Runtime/AssetRegistry/Private/AssetRegistry.h
--- a/Engine/Source/Runtime/AssetRegistry/Private/AssetRegistry.h
+++ b/Engine/Source/Runtime/AssetRegistry/Private/AssetRegistry.h
@@ -87,12 +87,15 @@ private:
 	/** Called every tick to when data is retrieved by the background asset search. If TickStartTime is < 0, the entire list of gathered assets will be cached. Also used in sychronous searches */
 	void AssetSearchDataGathered(const double TickStartTime, TArray<FAssetData*>& AssetResults);
 
-	/** Called every tick to when data is retrieved by the background path search. If TickStartTime is < 0, the entire list of gathered assets will be cached. Also used in sychronous searches */
+	/** Called every tick when data is retrieved by the background path search. If TickStartTime is < 0, the entire list of gathered assets will be cached. Also used in sychronous searches */
 	void PathDataGathered(const double TickStartTime, TArray<FString>& PathResults);
 
-	/** Called every tick to when data is retrieved by the background dependency search */
+	/** Called every tick when data is retrieved by the background dependency search */
 	void DependencyDataGathered(const double TickStartTime, TArray<FPackageDependencyData>& DependsResults);
 
+	/** Called every tick when data is retrieved by the background search for cooked packages that do not contain asset data */
+	void CookedPackageNamesWithoutAssetDataGathered(const double TickStartTime, TArray<FString>& CookedPackageNamesWithoutAssetDataResults);
+
 	/** Finds an existing node for the given package and returns it, or returns null if one isn't found */
 	FDependsNode* FindDependsNode(FName ObjectName);
 
@@ -210,6 +213,7 @@ private:
 	TArray<class FAssetData*> BackgroundAssetResults;
 	TArray<FString> BackgroundPathResults;
 	TArray<class FPackageDependencyData> BackgroundDependencyResults;
+	TArray<FString> BackgroundCookedPackageNamesWithoutAssetDataResults;
 
 	/** The max number of results to process per tick */
 	float MaxSecondsPerFrame;
diff --git a/Engine/Source/Runtime/AssetRegistry/Private/PackageReader.cpp b/Engine/Source/Runtime/AssetRegistry/Private/PackageReader.cpp
--- a/Engine/Source/Runtime/AssetRegistry/Private/PackageReader.cpp
+++ b/Engine/Source/Runtime/AssetRegistry/Private/PackageReader.cpp
@@ -250,6 +250,67 @@ bool FPackageReader::ReadAssetDataFromThumbnailCache(TArray<FAssetData*>& AssetD
 	return true;
 }
 
+bool FPackageReader::ReadAssetRegistryDataIfCookedPackage(TArray<FAssetData*>& AssetDataList, TArray<FString>& CookedPackageNamesWithoutAssetData)
+{
+	if (!!(GetPackageFlags() & PKG_FilterEditorOnly))
+	{
+		const FString PackageName = FPackageName::FilenameToLongPackageName(PackageFilename);
+		
+		bool bFoundAtLeastOneAsset = false;
+
+		// If the packaged is saved with the right version we have the information
+		// which of the objects in the export map as the asset.
+		// Otherwise we need to store a temp minimal data and then force load the asset
+		// to re-generate its registry data
+		if (UE4Ver() >= VER_UE4_COOKED_ASSETS_IN_EDITOR_SUPPORT)
+		{
+			const FString PackagePath = FPackageName::GetLongPackagePath(PackageName);
+
+			TArray<FObjectImport> ImportMap;
+			TArray<FObjectExport> ExportMap;
+			SerializeNameMap();
+			SerializeImportMap(ImportMap);
+			SerializeExportMap(ExportMap);
+			for (FObjectExport& Export : ExportMap)
+			{
+				if (Export.bIsAsset)
+				{
+					FString GroupNames; // Not used for anything
+					TMap<FName, FString> Tags; // Not used for anything
+					TArray<int32> ChunkIDs; // Not used for anything
+
+					// We need to get the class name from the import/export maps
+					FName ObjectClassName;
+					if (Export.ClassIndex.IsNull())
+					{
+						ObjectClassName = UClass::StaticClass()->GetFName();
+					}
+					else if (Export.ClassIndex.IsExport())
+					{
+						const FObjectExport& ClassExport = ExportMap[Export.ClassIndex.ToExport()];
+						ObjectClassName = ClassExport.ObjectName;
+					}
+					else if (Export.ClassIndex.IsImport())
+					{
+						const FObjectImport& ClassImport = ImportMap[Export.ClassIndex.ToImport()];
+						ObjectClassName = ClassImport.ObjectName;
+					}
+
+					new(AssetDataList) FAssetData(FName(*PackageName), FName(*PackagePath), FName(*GroupNames), Export.ObjectName, ObjectClassName, Tags, ChunkIDs, GetPackageFlags());
+					bFoundAtLeastOneAsset = true;
+				}
+			}
+		}
+		if (!bFoundAtLeastOneAsset)
+		{
+			CookedPackageNamesWithoutAssetData.Add(PackageName);
+		}
+		return true;
+	}
+
+	return false;
+}
+
 bool FPackageReader::ReadDependencyData (FPackageDependencyData& OutDependencyData)
 {
 	OutDependencyData.PackageName = FName(*FPackageName::FilenameToLongPackageName(PackageFilename));
diff --git a/Engine/Source/Runtime/AssetRegistry/Private/PackageReader.h b/Engine/Source/Runtime/AssetRegistry/Private/PackageReader.h
--- a/Engine/Source/Runtime/AssetRegistry/Private/PackageReader.h
+++ b/Engine/Source/Runtime/AssetRegistry/Private/PackageReader.h
@@ -17,6 +17,9 @@ public:
 	/** Attempts to get the class name of an object from the thumbnail cache for packages older than VER_UE4_ASSET_REGISTRY_TAGS */
 	bool ReadAssetDataFromThumbnailCache(TArray<FAssetData*>& AssetDataList);
 
+	/** Creates asset data reconstructing all the required data from cooked package info */
+	bool ReadAssetRegistryDataIfCookedPackage(TArray<FAssetData*>& AssetDataList, TArray<FString>& CookedPackageNamesWithoutAssetData);
+
 	/** Reads information used by the dependency graph */
 	bool ReadDependencyData(FPackageDependencyData& OutDependencyData);
 
