Asset registry no longer reloads ddc data in cooker when saving asset registry.  End of cook performance improvement.

diff --git a/Engine/Source/Editor/UnrealEd/Private/Commandlets/ChunkManifestGenerator.cpp b/Engine/Source/Editor/UnrealEd/Private/Commandlets/ChunkManifestGenerator.cpp
--- a/Engine/Source/Editor/UnrealEd/Private/Commandlets/ChunkManifestGenerator.cpp
+++ b/Engine/Source/Editor/UnrealEd/Private/Commandlets/ChunkManifestGenerator.cpp
@@ -374,7 +374,7 @@ void FChunkManifestGenerator::BuildChunkManifest(const TArray<FName>& CookedPack
 	// initialize LargestChunkId, FoundIDList, PackageChunkIDMap, AssetRegistryData
 
 	// Calculate the largest chunk id used by the registry to get the indices for the default chunks
-	AssetRegistry.GetAllAssets(AssetRegistryData);
+	AssetRegistry.GetAllAssets(AssetRegistryData, true);
 	int32 LargestChunkID = -1;
 
 	for (int32 Index = 0; Index < AssetRegistryData.Num(); ++Index)
diff --git a/Engine/Source/Editor/UnrealEd/Private/CookOnTheFlyServer.cpp b/Engine/Source/Editor/UnrealEd/Private/CookOnTheFlyServer.cpp
--- a/Engine/Source/Editor/UnrealEd/Private/CookOnTheFlyServer.cpp
+++ b/Engine/Source/Editor/UnrealEd/Private/CookOnTheFlyServer.cpp
@@ -3854,7 +3854,7 @@ void UCookOnTheFlyServer::CookByTheBookFinished()
 		CleanUpChildCookers();
 
 		check(CookByTheBookOptions->ChildUnsolicitedPackages.Num() == 0);
-
+		SCOPE_COOKING_STAT(SavingAssetRegistry);
 		SCOPE_TIMER(SavingAssetRegistry);
 		// Save modified asset registry with all streaming chunk info generated during cook
 		const FString& SandboxRegistryFilename = GetSandboxAssetRegistryFilename();
diff --git a/Engine/Source/Runtime/AssetRegistry/Private/AssetRegistry.cpp b/Engine/Source/Runtime/AssetRegistry/Private/AssetRegistry.cpp
--- a/Engine/Source/Runtime/AssetRegistry/Private/AssetRegistry.cpp
+++ b/Engine/Source/Runtime/AssetRegistry/Private/AssetRegistry.cpp
@@ -772,18 +772,21 @@ FAssetData FAssetRegistry::GetAssetByObjectPath( const FName ObjectPath ) const
 	return **AssetData;
 }
 
-bool FAssetRegistry::GetAllAssets(TArray<FAssetData>& OutAssetData) const
+bool FAssetRegistry::GetAllAssets(TArray<FAssetData>& OutAssetData, bool bIncludeOnlyOnDiskAssets) const
 {
 	TSet<FName> InMemoryObjectPaths;
 	double GetAllAssetsStartTime = FPlatformTime::Seconds();
 
 	// All in memory assets
-	for (FObjectIterator ObjIt; ObjIt; ++ObjIt)
+	if (bIncludeOnlyOnDiskAssets)
 	{
-		if (ObjIt->IsAsset())
+		for (FObjectIterator ObjIt; ObjIt; ++ObjIt)
 		{
-			FAssetData* AssetData = new (OutAssetData) FAssetData(*ObjIt);
-			InMemoryObjectPaths.Add(AssetData->ObjectPath);
+			if (ObjIt->IsAsset())
+			{
+				FAssetData* AssetData = new (OutAssetData)FAssetData(*ObjIt);
+				InMemoryObjectPaths.Add(AssetData->ObjectPath);
+			}
 		}
 	}
 
diff --git a/Engine/Source/Runtime/AssetRegistry/Private/AssetRegistry.h b/Engine/Source/Runtime/AssetRegistry/Private/AssetRegistry.h
--- a/Engine/Source/Runtime/AssetRegistry/Private/AssetRegistry.h
+++ b/Engine/Source/Runtime/AssetRegistry/Private/AssetRegistry.h
@@ -20,7 +20,7 @@ public:
 	virtual bool GetAssetsByTagValues(const TMultiMap<FName, FString>& AssetTagsAndValues, TArray<FAssetData>& OutAssetData) const override;
 	virtual bool GetAssets(const FARFilter& Filter, TArray<FAssetData>& OutAssetData) const override;
 	virtual FAssetData GetAssetByObjectPath( const FName ObjectPath ) const override;
-	virtual bool GetAllAssets(TArray<FAssetData>& OutAssetData) const override;
+	virtual bool GetAllAssets(TArray<FAssetData>& OutAssetData, bool bIncludeOnlyOnDiskAssets = false) const override;
 	virtual bool GetDependencies(FName PackageName, TArray<FName>& OutDependencies, EAssetRegistryDependencyType::Type InDependencyType = EAssetRegistryDependencyType::All, bool bResolveIniStringReferences = false) const override;
 	virtual bool GetReferencers(FName PackageName, TArray<FName>& OutReferencers) const override;
 	virtual bool GetAncestorClassNames(FName ClassName, TArray<FName>& OutAncestorClassNames) const override;
diff --git a/Engine/Source/Runtime/AssetRegistry/Public/IAssetRegistry.h b/Engine/Source/Runtime/AssetRegistry/Public/IAssetRegistry.h
--- a/Engine/Source/Runtime/AssetRegistry/Public/IAssetRegistry.h
+++ b/Engine/Source/Runtime/AssetRegistry/Public/IAssetRegistry.h
@@ -92,7 +92,7 @@ public:
 	 *
 	 * @param OutAssetData the list of assets in this path
 	 */
-	virtual bool GetAllAssets(TArray<FAssetData>& OutAssetData) const = 0;
+	virtual bool GetAllAssets(TArray<FAssetData>& OutAssetData, bool bIncludeOnlyOnDiskAssets = false) const = 0;
 
 	/**
 	 * Gets a list of paths to objects that are referenced by the supplied package. (On disk references ONLY)
