Fix for P0 crash when running a cook on the fly server (UE-23668)

diff --git a/Engine/Source/Runtime/AssetRegistry/Private/AssetRegistry.cpp b/Engine/Source/Runtime/AssetRegistry/Private/AssetRegistry.cpp
--- a/Engine/Source/Runtime/AssetRegistry/Private/AssetRegistry.cpp
+++ b/Engine/Source/Runtime/AssetRegistry/Private/AssetRegistry.cpp
@@ -1724,7 +1724,10 @@ void FAssetRegistry::SaveRegistryData(FArchive& Ar, TMap<FName, FAssetData*>& Da
 		for (auto Dependency : ProcessedDependencies)
 		{
 			FDependsNode* RedirectedDependency = ResolveRedirector(Dependency, Data, RedirectCache);
-			check(RedirectedDependency);
+			if (RedirectedDependency == nullptr)
+			{
+				RedirectedDependency = Dependency;
+			}
 			int32 Index = AssetIndexMap[RedirectedDependency->GetPackageName()];
 			Ar << Index;
 		}
