Fix for UE-25346 - Improve asset registry redirector resolving and circular references

diff --git a/Engine/Source/Runtime/AssetRegistry/Private/AssetRegistry.cpp b/Engine/Source/Runtime/AssetRegistry/Private/AssetRegistry.cpp
--- a/Engine/Source/Runtime/AssetRegistry/Private/AssetRegistry.cpp
+++ b/Engine/Source/Runtime/AssetRegistry/Private/AssetRegistry.cpp
@@ -1575,43 +1575,55 @@ FDependsNode* FAssetRegistry::ResolveRedirector(FDependsNode* InDependency, TMap
 		return InCache[InDependency];
 	}
 
+	static TSet<FName> EncounteredDependencies;
+	EncounteredDependencies.Empty();
+	
 	while (Result == nullptr)
 	{
+		if (EncounteredDependencies.Contains(InDependency->GetPackageName()))
+		{
+			break;
+		}
+
+		EncounteredDependencies.Add(InDependency->GetPackageName());
+
 		if (CachedAssetsByPackageName.Contains(InDependency->GetPackageName()))
 		{
-			auto DependencyAsset = CachedAssetsByPackageName[InDependency->GetPackageName()][0];
+			// Get the list of assets contained in this package
+			TArray<FAssetData*>& Assets = CachedAssetsByPackageName[InDependency->GetPackageName()];
 
-			if (DependencyAsset->AssetClass == ObjectRedirectorClassName)
+			for (FAssetData* Asset : Assets)
 			{
-				FDependsNode* Before = InDependency;
-
-				InDependency->IterateOverDependencies([&](FDependsNode* InDepends, EAssetRegistryDependencyType::Type)
+				if (Asset->AssetClass == ObjectRedirectorClassName)
 				{
-					if (InAllowedAssets.Contains(InDepends->GetPackageName()))
+					// This asset is a redirector, so we want to look at its dependencies and find the asset that it is redirecting to
+					InDependency->IterateOverDependencies([&](FDependsNode* InDepends, EAssetRegistryDependencyType::Type)
 					{
-						Result = InDepends;
-					}
-					else if (CachedAssetsByPackageName.Contains(InDepends->GetPackageName()))
-					{
-						auto SubDependencyAsset = CachedAssetsByPackageName[InDepends->GetPackageName()][0];
-						if (SubDependencyAsset->AssetClass == ObjectRedirectorClassName)
+						if (InAllowedAssets.Contains(InDepends->GetPackageName()))
 						{
+							// This asset is in the allowed asset list, so take this as the redirect target
+							Result = InDepends;
+						}
+						else if (CachedAssetsByPackageName.Contains(InDepends->GetPackageName()))
+						{
+							// This dependency isn't in the allowed list, but it is a valid asset in the registry.
+							// Because this is a redirector, this should mean that the redirector is pointing at ANOTHER
+							// redirector (or itself in some horrible situations) so we'll move to that node and try again
 							InDependency = InDepends;
 						}
-					}
-				});
+					});
+				}
+				else
+				{
+					Result = InDependency;
+				}
 
-				if (Result == nullptr && InDependency == Before)
+				if (Result)
 				{
-					// The redirector doesn't point at any files that were cooked. Allow function to return with 
-					// result as nullptr to indicate this
+					// We found an allowed asset from the original dependency node. We're finished!
 					break;
 				}
 			}
-			else
-			{
-				Result = InDependency;
-			}
 		}
 		else
 		{
