Copying //UE4/Dev-Framework to //UE4/Dev-Main (Source: //UE4/Dev-Framework @ 3431384)

diff --git a/Engine/Source/Runtime/AIModule/Classes/Perception/AIPerceptionSystem.h b/Engine/Source/Runtime/AIModule/Classes/Perception/AIPerceptionSystem.h
--- a/Engine/Source/Runtime/AIModule/Classes/Perception/AIPerceptionSystem.h
+++ b/Engine/Source/Runtime/AIModule/Classes/Perception/AIPerceptionSystem.h
@@ -132,8 +132,8 @@ public:
 	UFUNCTION(BlueprintCallable, Category = "AI|Perception")
 	void ReportEvent(UAISenseEvent* PerceptionEvent);
 
-	UFUNCTION(BlueprintCallable, Category = "AI|Perception", meta = (WorldContext="WorldContext"))
-	static void ReportPerceptionEvent(UObject* WorldContext, UAISenseEvent* PerceptionEvent);
+	UFUNCTION(BlueprintCallable, Category = "AI|Perception", meta = (WorldContext="WorldContextObject"))
+	static void ReportPerceptionEvent(UObject* WorldContextObject, UAISenseEvent* PerceptionEvent);
 
 	template<typename FSenseClass>
 	void RegisterSource(AActor& SourceActor);
@@ -155,13 +155,13 @@ public:
 
 	static void MakeNoiseImpl(AActor* NoiseMaker, float Loudness, APawn* NoiseInstigator, const FVector& NoiseLocation, float MaxRange, FName Tag);
 
-	UFUNCTION(BlueprintCallable, Category = "AI|Perception", meta = (WorldContext="WorldContext"))
-	static bool RegisterPerceptionStimuliSource(UObject* WorldContext, TSubclassOf<UAISense> Sense, AActor* Target);
+	UFUNCTION(BlueprintCallable, Category = "AI|Perception", meta = (WorldContext="WorldContextObject"))
+	static bool RegisterPerceptionStimuliSource(UObject* WorldContextObject, TSubclassOf<UAISense> Sense, AActor* Target);
 
 	FAISenseID RegisterSenseClass(TSubclassOf<UAISense> SenseClass);
 
-	UFUNCTION(BlueprintCallable, Category = "AI|Perception", meta = (WorldContext="WorldContext"))
-	static TSubclassOf<UAISense> GetSenseClassForStimulus(UObject* WorldContext, const FAIStimulus& Stimulus);
+	UFUNCTION(BlueprintCallable, Category = "AI|Perception", meta = (WorldContext="WorldContextObject"))
+	static TSubclassOf<UAISense> GetSenseClassForStimulus(UObject* WorldContextObject, const FAIStimulus& Stimulus);
 	
 protected:
 	
diff --git a/Engine/Source/Runtime/AIModule/Classes/Perception/AISense_Damage.h b/Engine/Source/Runtime/AIModule/Classes/Perception/AISense_Damage.h
--- a/Engine/Source/Runtime/AIModule/Classes/Perception/AISense_Damage.h
+++ b/Engine/Source/Runtime/AIModule/Classes/Perception/AISense_Damage.h
@@ -10,7 +10,7 @@
 class IAIPerceptionListenerInterface;
 class UAISenseEvent;
 
-USTRUCT()
+USTRUCT(BlueprintType)
 struct AIMODULE_API FAIDamageEvent
 {	
 	GENERATED_USTRUCT_BODY()
@@ -65,8 +65,8 @@ public:
 	virtual void RegisterWrappedEvent(UAISenseEvent& PerceptionEvent) override;
 
 	/** EventLocation will be reported as Instigator's location at the moment of event happening*/
-	UFUNCTION(BlueprintCallable, Category = "AI|Perception", meta = (WorldContext="WorldContext", AdvancedDisplay="HitLocation"))
-	static void ReportDamageEvent(UObject* WorldContext, AActor* DamagedActor, AActor* Instigator, float DamageAmount, FVector EventLocation, FVector HitLocation);
+	UFUNCTION(BlueprintCallable, Category = "AI|Perception", meta = (WorldContext="WorldContextObject", AdvancedDisplay="HitLocation"))
+	static void ReportDamageEvent(UObject* WorldContextObject, AActor* DamagedActor, AActor* Instigator, float DamageAmount, FVector EventLocation, FVector HitLocation);
 
 protected:
 	virtual float Update() override;
diff --git a/Engine/Source/Runtime/AIModule/Classes/Perception/AISense_Hearing.h b/Engine/Source/Runtime/AIModule/Classes/Perception/AISense_Hearing.h
--- a/Engine/Source/Runtime/AIModule/Classes/Perception/AISense_Hearing.h
+++ b/Engine/Source/Runtime/AIModule/Classes/Perception/AISense_Hearing.h
@@ -11,7 +11,7 @@
 class UAISenseConfig_Hearing;
 class UAISenseEvent;
 
-USTRUCT()
+USTRUCT(BlueprintType)
 struct AIMODULE_API FAINoiseEvent
 {	
 	GENERATED_USTRUCT_BODY()
@@ -100,8 +100,8 @@ public:
 	 * @param MaxRange Max range at which the sound can be heard, multiplied by Loudness. Values <= 0 mean no limit (still limited by listener's range however).
 	 * @param Tag Identifier for the event.
 	 */
-	UFUNCTION(BlueprintCallable, Category = "AI|Perception", meta = (WorldContext="WorldContext"))
-	static void ReportNoiseEvent(UObject* WorldContext, FVector NoiseLocation, float Loudness = 1.f, AActor* Instigator = nullptr, float MaxRange = 0.f, FName Tag = NAME_None);
+	UFUNCTION(BlueprintCallable, Category = "AI|Perception", meta = (WorldContext="WorldContextObject"))
+	static void ReportNoiseEvent(UObject* WorldContextObject, FVector NoiseLocation, float Loudness = 1.f, AActor* Instigator = nullptr, float MaxRange = 0.f, FName Tag = NAME_None);
 
 protected:
 	virtual float Update() override;
diff --git a/Engine/Source/Runtime/AIModule/Private/EnvironmentQuery/EnvQueryManager.cpp b/Engine/Source/Runtime/AIModule/Private/EnvironmentQuery/EnvQueryManager.cpp
--- a/Engine/Source/Runtime/AIModule/Private/EnvironmentQuery/EnvQueryManager.cpp
+++ b/Engine/Source/Runtime/AIModule/Private/EnvironmentQuery/EnvQueryManager.cpp
@@ -859,14 +859,14 @@ float UEnvQueryManager::FindNamedParam(int32 QueryId, FName ParamName) const
 //----------------------------------------------------------------------//
 // BP functions and related functionality 
 //----------------------------------------------------------------------//
-UEnvQueryInstanceBlueprintWrapper* UEnvQueryManager::RunEQSQuery(UObject* WorldContext, UEnvQuery* QueryTemplate, UObject* Querier, TEnumAsByte<EEnvQueryRunMode::Type> RunMode, TSubclassOf<UEnvQueryInstanceBlueprintWrapper> WrapperClass)
+UEnvQueryInstanceBlueprintWrapper* UEnvQueryManager::RunEQSQuery(UObject* WorldContextObject, UEnvQuery* QueryTemplate, UObject* Querier, TEnumAsByte<EEnvQueryRunMode::Type> RunMode, TSubclassOf<UEnvQueryInstanceBlueprintWrapper> WrapperClass)
 { 
 	if (QueryTemplate == nullptr || Querier == nullptr)
 	{
 		return nullptr;
 	}
 
-	UEnvQueryManager* EQSManager = GetCurrent(WorldContext);
+	UEnvQueryManager* EQSManager = GetCurrent(WorldContextObject);
 	UEnvQueryInstanceBlueprintWrapper* QueryInstanceWrapper = nullptr;
 
 	if (EQSManager)
diff --git a/Engine/Source/Runtime/AIModule/Private/GameplayDebugger/GameplayDebuggerCategory_AI.h b/Engine/Source/Runtime/AIModule/Private/GameplayDebugger/GameplayDebuggerCategory_AI.h
--- a/Engine/Source/Runtime/AIModule/Private/GameplayDebugger/GameplayDebuggerCategory_AI.h
+++ b/Engine/Source/Runtime/AIModule/Private/GameplayDebugger/GameplayDebuggerCategory_AI.h
@@ -2,10 +2,10 @@
 
 #pragma once
 
-#include "CoreMinimal.h"
 #if WITH_GAMEPLAY_DEBUGGER
+
+#include "CoreMinimal.h"
 #include "GameplayDebuggerCategory.h"
-#endif
 
 class AActor;
 class AAIController;
@@ -17,8 +17,6 @@ class UPrimitiveComponent;
 struct FDebugDrawDelegateHelper;
 struct FNavigationPath;
 
-#if WITH_GAMEPLAY_DEBUGGER
-
 class FGameplayDebuggerCategory_AI : public FGameplayDebuggerCategory
 {
 public:
diff --git a/Engine/Source/Runtime/AIModule/Private/GameplayDebugger/GameplayDebuggerCategory_Perception.h b/Engine/Source/Runtime/AIModule/Private/GameplayDebugger/GameplayDebuggerCategory_Perception.h
--- a/Engine/Source/Runtime/AIModule/Private/GameplayDebugger/GameplayDebuggerCategory_Perception.h
+++ b/Engine/Source/Runtime/AIModule/Private/GameplayDebugger/GameplayDebuggerCategory_Perception.h
@@ -2,17 +2,14 @@
 
 #pragma once
 
-#include "CoreMinimal.h"
 #if WITH_GAMEPLAY_DEBUGGER
+
+#include "CoreMinimal.h"
 #include "GameplayDebuggerCategory.h"
-#endif
 
 class AActor;
 class APlayerController;
 
-#if WITH_GAMEPLAY_DEBUGGER
-
-
 class FGameplayDebuggerCategory_Perception : public FGameplayDebuggerCategory
 {
 public:
diff --git a/Engine/Source/Runtime/AIModule/Private/Navigation/NavLocalGridManager.cpp b/Engine/Source/Runtime/AIModule/Private/Navigation/NavLocalGridManager.cpp
--- a/Engine/Source/Runtime/AIModule/Private/Navigation/NavLocalGridManager.cpp
+++ b/Engine/Source/Runtime/AIModule/Private/Navigation/NavLocalGridManager.cpp
@@ -316,11 +316,11 @@ bool UNavLocalGridManager::FindPath(const FVector& Start, const FVector& End, TA
 	return false;
 }
 
-int32 UNavLocalGridManager::AddLocalNavigationGridForPoint(UObject* WorldContext, const FVector& Location, const int32 Radius2D, const float Height, bool bRebuildGrids)
+int32 UNavLocalGridManager::AddLocalNavigationGridForPoint(UObject* WorldContextObject, const FVector& Location, const int32 Radius2D, const float Height, bool bRebuildGrids)
 {
 	int32 GridId = 0;
 
-	UNavLocalGridManager* GridManager = UNavLocalGridManager::GetCurrent(WorldContext);
+	UNavLocalGridManager* GridManager = UNavLocalGridManager::GetCurrent(WorldContextObject);
 	if (GridManager)
 	{
 		FNavLocalGridData GridData(Location, UNavLocalGridManager::GridCellSize * Radius2D);
@@ -333,11 +333,11 @@ int32 UNavLocalGridManager::AddLocalNavigationGridForPoint(UObject* WorldContext
 	return GridId;
 }
 
-int32 UNavLocalGridManager::AddLocalNavigationGridForPoints(UObject* WorldContext, const TArray<FVector>& Locations, const int32 Radius2D, const float Height, bool bRebuildGrids)
+int32 UNavLocalGridManager::AddLocalNavigationGridForPoints(UObject* WorldContextObject, const TArray<FVector>& Locations, const int32 Radius2D, const float Height, bool bRebuildGrids)
 {
 	int32 GridId = 0;
 
-	UNavLocalGridManager* GridManager = UNavLocalGridManager::GetCurrent(WorldContext);
+	UNavLocalGridManager* GridManager = UNavLocalGridManager::GetCurrent(WorldContextObject);
 	if (GridManager)
 	{
 		const FBox Bounds(Locations);
@@ -357,11 +357,11 @@ int32 UNavLocalGridManager::AddLocalNavigationGridForPoints(UObject* WorldContex
 	return GridId;
 }
 
-int32 UNavLocalGridManager::AddLocalNavigationGridForBox(UObject* WorldContext, const FVector& Location, FVector Extent, FRotator Rotation, const int32 Radius2D, const float Height, bool bRebuildGrids)
+int32 UNavLocalGridManager::AddLocalNavigationGridForBox(UObject* WorldContextObject, const FVector& Location, FVector Extent, FRotator Rotation, const int32 Radius2D, const float Height, bool bRebuildGrids)
 {
 	int32 GridId = 0;
 
-	UNavLocalGridManager* GridManager = UNavLocalGridManager::GetCurrent(WorldContext);
+	UNavLocalGridManager* GridManager = UNavLocalGridManager::GetCurrent(WorldContextObject);
 	if (GridManager)
 	{
 		FNavLocalGridData GridData(Location, FVector2D(Extent.X + UNavLocalGridManager::GridCellSize * Radius2D, Extent.Y + UNavLocalGridManager::GridCellSize * Radius2D));
@@ -374,11 +374,11 @@ int32 UNavLocalGridManager::AddLocalNavigationGridForBox(UObject* WorldContext,
 	return GridId;
 }
 
-int32 UNavLocalGridManager::AddLocalNavigationGridForCapsule(UObject* WorldContext, const FVector& Location, float CapsuleRadius, float CapsuleHalfHeight, const int32 Radius2D, const float Height, bool bRebuildGrids)
+int32 UNavLocalGridManager::AddLocalNavigationGridForCapsule(UObject* WorldContextObject, const FVector& Location, float CapsuleRadius, float CapsuleHalfHeight, const int32 Radius2D, const float Height, bool bRebuildGrids)
 {
 	int32 GridId = 0;
 
-	UNavLocalGridManager* GridManager = UNavLocalGridManager::GetCurrent(WorldContext);
+	UNavLocalGridManager* GridManager = UNavLocalGridManager::GetCurrent(WorldContextObject);
 	if (GridManager)
 	{
 		FNavLocalGridData GridData(Location, FVector2D(CapsuleRadius + UNavLocalGridManager::GridCellSize * Radius2D, CapsuleRadius + UNavLocalGridManager::GridCellSize * Radius2D));
@@ -391,18 +391,18 @@ int32 UNavLocalGridManager::AddLocalNavigationGridForCapsule(UObject* WorldConte
 	return GridId;
 }
 
-void UNavLocalGridManager::RemoveLocalNavigationGrid(UObject* WorldContext, int32 GridId, bool bRebuildGrids)
+void UNavLocalGridManager::RemoveLocalNavigationGrid(UObject* WorldContextObject, int32 GridId, bool bRebuildGrids)
 {
-	UNavLocalGridManager* GridManager = UNavLocalGridManager::GetCurrent(WorldContext);
+	UNavLocalGridManager* GridManager = UNavLocalGridManager::GetCurrent(WorldContextObject);
 	if (GridManager)
 	{
 		GridManager->RemoveGridData(GridId, bRebuildGrids);
 	}
 }
 
-bool UNavLocalGridManager::FindLocalNavigationGridPath(UObject* WorldContext, const FVector& Start, const FVector& End, TArray<FVector>& PathPoints)
+bool UNavLocalGridManager::FindLocalNavigationGridPath(UObject* WorldContextObject, const FVector& Start, const FVector& End, TArray<FVector>& PathPoints)
 {
-	UNavLocalGridManager* GridManager = UNavLocalGridManager::GetCurrent(WorldContext);
+	UNavLocalGridManager* GridManager = UNavLocalGridManager::GetCurrent(WorldContextObject);
 	if (GridManager)
 	{
 		return GridManager->FindPath(Start, End, PathPoints);
@@ -411,9 +411,9 @@ bool UNavLocalGridManager::FindLocalNavigationGridPath(UObject* WorldContext, co
 	return false;
 }
 
-bool UNavLocalGridManager::SetLocalNavigationGridDensity(UObject* WorldContext, float CellSize)
+bool UNavLocalGridManager::SetLocalNavigationGridDensity(UObject* WorldContextObject, float CellSize)
 {
-	UNavLocalGridManager* GridManager = UNavLocalGridManager::GetCurrent(WorldContext);
+	UNavLocalGridManager* GridManager = UNavLocalGridManager::GetCurrent(WorldContextObject);
 	if (GridManager)
 	{
 		return GridManager->SetCellSize(CellSize);
diff --git a/Engine/Source/Runtime/AIModule/Private/Perception/AIPerceptionSystem.cpp b/Engine/Source/Runtime/AIModule/Private/Perception/AIPerceptionSystem.cpp
--- a/Engine/Source/Runtime/AIModule/Private/Perception/AIPerceptionSystem.cpp
+++ b/Engine/Source/Runtime/AIModule/Private/Perception/AIPerceptionSystem.cpp
@@ -574,12 +574,12 @@ void UAIPerceptionSystem::RegisterAllPawnsAsSourcesForSense(FAISenseID SenseID)
 	}
 }
 
-bool UAIPerceptionSystem::RegisterPerceptionStimuliSource(UObject* WorldContext, TSubclassOf<UAISense> Sense, AActor* Target)
+bool UAIPerceptionSystem::RegisterPerceptionStimuliSource(UObject* WorldContextObject, TSubclassOf<UAISense> Sense, AActor* Target)
 {
 	bool bResult = false;
 	if (Sense && Target)
 	{
-		UWorld* World = GEngine->GetWorldFromContextObject(WorldContext);
+		UWorld* World = GEngine->GetWorldFromContextObject(WorldContextObject);
 		if (World && World->GetAISystem())
 		{
 			UAISystem* AISys = Cast<UAISystem>(World->GetAISystem());
@@ -617,10 +617,10 @@ void UAIPerceptionSystem::OnPerceptionStimuliSourceEndPlay(AActor* Actor, EEndPl
 	bStimuliSourcesRefreshRequired = true;
 }
 
-TSubclassOf<UAISense> UAIPerceptionSystem::GetSenseClassForStimulus(UObject* WorldContext, const FAIStimulus& Stimulus)
+TSubclassOf<UAISense> UAIPerceptionSystem::GetSenseClassForStimulus(UObject* WorldContextObject, const FAIStimulus& Stimulus)
 {
 	TSubclassOf<UAISense> Result = nullptr;
-	UAIPerceptionSystem* PercSys = GetCurrent(WorldContext);
+	UAIPerceptionSystem* PercSys = GetCurrent(WorldContextObject);
 	if (PercSys && PercSys->Senses.IsValidIndex(Stimulus.Type) && PercSys->Senses[Stimulus.Type] != nullptr)
 	{
 		Result = PercSys->Senses[Stimulus.Type]->GetClass();
@@ -650,9 +650,9 @@ void UAIPerceptionSystem::ReportEvent(UAISenseEvent* PerceptionEvent)
 	}
 }
 
-void UAIPerceptionSystem::ReportPerceptionEvent(UObject* WorldContext, UAISenseEvent* PerceptionEvent)
+void UAIPerceptionSystem::ReportPerceptionEvent(UObject* WorldContextObject, UAISenseEvent* PerceptionEvent)
 {
-	UAIPerceptionSystem* PerceptionSys = GetCurrent(WorldContext);
+	UAIPerceptionSystem* PerceptionSys = GetCurrent(WorldContextObject);
 	if (PerceptionSys != nullptr)
 	{
 		PerceptionSys->ReportEvent(PerceptionEvent);
diff --git a/Engine/Source/Runtime/AIModule/Private/Perception/AISense_Damage.cpp b/Engine/Source/Runtime/AIModule/Private/Perception/AISense_Damage.cpp
--- a/Engine/Source/Runtime/AIModule/Private/Perception/AISense_Damage.cpp
+++ b/Engine/Source/Runtime/AIModule/Private/Perception/AISense_Damage.cpp
@@ -127,9 +127,9 @@ void UAISense_Damage::RegisterWrappedEvent(UAISenseEvent& PerceptionEvent)
 	}
 }
 
-void UAISense_Damage::ReportDamageEvent(UObject* WorldContext, AActor* DamagedActor, AActor* Instigator, float DamageAmount, FVector EventLocation, FVector HitLocation)
+void UAISense_Damage::ReportDamageEvent(UObject* WorldContextObject, AActor* DamagedActor, AActor* Instigator, float DamageAmount, FVector EventLocation, FVector HitLocation)
 {
-	UAIPerceptionSystem* PerceptionSystem = UAIPerceptionSystem::GetCurrent(WorldContext);
+	UAIPerceptionSystem* PerceptionSystem = UAIPerceptionSystem::GetCurrent(WorldContextObject);
 	if (PerceptionSystem)
 	{
 		FAIDamageEvent Event(DamagedActor, Instigator, DamageAmount, EventLocation, HitLocation);
diff --git a/Engine/Source/Runtime/AIModule/Private/Perception/AISense_Hearing.cpp b/Engine/Source/Runtime/AIModule/Private/Perception/AISense_Hearing.cpp
--- a/Engine/Source/Runtime/AIModule/Private/Perception/AISense_Hearing.cpp
+++ b/Engine/Source/Runtime/AIModule/Private/Perception/AISense_Hearing.cpp
@@ -69,9 +69,9 @@ UAISense_Hearing::UAISense_Hearing(const FObjectInitializer& ObjectInitializer)
 	}
 }
 
-void UAISense_Hearing::ReportNoiseEvent(UObject* WorldContext, FVector NoiseLocation, float Loudness, AActor* Instigator, float MaxRange, FName Tag)
+void UAISense_Hearing::ReportNoiseEvent(UObject* WorldContextObject, FVector NoiseLocation, float Loudness, AActor* Instigator, float MaxRange, FName Tag)
 {
-	UAIPerceptionSystem* PerceptionSystem = UAIPerceptionSystem::GetCurrent(WorldContext);
+	UAIPerceptionSystem* PerceptionSystem = UAIPerceptionSystem::GetCurrent(WorldContextObject);
 	if (PerceptionSystem)
 	{
 		FAINoiseEvent Event(Instigator, NoiseLocation, Loudness, MaxRange, Tag);
diff --git a/Engine/Source/Runtime/AnimGraphRuntime/Private/BoneControllers/AnimNode_AnimDynamics.cpp b/Engine/Source/Runtime/AnimGraphRuntime/Private/BoneControllers/AnimNode_AnimDynamics.cpp
--- a/Engine/Source/Runtime/AnimGraphRuntime/Private/BoneControllers/AnimNode_AnimDynamics.cpp
+++ b/Engine/Source/Runtime/AnimGraphRuntime/Private/BoneControllers/AnimNode_AnimDynamics.cpp
@@ -718,9 +718,9 @@ void FAnimNode_AnimDynamics::PreUpdate(const UAnimInstance* InAnimInstance)
 
 				// Setup wind data
 				Body->bWindEnabled = true;
-				Scene->GetWindParameters_GameThread(SkelComp->ComponentToWorld.TransformPosition(Body->Pose.Position), Body->WindData.WindDirection, Body->WindData.WindSpeed, WindMinGust, WindMaxGust);
+				Scene->GetWindParameters_GameThread(SkelComp->GetComponentTransform().TransformPosition(Body->Pose.Position), Body->WindData.WindDirection, Body->WindData.WindSpeed, WindMinGust, WindMaxGust);
 
-				Body->WindData.WindDirection = SkelComp->ComponentToWorld.Inverse().TransformVector(Body->WindData.WindDirection);
+				Body->WindData.WindDirection = SkelComp->GetComponentTransform().Inverse().TransformVector(Body->WindData.WindDirection);
 				Body->WindData.WindAdaption = FMath::FRandRange(0.0f, 2.0f);
 				Body->WindData.BodyWindScale = WindScale;
 			}
diff --git a/Engine/Source/Runtime/AnimGraphRuntime/Public/AnimNodes/AnimNode_LayeredBoneBlend.h b/Engine/Source/Runtime/AnimGraphRuntime/Public/AnimNodes/AnimNode_LayeredBoneBlend.h
--- a/Engine/Source/Runtime/AnimGraphRuntime/Public/AnimNodes/AnimNode_LayeredBoneBlend.h
+++ b/Engine/Source/Runtime/AnimGraphRuntime/Public/AnimNodes/AnimNode_LayeredBoneBlend.h
@@ -10,7 +10,7 @@
 #include "AnimNode_LayeredBoneBlend.generated.h"
 
 // Layered blend (per bone); has dynamic number of blendposes that can blend per different bone sets
-USTRUCT()
+USTRUCT(BlueprintInternalUseOnly)
 struct ANIMGRAPHRUNTIME_API FAnimNode_LayeredBoneBlend : public FAnimNode_Base
 {
 	GENERATED_USTRUCT_BODY()
@@ -23,7 +23,7 @@ public:
 	UPROPERTY(EditAnywhere, BlueprintReadWrite, editfixedsize, Category=Links)
 	TArray<FPoseLink> BlendPoses;
 
-	UPROPERTY(EditAnywhere, BlueprintReadWrite, editfixedsize, Category=Config)
+	UPROPERTY(EditAnywhere, editfixedsize, Category=Config)
 	TArray<FInputBlendPose> LayerSetup;
 
 	UPROPERTY(EditAnywhere, BlueprintReadWrite, editfixedsize, Category=Runtime, meta=(PinShownByDefault))
diff --git a/Engine/Source/Runtime/AnimGraphRuntime/Public/AnimNodes/AnimNode_RandomPlayer.h b/Engine/Source/Runtime/AnimGraphRuntime/Public/AnimNodes/AnimNode_RandomPlayer.h
--- a/Engine/Source/Runtime/AnimGraphRuntime/Public/AnimNodes/AnimNode_RandomPlayer.h
+++ b/Engine/Source/Runtime/AnimGraphRuntime/Public/AnimNodes/AnimNode_RandomPlayer.h
@@ -50,7 +50,7 @@ struct FRandomAnimPlayData
 /** The random player node holds a list of sequences and parameter ranges which will be played continuously
   * In a random order. If shuffle mode is enabled then each entry will be played once before repeating any
   */
-USTRUCT()
+USTRUCT(BlueprintInternalUseOnly)
 struct FRandomPlayerSequenceEntry
 {
 	GENERATED_BODY()
@@ -91,11 +91,11 @@ struct FRandomPlayerSequenceEntry
 	float MaxPlayRate;
 
 	/** Blending properties used when this entry is blending in ontop of another entry */
-	UPROPERTY(EditAnywhere, BlueprintReadWrite, Category = "Settings")
+	UPROPERTY(EditAnywhere, Category = "Settings")
 	FAlphaBlend BlendIn;
 };
 
-USTRUCT()
+USTRUCT(BlueprintInternalUseOnly)
 struct ANIMGRAPHRUNTIME_API FAnimNode_RandomPlayer : public FAnimNode_Base
 {
 	GENERATED_BODY()
diff --git a/Engine/Source/Runtime/AnimGraphRuntime/Public/BoneControllers/AnimNode_AnimDynamics.h b/Engine/Source/Runtime/AnimGraphRuntime/Public/BoneControllers/AnimNode_AnimDynamics.h
--- a/Engine/Source/Runtime/AnimGraphRuntime/Public/BoneControllers/AnimNode_AnimDynamics.h
+++ b/Engine/Source/Runtime/AnimGraphRuntime/Public/BoneControllers/AnimNode_AnimDynamics.h
@@ -204,7 +204,7 @@ struct FAnimPhysSphericalLimit
 	ESphericalLimitType LimitType;
 };
 
-USTRUCT()
+USTRUCT(BlueprintInternalUseOnly)
 struct ANIMGRAPHRUNTIME_API FAnimNode_AnimDynamics : public FAnimNode_SkeletalControlBase
 {
 	GENERATED_BODY();
@@ -216,27 +216,27 @@ struct ANIMGRAPHRUNTIME_API FAnimNode_AnimDynamics : public FAnimNode_SkeletalCo
 	AnimPhysSimSpaceType SimulationSpace;
 
 	/** When in BoneRelative sim space, the simulation will use this bone as the origin */
-	UPROPERTY(EditAnywhere, BlueprintReadWrite, Category = Setup)
+	UPROPERTY(EditAnywhere, Category = Setup)
 	FBoneReference RelativeSpaceBone;
 
 	/** Set to true to use the solver to simulate a connected chain */
-	UPROPERTY(EditAnywhere, BlueprintReadWrite, Category = Setup)
+	UPROPERTY(EditAnywhere, Category = Setup)
 	bool bChain;
 
 	/** The bone to attach the physics body to, if bChain is true this is the top of the chain */
-	UPROPERTY(EditAnywhere, BlueprintReadWrite, Category = Setup)
+	UPROPERTY(EditAnywhere, Category = Setup)
 	FBoneReference BoundBone;
 	
 	/** If bChain is true this is the bottom of the chain, otherwise ignored */
-	UPROPERTY(EditAnywhere, BlueprintReadWrite, Category = Setup, meta=(EditCondition = bChain))
+	UPROPERTY(EditAnywhere, Category = Setup, meta=(EditCondition = bChain))
 	FBoneReference ChainEnd;
 
 	/** Extents of the box to use for simulation */
-	UPROPERTY(EditAnywhere, BlueprintReadWrite, Category = Setup)
+	UPROPERTY(EditAnywhere, Category = Setup)
 	FVector BoxExtents;
 
 	/** Vector relative to the body being simulated to attach the constraint to */
-	UPROPERTY(EditAnywhere, BlueprintReadWrite, Category = Setup)
+	UPROPERTY(EditAnywhere, Category = Setup)
 	FVector LocalJointOffset;
 
 	/** Scale for gravity, higher values increase forces due to gravity */
@@ -244,11 +244,11 @@ struct ANIMGRAPHRUNTIME_API FAnimNode_AnimDynamics : public FAnimNode_SkeletalCo
 	float GravityScale;
 
 	/** If true the body will attempt to spring back to its initial position */
-	UPROPERTY(EditAnywhere, BlueprintReadWrite, Category = Setup)
+	UPROPERTY(EditAnywhere, Category = Setup)
 	bool bLinearSpring;
 
 	/** If true the body will attempt to align itself with the specified angular target */
-	UPROPERTY(EditAnywhere, BlueprintReadWrite, Category = Setup)
+	UPROPERTY(EditAnywhere, Category = Setup)
 	bool bAngularSpring;
 
 	/** Spring constant to use when calculating linear springs, higher values mean a stronger spring.*/
@@ -260,18 +260,18 @@ struct ANIMGRAPHRUNTIME_API FAnimNode_AnimDynamics : public FAnimNode_SkeletalCo
 	float AngularSpringConstant;
 
 	/** Whether or not wind is enabled for the bodies in this simulation */
-	UPROPERTY(EditAnywhere, BlueprintReadWrite, Category = Wind)
+	UPROPERTY(EditAnywhere, Category = Wind)
 	bool bEnableWind;
 
 	UPROPERTY(Transient)
 	bool bWindWasEnabled;
 
 	/** Scale to apply to calculated wind velocities in the solver */
-	UPROPERTY(EditAnywhere, BlueprintReadWrite, Category = Wind)
+	UPROPERTY(EditAnywhere, Category = Wind)
 	float WindScale;
 
 	/** If true, the override value will be used for linear damping */
-	UPROPERTY(EditAnywhere, AdvancedDisplay, BlueprintReadWrite, Category = Setup)
+	UPROPERTY(EditAnywhere, AdvancedDisplay, Category = Setup)
 	bool bOverrideLinearDamping;
 
 	/** Overridden linear damping value */
@@ -279,7 +279,7 @@ struct ANIMGRAPHRUNTIME_API FAnimNode_AnimDynamics : public FAnimNode_SkeletalCo
 	float LinearDampingOverride;
 
 	/** If true, the override value will be used for angular damping */
-	UPROPERTY(EditAnywhere, AdvancedDisplay, BlueprintReadWrite, Category = Setup)
+	UPROPERTY(EditAnywhere, AdvancedDisplay, Category = Setup)
 	bool bOverrideAngularDamping;
 
 	/** Overridden angular damping value */
@@ -291,7 +291,7 @@ struct ANIMGRAPHRUNTIME_API FAnimNode_AnimDynamics : public FAnimNode_SkeletalCo
 	 *  in check. When using single-body systems sometimes angular forces will look like they are "catching-up" with
 	 *  the mesh, if that's the case override this and push it towards 1.0f until it settles correctly
 	 */
-	UPROPERTY(EditAnywhere, AdvancedDisplay, BlueprintReadWrite, Category = Setup)
+	UPROPERTY(EditAnywhere, AdvancedDisplay, Category = Setup)
 	bool bOverrideAngularBias;
 
 	/** Overridden angular bias value
@@ -303,11 +303,11 @@ struct ANIMGRAPHRUNTIME_API FAnimNode_AnimDynamics : public FAnimNode_SkeletalCo
 	float AngularBiasOverride;
 
 	/** If true we will perform physics update, otherwise skip - allows visualisation of the initial state of the bodies */
-	UPROPERTY(EditAnywhere, AdvancedDisplay, BlueprintReadWrite, Category = Setup)
+	UPROPERTY(EditAnywhere, AdvancedDisplay, Category = Setup)
 	bool bDoUpdate;
 
 	/** If true we will perform bone transform evaluation, otherwise skip - allows visualisation of the initial anim state compared to the physics sim */
-	UPROPERTY(EditAnywhere, AdvancedDisplay, BlueprintReadWrite, Category = Setup)
+	UPROPERTY(EditAnywhere, AdvancedDisplay, Category = Setup)
 	bool bDoEval;
 
 	/** Number of update passes on the linear and angular limits before we solve the position of the bodies recommended to be four times the value of NumSolverIterationsPostUpdate */
@@ -319,23 +319,23 @@ struct ANIMGRAPHRUNTIME_API FAnimNode_AnimDynamics : public FAnimNode_SkeletalCo
 	int32 NumSolverIterationsPostUpdate;
 
 	/** Data describing the constraints we will apply to the body */
-	UPROPERTY(EditAnywhere, BlueprintReadWrite, Category = Constraint)
+	UPROPERTY(EditAnywhere, Category = Constraint)
 	FAnimPhysConstraintSetup ConstraintSetup;
 
 	/** Whether to evaluate planar limits */
-	UPROPERTY(EditAnywhere, BlueprintReadWrite, Category=PlanarLimit)
+	UPROPERTY(EditAnywhere, Category=PlanarLimit)
 	bool bUsePlanarLimit;
 
 	/** List of available planar limits for this node */
-	UPROPERTY(EditAnywhere, BlueprintReadWrite, Category=PlanarLimit)
+	UPROPERTY(EditAnywhere, Category=PlanarLimit)
 	TArray<FAnimPhysPlanarLimit> PlanarLimits;
 
 	/** Whether to evaluate spherical limits */
-	UPROPERTY(EditAnywhere, BlueprintReadWrite, Category = SphericalLimit)
+	UPROPERTY(EditAnywhere, Category = SphericalLimit)
 	bool bUseSphericalLimits;
 
 	/** List of available spherical limits for this node */
-	UPROPERTY(EditAnywhere, BlueprintReadWrite, Category = SphericalLimit)
+	UPROPERTY(EditAnywhere, Category = SphericalLimit)
 	TArray<FAnimPhysSphericalLimit> SphericalLimits;
 
 	/** Resolution method for planar limits */
diff --git a/Engine/Source/Runtime/AnimGraphRuntime/Public/BoneControllers/AnimNode_CopyBone.h b/Engine/Source/Runtime/AnimGraphRuntime/Public/BoneControllers/AnimNode_CopyBone.h
--- a/Engine/Source/Runtime/AnimGraphRuntime/Public/BoneControllers/AnimNode_CopyBone.h
+++ b/Engine/Source/Runtime/AnimGraphRuntime/Public/BoneControllers/AnimNode_CopyBone.h
@@ -15,17 +15,17 @@ class USkeletalMeshComponent;
  *	Simple controller to copy a bone's transform to another one.
  */
 
-USTRUCT()
+USTRUCT(BlueprintInternalUseOnly)
 struct ANIMGRAPHRUNTIME_API FAnimNode_CopyBone : public FAnimNode_SkeletalControlBase
 {
 	GENERATED_USTRUCT_BODY()
 
 	/** Source Bone Name to get transform from */
-	UPROPERTY(EditAnywhere, BlueprintReadWrite, Category = Copy)
+	UPROPERTY(EditAnywhere, Category = Copy)
 	FBoneReference SourceBone;
 
 	/** Name of bone to control. This is the main bone chain to modify from. **/
-	UPROPERTY(EditAnywhere, BlueprintReadWrite, Category=Copy) 
+	UPROPERTY(EditAnywhere, Category=Copy) 
 	FBoneReference TargetBone;
 
 	/** If Translation should be copied */
@@ -41,7 +41,7 @@ struct ANIMGRAPHRUNTIME_API FAnimNode_CopyBone : public FAnimNode_SkeletalContro
 	bool bCopyScale;
 
 	/** Space to convert transforms into prior to copying components */
-	UPROPERTY(EditAnywhere, BlueprintReadWrite, Category = Copy)
+	UPROPERTY(EditAnywhere, Category = Copy)
 	TEnumAsByte<EBoneControlSpace> ControlSpace;
 
 	FAnimNode_CopyBone();
diff --git a/Engine/Source/Runtime/AnimGraphRuntime/Public/BoneControllers/AnimNode_CopyBoneDelta.h b/Engine/Source/Runtime/AnimGraphRuntime/Public/BoneControllers/AnimNode_CopyBoneDelta.h
--- a/Engine/Source/Runtime/AnimGraphRuntime/Public/BoneControllers/AnimNode_CopyBoneDelta.h
+++ b/Engine/Source/Runtime/AnimGraphRuntime/Public/BoneControllers/AnimNode_CopyBoneDelta.h
@@ -23,15 +23,15 @@ enum class CopyBoneDeltaMode : uint8
  *	instead of the copy bone node which copies the absolute transform
  */
 
-USTRUCT()
+USTRUCT(BlueprintInternalUseOnly)
 struct ANIMGRAPHRUNTIME_API FAnimNode_CopyBoneDelta : public FAnimNode_SkeletalControlBase
 {
 	GENERATED_USTRUCT_BODY()
 
-	UPROPERTY(EditAnywhere, BlueprintReadWrite, Category = Settings)
+	UPROPERTY(EditAnywhere, Category = Settings)
 	FBoneReference SourceBone;
 
-	UPROPERTY(EditAnywhere, BlueprintReadWrite, Category = Settings)
+	UPROPERTY(EditAnywhere, Category = Settings)
 	FBoneReference TargetBone;
 
 	UPROPERTY(EditAnywhere, BlueprintReadWrite, Category = Settings, meta = (PinShownByDefault))
@@ -43,7 +43,7 @@ struct ANIMGRAPHRUNTIME_API FAnimNode_CopyBoneDelta : public FAnimNode_SkeletalC
 	UPROPERTY(EditAnywhere, BlueprintReadWrite, Category = Settings, meta = (PinShownByDefault))
 	bool bCopyScale;
 
-	UPROPERTY(EditAnywhere, BlueprintReadWrite, Category = Settings)
+	UPROPERTY(EditAnywhere, Category = Settings)
 	CopyBoneDeltaMode CopyMode;
 
 	UPROPERTY(EditAnywhere, BlueprintReadWrite, Category = Settings, meta = (PinHiddenByDefault))
diff --git a/Engine/Source/Runtime/AnimGraphRuntime/Public/BoneControllers/AnimNode_Fabrik.h b/Engine/Source/Runtime/AnimGraphRuntime/Public/BoneControllers/AnimNode_Fabrik.h
--- a/Engine/Source/Runtime/AnimGraphRuntime/Public/BoneControllers/AnimNode_Fabrik.h
+++ b/Engine/Source/Runtime/AnimGraphRuntime/Public/BoneControllers/AnimNode_Fabrik.h
@@ -54,7 +54,7 @@ public:
 	}
 };
 
-USTRUCT()
+USTRUCT(BlueprintInternalUseOnly)
 struct ANIMGRAPHRUNTIME_API FAnimNode_Fabrik : public FAnimNode_SkeletalControlBase
 {
 	GENERATED_USTRUCT_BODY()
@@ -64,34 +64,34 @@ struct ANIMGRAPHRUNTIME_API FAnimNode_Fabrik : public FAnimNode_SkeletalControlB
 	FTransform EffectorTransform;
 
 	/** Reference frame of Effector Transform. */
-	UPROPERTY(EditAnywhere, BlueprintReadWrite, Category = EndEffector)
+	UPROPERTY(EditAnywhere, Category = EndEffector)
 	TEnumAsByte<enum EBoneControlSpace> EffectorTransformSpace;
 
 	/** If EffectorTransformSpace is a bone, this is the bone to use. **/
-	UPROPERTY(EditAnywhere, BlueprintReadWrite, Category = EndEffector)
+	UPROPERTY(EditAnywhere,  Category = EndEffector)
 	FBoneReference EffectorTransformBone;
 
-	UPROPERTY(EditAnywhere, BlueprintReadWrite, Category = EndEffector)
+	UPROPERTY(EditAnywhere, Category = EndEffector)
 	TEnumAsByte<enum EBoneRotationSource> EffectorRotationSource;
 
 	/** Name of tip bone */
-	UPROPERTY(EditAnywhere, BlueprintReadWrite, Category = Solver)
+	UPROPERTY(EditAnywhere, Category = Solver)
 	FBoneReference TipBone;
 
 	/** Name of the root bone*/
-	UPROPERTY(EditAnywhere, BlueprintReadWrite, Category = Solver)
+	UPROPERTY(EditAnywhere, Category = Solver)
 	FBoneReference RootBone;
 
 	/** Tolerance for final tip location delta from EffectorLocation*/
-	UPROPERTY(EditAnywhere, BlueprintReadWrite, Category = Solver)
+	UPROPERTY(EditAnywhere, Category = Solver)
 	float Precision;
 
 	/** Maximum number of iterations allowed, to control performance. */
-	UPROPERTY(EditAnywhere, BlueprintReadWrite, Category = Solver)
+	UPROPERTY(EditAnywhere, Category = Solver)
 	int32 MaxIterations;
 
 	/** Toggle drawing of axes to debug joint rotation*/
-	UPROPERTY(EditAnywhere, BlueprintReadWrite, Category = Solver)
+	UPROPERTY(EditAnywhere, Category = Solver)
 	bool bEnableDebugDraw;
 
 public:
diff --git a/Engine/Source/Runtime/AnimGraphRuntime/Public/BoneControllers/AnimNode_HandIKRetargeting.h b/Engine/Source/Runtime/AnimGraphRuntime/Public/BoneControllers/AnimNode_HandIKRetargeting.h
--- a/Engine/Source/Runtime/AnimGraphRuntime/Public/BoneControllers/AnimNode_HandIKRetargeting.h
+++ b/Engine/Source/Runtime/AnimGraphRuntime/Public/BoneControllers/AnimNode_HandIKRetargeting.h
@@ -18,29 +18,29 @@ class USkeletalMeshComponent;
  * This is used so characters of different proportions can handle the same props.
  */
 
-USTRUCT()
+USTRUCT(BlueprintInternalUseOnly)
 struct ANIMGRAPHRUNTIME_API FAnimNode_HandIKRetargeting : public FAnimNode_SkeletalControlBase
 {
 	GENERATED_USTRUCT_BODY()
 
 	/** Bone for Right Hand FK */
-	UPROPERTY(EditAnywhere, BlueprintReadWrite, Category = "HandIKRetargeting")
+	UPROPERTY(EditAnywhere, Category = "HandIKRetargeting")
 	FBoneReference RightHandFK;
 
 	/** Bone for Left Hand FK */
-	UPROPERTY(EditAnywhere, BlueprintReadWrite, Category = "HandIKRetargeting")
+	UPROPERTY(EditAnywhere, Category = "HandIKRetargeting")
 	FBoneReference LeftHandFK;
 
 	/** Bone for Right Hand IK */
-	UPROPERTY(EditAnywhere, BlueprintReadWrite, Category = "HandIKRetargeting")
+	UPROPERTY(EditAnywhere, Category = "HandIKRetargeting")
 	FBoneReference RightHandIK;
 
 	/** Bone for Left Hand IK */
-	UPROPERTY(EditAnywhere, BlueprintReadWrite, Category = "HandIKRetargeting")
+	UPROPERTY(EditAnywhere, Category = "HandIKRetargeting")
 	FBoneReference LeftHandIK;
 
 	/** IK Bones to move. */
-	UPROPERTY(EditAnywhere, BlueprintReadWrite, Category = "HandIKRetargeting")
+	UPROPERTY(EditAnywhere, Category = "HandIKRetargeting")
 	TArray<FBoneReference> IKBonesToMove;
 
 	/** Which hand to favor. 0.5 is equal weight for both, 1 = right hand, 0 = left hand. */
diff --git a/Engine/Source/Runtime/AnimGraphRuntime/Public/BoneControllers/AnimNode_LookAt.h b/Engine/Source/Runtime/AnimGraphRuntime/Public/BoneControllers/AnimNode_LookAt.h
--- a/Engine/Source/Runtime/AnimGraphRuntime/Public/BoneControllers/AnimNode_LookAt.h
+++ b/Engine/Source/Runtime/AnimGraphRuntime/Public/BoneControllers/AnimNode_LookAt.h
@@ -33,21 +33,21 @@ namespace EInterpolationBlend
 /**
  *	Simple controller that make a bone to look at the point or another bone
  */
-USTRUCT()
+USTRUCT(BlueprintInternalUseOnly)
 struct ANIMGRAPHRUNTIME_API FAnimNode_LookAt : public FAnimNode_SkeletalControlBase
 {
 	GENERATED_USTRUCT_BODY()
 
 	/** Name of bone to control. This is the main bone chain to modify from. **/
-	UPROPERTY(EditAnywhere, BlueprintReadWrite, Category=SkeletalControl) 
+	UPROPERTY(EditAnywhere, Category=SkeletalControl) 
 	FBoneReference BoneToModify;
 
 	/** Target Bone to look at - You can use  LookAtLocation if you need offset from this point. That location will be used in their local space. **/
-	UPROPERTY(EditAnywhere, BlueprintReadWrite, Category=Target)
+	UPROPERTY(EditAnywhere, Category=Target)
 	FBoneReference LookAtBone;
 
 	/** Target socket to look at. Used if LookAtBone is empty. - You can use  LookAtLocation if you need offset from this point. That location will be used in their local space. **/
-	UPROPERTY(EditAnywhere, BlueprintReadWrite, Category = Target)
+	UPROPERTY(EditAnywhere, Category = Target)
 	FName LookAtSocket;
 
 	/** Target Offset. It's in world space if LookAtBone is empty or it is based on LookAtBone or LookAtSocket in their local space*/
@@ -62,11 +62,11 @@ struct ANIMGRAPHRUNTIME_API FAnimNode_LookAt : public FAnimNode_SkeletalControlB
 	UPROPERTY()
 	FVector	CustomLookAtAxis_DEPRECATED;
 
-	UPROPERTY(EditAnywhere, BlueprintReadWrite, Category = SkeletalControl)
+	UPROPERTY(EditAnywhere, Category = SkeletalControl)
 	FAxis LookAt_Axis;
 
 	/** Whether or not to use Look up axis */
-	UPROPERTY(EditAnywhere, BlueprintReadWrite, Category=SkeletalControl)
+	UPROPERTY(EditAnywhere, Category=SkeletalControl)
 	bool bUseLookUpAxis;
 
 	/** Look up axis in local space */
@@ -77,7 +77,7 @@ struct ANIMGRAPHRUNTIME_API FAnimNode_LookAt : public FAnimNode_SkeletalControlB
 	UPROPERTY()
 	FVector	CustomLookUpAxis_DEPRECATED;
 
-	UPROPERTY(EditAnywhere, BlueprintReadWrite, Category = SkeletalControl)
+	UPROPERTY(EditAnywhere, Category = SkeletalControl)
 	FAxis LookUp_Axis;
 
 	/** Look at Clamp value in degree - if you're look at axis is Z, only X, Y degree of clamp will be used*/
diff --git a/Engine/Source/Runtime/AnimGraphRuntime/Public/BoneControllers/AnimNode_ModifyBone.h b/Engine/Source/Runtime/AnimGraphRuntime/Public/BoneControllers/AnimNode_ModifyBone.h
--- a/Engine/Source/Runtime/AnimGraphRuntime/Public/BoneControllers/AnimNode_ModifyBone.h
+++ b/Engine/Source/Runtime/AnimGraphRuntime/Public/BoneControllers/AnimNode_ModifyBone.h
@@ -27,13 +27,13 @@ enum EBoneModificationMode
 /**
  *	Simple controller that replaces or adds to the translation/rotation of a single bone.
  */
-USTRUCT()
+USTRUCT(BlueprintInternalUseOnly)
 struct ANIMGRAPHRUNTIME_API FAnimNode_ModifyBone : public FAnimNode_SkeletalControlBase
 {
 	GENERATED_USTRUCT_BODY()
 
 	/** Name of bone to control. This is the main bone chain to modify from. **/
-	UPROPERTY(EditAnywhere, BlueprintReadWrite, Category=SkeletalControl) 
+	UPROPERTY(EditAnywhere, Category=SkeletalControl) 
 	FBoneReference BoneToModify;
 
 	/** New translation of bone to apply. */
@@ -49,27 +49,27 @@ struct ANIMGRAPHRUNTIME_API FAnimNode_ModifyBone : public FAnimNode_SkeletalCont
 	FVector Scale;
 
 	/** Whether and how to modify the translation of this bone. */
-	UPROPERTY(EditAnywhere, BlueprintReadWrite, Category=Translation)
+	UPROPERTY(EditAnywhere, Category=Translation)
 	TEnumAsByte<EBoneModificationMode> TranslationMode;
 
 	/** Whether and how to modify the translation of this bone. */
-	UPROPERTY(EditAnywhere, BlueprintReadWrite, Category=Rotation)
+	UPROPERTY(EditAnywhere, Category=Rotation)
 	TEnumAsByte<EBoneModificationMode> RotationMode;
 
 	/** Whether and how to modify the translation of this bone. */
-	UPROPERTY(EditAnywhere, BlueprintReadWrite, Category=Scale)
+	UPROPERTY(EditAnywhere, Category=Scale)
 	TEnumAsByte<EBoneModificationMode> ScaleMode;
 
 	/** Reference frame to apply Translation in. */
-	UPROPERTY(EditAnywhere, BlueprintReadWrite, Category=Translation)
+	UPROPERTY(EditAnywhere, Category=Translation)
 	TEnumAsByte<enum EBoneControlSpace> TranslationSpace;
 
 	/** Reference frame to apply Rotation in. */
-	UPROPERTY(EditAnywhere, BlueprintReadWrite, Category=Rotation)
+	UPROPERTY(EditAnywhere, Category=Rotation)
 	TEnumAsByte<enum EBoneControlSpace> RotationSpace;
 
 	/** Reference frame to apply Scale in. */
-	UPROPERTY(EditAnywhere, BlueprintReadWrite, Category=Scale)
+	UPROPERTY(EditAnywhere, Category=Scale)
 	TEnumAsByte<enum EBoneControlSpace> ScaleSpace;
 
 	FAnimNode_ModifyBone();
diff --git a/Engine/Source/Runtime/AnimGraphRuntime/Public/BoneControllers/AnimNode_RotationMultiplier.h b/Engine/Source/Runtime/AnimGraphRuntime/Public/BoneControllers/AnimNode_RotationMultiplier.h
--- a/Engine/Source/Runtime/AnimGraphRuntime/Public/BoneControllers/AnimNode_RotationMultiplier.h
+++ b/Engine/Source/Runtime/AnimGraphRuntime/Public/BoneControllers/AnimNode_RotationMultiplier.h
@@ -15,17 +15,17 @@ class USkeletalMeshComponent;
  *	Simple controller that multiplies scalar value to the translation/rotation/scale of a single bone.
  */
 
-USTRUCT()
+USTRUCT(BlueprintInternalUseOnly)
 struct ANIMGRAPHRUNTIME_API FAnimNode_RotationMultiplier : public FAnimNode_SkeletalControlBase
 {
 	GENERATED_USTRUCT_BODY()
 
 	/** Name of bone to control. This is the main bone chain to modify from. */
-	UPROPERTY(EditAnywhere, BlueprintReadWrite, Category=Multiplier) 
+	UPROPERTY(EditAnywhere, Category=Multiplier) 
 	FBoneReference	TargetBone;
 
 	/** Source to get transform from */
-	UPROPERTY(EditAnywhere, BlueprintReadWrite, Category=Multiplier)
+	UPROPERTY(EditAnywhere, Category=Multiplier)
 	FBoneReference	SourceBone;
 
 	// To make these to be easily pin-hookable, I'm not making it struct, but each variable
@@ -33,10 +33,10 @@ struct ANIMGRAPHRUNTIME_API FAnimNode_RotationMultiplier : public FAnimNode_Skel
 	UPROPERTY(EditAnywhere, BlueprintReadWrite, Category=Multiplier, meta=(PinShownByDefault))
 	float	Multiplier;
 
-	UPROPERTY(EditAnywhere, BlueprintReadWrite, Category=Multiplier)
+	UPROPERTY(EditAnywhere, Category=Multiplier)
 	TEnumAsByte<EBoneAxis> RotationAxisToRefer;
 	
-	UPROPERTY(EditAnywhere, BlueprintReadWrite, Category = Multiplier)
+	UPROPERTY(EditAnywhere, Category = Multiplier)
 	bool bIsAdditive;	
 
 	FAnimNode_RotationMultiplier();
diff --git a/Engine/Source/Runtime/AnimGraphRuntime/Public/BoneControllers/AnimNode_ScaleChainLength.h b/Engine/Source/Runtime/AnimGraphRuntime/Public/BoneControllers/AnimNode_ScaleChainLength.h
--- a/Engine/Source/Runtime/AnimGraphRuntime/Public/BoneControllers/AnimNode_ScaleChainLength.h
+++ b/Engine/Source/Runtime/AnimGraphRuntime/Public/BoneControllers/AnimNode_ScaleChainLength.h
@@ -15,22 +15,22 @@ class USkeletalMeshComponent;
 /**
  *	
  */
-USTRUCT()
+USTRUCT(BlueprintInternalUseOnly)
 struct ANIMGRAPHRUNTIME_API FAnimNode_ScaleChainLength : public FAnimNode_Base
 {
 	GENERATED_USTRUCT_BODY()
 
-	UPROPERTY(EditAnywhere, BlueprintReadWrite, Category = Links)
+	UPROPERTY(EditAnywhere, Category = Links)
 	FPoseLink InputPose;
 
 	/** Default chain length, as animated. */
 	UPROPERTY(EditAnywhere, BlueprintReadWrite, Category = ScaleChainLength, meta = (PinHiddenByDefault))
 	float DefaultChainLength;
 
-	UPROPERTY(EditAnywhere, BlueprintReadWrite, Category = ScaleChainLength)
+	UPROPERTY(EditAnywhere, Category = ScaleChainLength)
 	FBoneReference ChainStartBone;
 
-	UPROPERTY(EditAnywhere, BlueprintReadWrite, Category = ScaleChainLength)
+	UPROPERTY(EditAnywhere, Category = ScaleChainLength)
 	FBoneReference ChainEndBone;
 
 	UPROPERTY(EditAnywhere, BlueprintReadWrite, Category = ScaleChainLength, meta = (PinShownByDefault))
@@ -39,7 +39,7 @@ struct ANIMGRAPHRUNTIME_API FAnimNode_ScaleChainLength : public FAnimNode_Base
 	UPROPERTY(EditAnywhere, BlueprintReadWrite, Category = Settings, meta = (PinShownByDefault))
 	mutable float Alpha;
 
-	UPROPERTY(EditAnywhere, BlueprintReadWrite, Category = Settings)
+	UPROPERTY(EditAnywhere, Category = Settings)
 	FInputScaleBias AlphaScaleBias;
 
 	UPROPERTY(Transient)
diff --git a/Engine/Source/Runtime/AnimGraphRuntime/Public/BoneControllers/AnimNode_SplineIK.h b/Engine/Source/Runtime/AnimGraphRuntime/Public/BoneControllers/AnimNode_SplineIK.h
--- a/Engine/Source/Runtime/AnimGraphRuntime/Public/BoneControllers/AnimNode_SplineIK.h
+++ b/Engine/Source/Runtime/AnimGraphRuntime/Public/BoneControllers/AnimNode_SplineIK.h
@@ -44,29 +44,29 @@ struct FSplineIKCachedBoneData
 	int32 RefSkeletonIndex;
 };
 
-USTRUCT()
+USTRUCT(BlueprintInternalUseOnly)
 struct ANIMGRAPHRUNTIME_API FAnimNode_SplineIK : public FAnimNode_SkeletalControlBase
 {
 	GENERATED_BODY()
 
 	/** Name of root bone from which the spline extends **/
-	UPROPERTY(EditAnywhere, BlueprintReadWrite, Category = "Parameters")
+	UPROPERTY(EditAnywhere, Category = "Parameters")
 	FBoneReference StartBone;
 
 	/** Name of bone at the end of the spline chain. Bones after this will not be altered by the controller. */
-	UPROPERTY(EditAnywhere, BlueprintReadWrite, Category = "Parameters")
+	UPROPERTY(EditAnywhere, Category = "Parameters")
 	FBoneReference EndBone;
 
 	/** Axis of the controlled bone (ie the direction of the spline) to use as the direction for the curve. */
-	UPROPERTY(EditAnywhere, BlueprintReadWrite, Category = "Parameters")
+	UPROPERTY(EditAnywhere, Category = "Parameters")
 	ESplineBoneAxis BoneAxis;
 
 	/** The number of points in the spline if we are specifying it directly */
-	UPROPERTY(EditAnywhere, BlueprintReadWrite, Category = "Parameters")
+	UPROPERTY(EditAnywhere, Category = "Parameters")
 	bool bAutoCalculateSpline;
 
 	/** The number of points in the spline if we are not auto-calculating */
-	UPROPERTY(EditAnywhere, BlueprintReadWrite, Category = "Parameters", meta = (ClampMin = 2, UIMin = 2, EditCondition = "!bAutoCalculateSpline"))
+	UPROPERTY(EditAnywhere, Category = "Parameters", meta = (ClampMin = 2, UIMin = 2, EditCondition = "!bAutoCalculateSpline"))
 	int32 PointCount;
 
 	/** Transforms applied to spline points **/
@@ -86,7 +86,7 @@ struct ANIMGRAPHRUNTIME_API FAnimNode_SplineIK : public FAnimNode_SkeletalContro
 	float TwistEnd;
 
 	/** How to interpolate twist along the length of the spline */
-	UPROPERTY(EditAnywhere, BlueprintReadWrite, Category = "Parameters")
+	UPROPERTY(EditAnywhere, Category = "Parameters")
 	FAlphaBlend TwistBlend;
 
 	/**
diff --git a/Engine/Source/Runtime/AnimGraphRuntime/Public/BoneControllers/AnimNode_SpringBone.h b/Engine/Source/Runtime/AnimGraphRuntime/Public/BoneControllers/AnimNode_SpringBone.h
--- a/Engine/Source/Runtime/AnimGraphRuntime/Public/BoneControllers/AnimNode_SpringBone.h
+++ b/Engine/Source/Runtime/AnimGraphRuntime/Public/BoneControllers/AnimNode_SpringBone.h
@@ -16,17 +16,17 @@ class USkeletalMeshComponent;
  *	Simple controller that replaces or adds to the translation/rotation of a single bone.
  */
 
-USTRUCT()
+USTRUCT(BlueprintInternalUseOnly)
 struct ANIMGRAPHRUNTIME_API FAnimNode_SpringBone : public FAnimNode_SkeletalControlBase
 {
 	GENERATED_USTRUCT_BODY()
 
 	/** Name of bone to control. This is the main bone chain to modify from. **/
-	UPROPERTY(EditAnywhere, BlueprintReadWrite, Category=Spring) 
+	UPROPERTY(EditAnywhere, Category=Spring) 
 	FBoneReference SpringBone;
 
 	/** Limit the amount that a bone can stretch from its ref-pose length. */
-	UPROPERTY(EditAnywhere, BlueprintReadWrite, Category=Spring)
+	UPROPERTY(EditAnywhere, Category=Spring)
 	bool bLimitDisplacement;
 
 	/** If bLimitDisplacement is true, this indicates how long a bone can stretch beyond its length in the ref-pose. */
@@ -34,15 +34,15 @@ struct ANIMGRAPHRUNTIME_API FAnimNode_SpringBone : public FAnimNode_SkeletalCont
 	float MaxDisplacement;
 
 	/** Stiffness of spring */
-	UPROPERTY(EditAnywhere, BlueprintReadWrite, Category=Spring)
+	UPROPERTY(EditAnywhere, Category=Spring)
 	float SpringStiffness;
 
 	/** Damping of spring */
-	UPROPERTY(EditAnywhere, BlueprintReadWrite, Category=Spring)
+	UPROPERTY(EditAnywhere, Category=Spring)
 	float SpringDamping;
 
 	/** If spring stretches more than this, reset it. Useful for catching teleports etc */
-	UPROPERTY(EditAnywhere, BlueprintReadWrite, Category=Spring)
+	UPROPERTY(EditAnywhere, Category=Spring)
 	float ErrorResetThresh;
 
 	/** If true, Z position is always correct, no spring applied */
@@ -50,27 +50,27 @@ struct ANIMGRAPHRUNTIME_API FAnimNode_SpringBone : public FAnimNode_SkeletalCont
 	bool bNoZSpring_DEPRECATED;
 
 	/** If true take the spring calculation for translation in X */
-	UPROPERTY(EditAnywhere, BlueprintReadWrite, Category=FilterChannels)
+	UPROPERTY(EditAnywhere, Category=FilterChannels)
 	bool bTranslateX;
 
 	/** If true take the spring calculation for translation in Y */
-	UPROPERTY(EditAnywhere, BlueprintReadWrite, Category=FilterChannels)
+	UPROPERTY(EditAnywhere, Category=FilterChannels)
 	bool bTranslateY;
 
 	/** If true take the spring calculation for translation in Z */
-	UPROPERTY(EditAnywhere, BlueprintReadWrite, Category=FilterChannels)
+	UPROPERTY(EditAnywhere, Category=FilterChannels)
 	bool bTranslateZ;
 
 	/** If true take the spring calculation for rotation in X */
-	UPROPERTY(EditAnywhere, BlueprintReadWrite, Category=FilterChannels)
+	UPROPERTY(EditAnywhere, Category=FilterChannels)
 	bool bRotateX;
 
 	/** If true take the spring calculation for rotation in Y */
-	UPROPERTY(EditAnywhere, BlueprintReadWrite, Category=FilterChannels)
+	UPROPERTY(EditAnywhere, Category=FilterChannels)
 	bool bRotateY;
 
 	/** If true take the spring calculation for rotation in Z */
-	UPROPERTY(EditAnywhere, BlueprintReadWrite, Category=FilterChannels)
+	UPROPERTY(EditAnywhere, Category=FilterChannels)
 	bool bRotateZ;
 
 	/** Internal use - Amount of time we need to simulate. */
diff --git a/Engine/Source/Runtime/AnimGraphRuntime/Public/BoneControllers/AnimNode_Trail.h b/Engine/Source/Runtime/AnimGraphRuntime/Public/BoneControllers/AnimNode_Trail.h
--- a/Engine/Source/Runtime/AnimGraphRuntime/Public/BoneControllers/AnimNode_Trail.h
+++ b/Engine/Source/Runtime/AnimGraphRuntime/Public/BoneControllers/AnimNode_Trail.h
@@ -24,25 +24,25 @@ struct FPerJointTrailSetup
  * Trail Controller
  */
 
-USTRUCT()
+USTRUCT(BlueprintInternalUseOnly)
 struct ANIMGRAPHRUNTIME_API FAnimNode_Trail : public FAnimNode_SkeletalControlBase
 {
 	GENERATED_USTRUCT_BODY()
 
 	/** Reference to the active bone in the hierarchy to modify. */
-	UPROPERTY(EditAnywhere, BlueprintReadWrite, Category=Trail)
+	UPROPERTY(EditAnywhere, Category=Trail)
 	FBoneReference TrailBone;
 
 	/** Number of bones above the active one in the hierarchy to modify. ChainLength should be at least 2. */
-	UPROPERTY(EditAnywhere, BlueprintReadWrite, Category = Trail, meta = (ClampMin = "2", UIMin = "2"))
+	UPROPERTY(EditAnywhere, Category = Trail, meta = (ClampMin = "2", UIMin = "2"))
 	int32	ChainLength;
 
 	/** Axis of the bones to point along trail. */
-	UPROPERTY(EditAnywhere, BlueprintReadWrite, Category=Trail)
+	UPROPERTY(EditAnywhere, Category=Trail)
 	TEnumAsByte<EAxis::Type>	ChainBoneAxis;
 
 	/** Invert the direction specified in ChainBoneAxis. */
-	UPROPERTY(EditAnywhere, BlueprintReadWrite, Category=Trail)
+	UPROPERTY(EditAnywhere, Category=Trail)
 	bool	bInvertChainBoneAxis;
 
 	/** How quickly we 'relax' the bones to their animated positions. Deprecated. Replaced to TrailRelaxationCurve */
@@ -50,27 +50,27 @@ struct ANIMGRAPHRUNTIME_API FAnimNode_Trail : public FAnimNode_SkeletalControlBa
 	float	TrailRelaxation_DEPRECATED;
 
 	/** How quickly we 'relax' the bones to their animated positions. Time 0 will map to top root joint, time 1 will map to the bottom joint. */
-	UPROPERTY(EditAnywhere, BlueprintReadWrite, Category=Trail, meta=(CustomizeProperty))
+	UPROPERTY(EditAnywhere, Category=Trail, meta=(CustomizeProperty))
 	FRuntimeFloatCurve TrailRelaxationSpeed;
 
 	/** Limit the amount that a bone can stretch from its ref-pose length. */
-	UPROPERTY(EditAnywhere, BlueprintReadWrite, Category=Limit)
+	UPROPERTY(EditAnywhere, Category=Limit)
 	bool	bLimitStretch;
 
 	/** If bLimitStretch is true, this indicates how long a bone can stretch beyond its length in the ref-pose. */
-	UPROPERTY(EditAnywhere, BlueprintReadWrite, Category=Limit)
+	UPROPERTY(EditAnywhere, Category=Limit)
 	float	StretchLimit;
 
 	/** 'Fake' velocity applied to bones. */
-	UPROPERTY(EditAnywhere, BlueprintReadWrite, Category=Velocity)
+	UPROPERTY(EditAnywhere, Category=Velocity)
 	FVector	FakeVelocity;
 
 	/** Whether 'fake' velocity should be applied in actor or world space. */
-	UPROPERTY(EditAnywhere, BlueprintReadWrite, Category=Velocity)
+	UPROPERTY(EditAnywhere,  Category=Velocity)
 	bool	bActorSpaceFakeVel;
 
 	/** Base Joint to calculate velocity from. If none, it will use Component's World Transform. . */
-	UPROPERTY(EditAnywhere, BlueprintReadWrite, Category=Velocity)
+	UPROPERTY(EditAnywhere, Category=Velocity)
 	FBoneReference BaseJoint;
 
 	/** Internal use - we need the timestep to do the relaxation in CalculateNewBoneTransforms. */
diff --git a/Engine/Source/Runtime/AnimGraphRuntime/Public/BoneControllers/AnimNode_TwoBoneIK.h b/Engine/Source/Runtime/AnimGraphRuntime/Public/BoneControllers/AnimNode_TwoBoneIK.h
--- a/Engine/Source/Runtime/AnimGraphRuntime/Public/BoneControllers/AnimNode_TwoBoneIK.h
+++ b/Engine/Source/Runtime/AnimGraphRuntime/Public/BoneControllers/AnimNode_TwoBoneIK.h
@@ -15,13 +15,13 @@ class USkeletalMeshComponent;
  * Simple 2 Bone IK Controller.
  */
 
-USTRUCT()
+USTRUCT(BlueprintInternalUseOnly)
 struct ANIMGRAPHRUNTIME_API FAnimNode_TwoBoneIK : public FAnimNode_SkeletalControlBase
 {
 	GENERATED_USTRUCT_BODY()
 	
 	/** Name of bone to control. This is the main bone chain to modify from. **/
-	UPROPERTY(EditAnywhere, BlueprintReadWrite, Category=IK)
+	UPROPERTY(EditAnywhere, Category=IK)
 	FBoneReference IKBone;
 
 	/** Effector Location. Target Location to reach. */
@@ -33,19 +33,19 @@ struct ANIMGRAPHRUNTIME_API FAnimNode_TwoBoneIK : public FAnimNode_SkeletalContr
 	FVector JointTargetLocation;
 
 	/** If EffectorLocationSpace is a bone, this is the bone to use. **/
-	UPROPERTY(EditAnywhere, BlueprintReadWrite, Category=EndEffector)
+	UPROPERTY(EditAnywhere, Category=EndEffector)
 	FName EffectorSpaceBoneName;
 
 	/** Set end bone to use End Effector rotation */
-	UPROPERTY(EditAnywhere, BlueprintReadWrite, Category=EndEffector)
+	UPROPERTY(EditAnywhere, Category=EndEffector)
 	uint32 bTakeRotationFromEffectorSpace:1;
 
 	/** Keep local rotation of end bone */
-	UPROPERTY(EditAnywhere, BlueprintReadWrite, Category=EndEffector)
+	UPROPERTY(EditAnywhere, Category=EndEffector)
 	uint32 bMaintainEffectorRelRot:1;
 
 	/** Should stretching be allowed, to be prevent over extension */
-	UPROPERTY(EditAnywhere, BlueprintReadWrite, Category=IK)
+	UPROPERTY(EditAnywhere, Category=IK)
 	uint32 bAllowStretching:1;
 
 	/** Limits to use if stretching is allowed - old property DEPRECATED */
@@ -53,23 +53,23 @@ struct ANIMGRAPHRUNTIME_API FAnimNode_TwoBoneIK : public FAnimNode_SkeletalContr
 	FVector2D StretchLimits_DEPRECATED;
 
 	/** Limits to use if stretching is allowed. This value determines when to start stretch. For example, 0.9 means once it reaches 90% of the whole length of the limb, it will start apply. */
-	UPROPERTY(EditAnywhere, BlueprintReadWrite, Category=IK, meta = (editcondition = "bAllowStretching", ClampMin = "0.0", UIMin = "0.0"))
+	UPROPERTY(EditAnywhere, Category=IK, meta = (editcondition = "bAllowStretching", ClampMin = "0.0", UIMin = "0.0"))
 	float StartStretchRatio;
 
 	/** Limits to use if stretching is allowed. This value determins what is the max stretch scale. For example, 1.5 means it will stretch until 150 % of the whole length of the limb.*/
-	UPROPERTY(EditAnywhere, BlueprintReadWrite, Category = IK, meta = (editcondition = "bAllowStretching", ClampMin = "0.0", UIMin = "0.0"))
+	UPROPERTY(EditAnywhere, Category = IK, meta = (editcondition = "bAllowStretching", ClampMin = "0.0", UIMin = "0.0"))
 	float MaxStretchScale;
 
 	/** Reference frame of Effector Location. */
-	UPROPERTY(EditAnywhere, BlueprintReadWrite, Category = IK)
+	UPROPERTY(EditAnywhere, Category = IK)
 	TEnumAsByte<enum EBoneControlSpace> EffectorLocationSpace;
 
 	/** Reference frame of Joint Target Location. */
-	UPROPERTY(EditAnywhere, BlueprintReadWrite, Category=JointTarget)
+	UPROPERTY(EditAnywhere, Category=JointTarget)
 	TEnumAsByte<enum EBoneControlSpace> JointTargetLocationSpace;
 
 	/** If JointTargetSpaceBoneName is a bone, this is the bone to use. **/
-	UPROPERTY(EditAnywhere, BlueprintReadWrite, Category=JointTarget)
+	UPROPERTY(EditAnywhere, Category=JointTarget)
 	FName JointTargetSpaceBoneName;
 
 	FAnimNode_TwoBoneIK();
diff --git a/Engine/Source/Runtime/AssetRegistry/Private/AssetBundleData.cpp b/Engine/Source/Runtime/AssetRegistry/Private/AssetBundleData.cpp
--- a/Engine/Source/Runtime/AssetRegistry/Private/AssetBundleData.cpp
+++ b/Engine/Source/Runtime/AssetRegistry/Private/AssetBundleData.cpp
@@ -8,6 +8,9 @@ bool FAssetBundleData::SetFromAssetData(const FAssetData& AssetData)
 {
 	FString TagValue;
 
+	// Register that we're reading string assets for a specific package
+	FStringAssetReferenceSerializationScope SerializationScope(AssetData.PackageName, FAssetBundleData::StaticStruct()->GetFName(), EStringAssetReferenceCollectType::AlwaysCollect);
+
 	if (AssetData.GetTagValue(FAssetBundleData::StaticStruct()->GetFName(), TagValue))
 	{
 		if (FAssetBundleData::StaticStruct()->ImportText(*TagValue, this, nullptr, PPF_None, (FOutputDevice*)GWarn, AssetData.AssetName.ToString()))
@@ -77,7 +80,41 @@ void FAssetBundleData::AddBundleAssets(FName BundleName, const TArray<FStringAss
 	}
 }
 
+void FAssetBundleData::SetBundleAssets(FName BundleName, TArray<FStringAssetReference>&& AssetPaths)
+{
+	FAssetBundleEntry* FoundEntry = FindEntry(FPrimaryAssetId(), BundleName);
+
+	if (!FoundEntry)
+	{
+		FoundEntry = new(Bundles) FAssetBundleEntry(FPrimaryAssetId(), BundleName);
+	}
+
+	FoundEntry->BundleAssets = AssetPaths;
+}
+
 void FAssetBundleData::Reset()
 {
 	Bundles.Reset();
+}
+
+bool FAssetBundleData::ExportTextItem(FString& ValueStr, FAssetBundleData const& DefaultValue, UObject* Parent, int32 PortFlags, UObject* ExportRootScope) const
+{
+	if (Bundles.Num() == 0)
+	{
+		// Empty, don't write anything to avoid it cluttering the asset registry tags
+		return true;
+	}
+	// Not empty, do normal export
+	return false;
+}
+
+bool FAssetBundleData::ImportTextItem(const TCHAR*& Buffer, int32 PortFlags, UObject* Parent, FOutputDevice* ErrorText)
+{
+	if (*Buffer != TEXT('('))
+	{
+		// Empty, don't read/write anything
+		return true;
+	}
+	// Full structure, do normal parse
+	return false;
 }
\ No newline at end of file
diff --git a/Engine/Source/Runtime/AssetRegistry/Private/AssetData.cpp b/Engine/Source/Runtime/AssetRegistry/Private/AssetData.cpp
--- a/Engine/Source/Runtime/AssetRegistry/Private/AssetData.cpp
+++ b/Engine/Source/Runtime/AssetRegistry/Private/AssetData.cpp
@@ -1,5 +1,42 @@
 // Copyright 1998-2017 Epic Games, Inc. All Rights Reserved.
 
 #include "AssetData.h"
+#include "Serialization/CustomVersion.h"
 
 DEFINE_LOG_CATEGORY(LogAssetData);
+
+// Register Asset Registry version
+const FGuid FAssetRegistryVersion::GUID(0x717F9EE7, 0xE9B0493A, 0x88B39132, 0x1B388107);
+FCustomVersionRegistration GRegisterAssetRegistryVersion(FAssetRegistryVersion::GUID, FAssetRegistryVersion::LatestVersion, TEXT("AssetRegistry"));
+
+bool FAssetRegistryVersion::SerializeVersion(FArchive& Ar, FAssetRegistryVersion::Type& Version)
+{
+	FGuid Guid = FAssetRegistryVersion::GUID;
+
+	if (Ar.IsLoading())
+	{
+		Version = FAssetRegistryVersion::PreVersioning;
+	}
+
+	Ar << Guid;
+
+	if (Ar.IsError())
+	{
+		return false;
+	}
+
+	if (Guid == FAssetRegistryVersion::GUID)
+	{
+		int32 VersionInt = Version;
+		Ar << VersionInt;
+		Version = (FAssetRegistryVersion::Type)VersionInt;
+
+		Ar.SetCustomVersion(Guid, VersionInt, TEXT("AssetRegistry"));
+	}
+	else
+	{
+		return false;
+	}
+
+	return !Ar.IsError();
+}
\ No newline at end of file
diff --git a/Engine/Source/Runtime/AssetRegistry/Private/AssetDataGatherer.cpp b/Engine/Source/Runtime/AssetRegistry/Private/AssetDataGatherer.cpp
--- a/Engine/Source/Runtime/AssetRegistry/Private/AssetDataGatherer.cpp
+++ b/Engine/Source/Runtime/AssetRegistry/Private/AssetDataGatherer.cpp
@@ -13,13 +13,12 @@
 
 namespace AssetDataGathererConstants
 {
-	static const int32 CacheSerializationVersion = 9;
+	static const int32 CacheSerializationVersion = 10;
 	static const int32 MaxFilesToDiscoverBeforeFlush = 2500;
 	static const int32 MaxFilesToGatherBeforeFlush = 250;
 	static const int32 MaxFilesToProcessBeforeCacheWrite = 50000;
 }
 
-
 namespace
 {
 	class FLambdaDirectoryStatVisitor : public IPlatformFile::FDirectoryStatVisitor
@@ -453,21 +452,20 @@ uint32 FAssetDataGatherer::Run()
 {
 	int32 CacheSerializationVersion = AssetDataGathererConstants::CacheSerializationVersion;
 	
-	static const bool bUsingWorldAssets = FAssetRegistry::IsUsingWorldAssets();
-	if ( bUsingWorldAssets )
-	{
-		// Bump the serialization version to refresh the cache when switching between -WorldAssets and without.
-		// This is a temporary hack just while WorldAssets are under development
-		CacheSerializationVersion++;
-	}
-
 	if ( bLoadAndSaveCache )
 	{
 		// load the cached data
 		FNameTableArchiveReader CachedAssetDataReader(CacheSerializationVersion, CacheFilename);
 		if (!CachedAssetDataReader.IsError())
 		{
-			SerializeCache(CachedAssetDataReader);
+			FAssetRegistryVersion::Type Version = FAssetRegistryVersion::LatestVersion;
+			if (FAssetRegistryVersion::SerializeVersion(CachedAssetDataReader, Version))
+			{
+				SerializeCache(CachedAssetDataReader);
+
+				DependencyResults.Reserve(DiskCachedAssetDataMap.Num());
+				AssetResults.Reserve(DiskCachedAssetDataMap.Num());
+			}
 		}
 	}
 
@@ -484,6 +482,10 @@ uint32 FAssetDataGatherer::Run()
 	auto WriteAssetCacheFile = [&]()
 	{
 		FNameTableArchiveWriter CachedAssetDataWriter(CacheSerializationVersion, CacheFilename);
+
+		FAssetRegistryVersion::Type Version = FAssetRegistryVersion::LatestVersion;
+		FAssetRegistryVersion::SerializeVersion(CachedAssetDataWriter, Version);
+
 		SerializeCache(CachedAssetDataWriter);
 
 		NumFilesProcessedSinceLastCacheSave = 0;
@@ -880,11 +882,10 @@ void FAssetDataGatherer::SerializeCache(FArchive& Ar)
 	if (Ar.IsSaving())
 	{
 		// save out by walking the TMap
-		for (auto CacheIt = NewCachedAssetDataMap.CreateConstIterator(); CacheIt; ++CacheIt)
+		for (TPair<FName, FDiskCachedAssetData*>& Pair : NewCachedAssetDataMap)
 		{
-			FName PackageName = CacheIt.Key();
-			Ar << PackageName;
-			Ar << *CacheIt.Value();
+			Ar << Pair.Key;
+			Pair.Value->SerializeForCache(Ar);
 		}
 	}
 	else
@@ -907,7 +908,7 @@ void FAssetDataGatherer::SerializeCache(FArchive& Ar)
 			FDiskCachedAssetData& CachedAssetData = DiskCachedAssetDataMap.Add(PackageName);
 
 			// Now load the data
-			Ar << CachedAssetData;
+			CachedAssetData.SerializeForCache(Ar);
 
 			if (Ar.IsError())
 			{
diff --git a/Engine/Source/Runtime/AssetRegistry/Private/AssetRegistry.cpp b/Engine/Source/Runtime/AssetRegistry/Private/AssetRegistry.cpp
--- a/Engine/Source/Runtime/AssetRegistry/Private/AssetRegistry.cpp
+++ b/Engine/Source/Runtime/AssetRegistry/Private/AssetRegistry.cpp
@@ -1,6 +1,5 @@
 // Copyright 1998-2017 Epic Games, Inc. All Rights Reserved.
 
-
 #include "AssetRegistry.h"
 #include "Misc/CommandLine.h"
 #include "Misc/FileHelper.h"
@@ -16,6 +15,9 @@
 #include "PackageReader.h"
 #include "GenericPlatform/GenericPlatformChunkInstall.h"
 #include "IPluginManager.h"
+#include "Misc/CoreDelegates.h"
+#include "UObject/ConstructorHelpers.h"
+#include "Misc/RedirectCollector.h"
 
 #if WITH_EDITOR
 #include "IDirectoryWatcher.h"
@@ -110,8 +112,8 @@ FAssetRegistry::FAssetRegistry()
 	UE_LOG(LogAssetRegistry, Log, TEXT( "FAssetRegistry took %0.4f seconds to start up" ), FPlatformTime::Seconds() - StartupStartTime );
 
 #if WITH_EDITOR
-	// Commandlets and in-game don't listen for directory changes
-	if ( !IsRunningCommandlet() && GIsEditor)
+	// In-game doesn't listen for directory changes
+	if (GIsEditor)
 	{
 		FDirectoryWatcherModule& DirectoryWatcherModule = FModuleManager::LoadModuleChecked<FDirectoryWatcherModule>(TEXT("DirectoryWatcher"));
 		IDirectoryWatcher* DirectoryWatcher = DirectoryWatcherModule.Get();
@@ -150,8 +152,11 @@ FAssetRegistry::FAssetRegistry()
 
 	// Listen for new content paths being added or removed at runtime.  These are usually plugin-specific asset paths that
 	// will be loaded a bit later on.
-	FPackageName::OnContentPathMounted().AddRaw( this, &FAssetRegistry::OnContentPathMounted );
-	FPackageName::OnContentPathDismounted().AddRaw( this, &FAssetRegistry::OnContentPathDismounted );
+	FPackageName::OnContentPathMounted().AddRaw(this, &FAssetRegistry::OnContentPathMounted);
+	FPackageName::OnContentPathDismounted().AddRaw(this, &FAssetRegistry::OnContentPathDismounted);
+
+	// If we were called before engine has fully initialized, refresh classes on initialize. If not this won't do anything as it already happened
+	FCoreDelegates::OnPostEngineInit.AddRaw(this, &FAssetRegistry::RefreshNativeClasses);
 }
 
 void FAssetRegistry::InitializeSerializationOptions(FAssetRegistrySerializationOptions& Options, const FString& PlatformIniName) const
@@ -168,6 +173,7 @@ void FAssetRegistry::InitializeSerializationOptions(FAssetRegistrySerializationO
 	PlatformEngineIni.GetBool(TEXT("AssetRegistry"), TEXT("bSerializeManageDependencies"), Options.bSerializeManageDependencies);
 	PlatformEngineIni.GetBool(TEXT("AssetRegistry"), TEXT("bSerializePackageData"), Options.bSerializePackageData);
 	PlatformEngineIni.GetBool(TEXT("AssetRegistry"), TEXT("bUseAssetRegistryTagsWhitelistInsteadOfBlacklist"), Options.bUseAssetRegistryTagsWhitelistInsteadOfBlacklist);
+	PlatformEngineIni.GetBool(TEXT("AssetRegistry"), TEXT("bFilterAssetDataWithNoTags"), Options.bFilterAssetDataWithNoTags);
 		
 	TArray<FString> FilterlistItems;
 	if (Options.bUseAssetRegistryTagsWhitelistInsteadOfBlacklist)
@@ -265,6 +271,15 @@ void FAssetRegistry::CollectCodeGeneratorClasses()
 	}
 }
 
+void FAssetRegistry::RefreshNativeClasses()
+{
+	// Native classes have changed so reinitialize code generator and serialization options
+	CollectCodeGeneratorClasses();
+
+	// Read default serialization options
+	InitializeSerializationOptions(SerializationOptions);
+}
+
 FAssetRegistry::~FAssetRegistry()
 {
 	// Make sure the asset search thread is closed
@@ -275,12 +290,12 @@ FAssetRegistry::~FAssetRegistry()
 	}
 
 	// Stop listening for content mount point events
-	FPackageName::OnContentPathMounted().RemoveAll( this );
-	FPackageName::OnContentPathDismounted().RemoveAll( this );
+	FPackageName::OnContentPathMounted().RemoveAll(this);
+	FPackageName::OnContentPathDismounted().RemoveAll(this);
+	FCoreDelegates::OnPostEngineInit.RemoveAll(this);
 
-	// Commandlets dont listen for directory changes
 #if WITH_EDITOR
-	if ( !IsRunningCommandlet() && GIsEditor )
+	if (GIsEditor)
 	{
 		// If the directory module is still loaded, unregister any delegates
 		if ( FModuleManager::Get().IsModuleLoaded("DirectoryWatcher") )
@@ -335,6 +350,14 @@ void FAssetRegistry::SearchAllAssets(bool bSynchronousSearch)
 	{
 		const bool bForceRescan = false;
 		ScanPathsAndFilesSynchronous(PathsToSearch, TArray<FString>(), bForceRescan, EAssetDataCacheMode::UseMonolithicCache);
+
+#if WITH_EDITOR
+		if (IsRunningCommandlet())
+		{
+			// Update redirectors
+			UpdateRedirectCollector();
+		}
+#endif
 	}
 	else if ( !BackgroundAssetSearch.IsValid() )
 	{
@@ -484,19 +507,18 @@ bool FAssetRegistry::GetAssets(const FARFilter& InFilter, TArray<FAssetData>& Ou
 					}
 				}
 
-				// Find the group names
-				FString GroupNamesStr;
-				FString AssetNameStr;
-				Obj->GetPathName(InMemoryPackage).Split(TEXT("."), &GroupNamesStr, &AssetNameStr, ESearchCase::CaseSensitive, ESearchDir::FromEnd);
-
-				TMap<FName, FString> TagMap;
+				FAssetDataTagMap TagMap;
 				for (UObject::FAssetRegistryTag& AssetRegistryTag : ObjectTags)
 				{
-					TagMap.Add(AssetRegistryTag.Name, AssetRegistryTag.Value);
+					if (AssetRegistryTag.Name != NAME_None && !AssetRegistryTag.Value.IsEmpty())
+					{
+						// Don't add empty tags
+						TagMap.Add(AssetRegistryTag.Name, AssetRegistryTag.Value);
+					}
 				}
 
 				// This asset is in memory and passes all filters
-				FAssetData* AssetData = new (OutAssetData)FAssetData(PackageName, PackagePath, FName(*GroupNamesStr), Obj->GetFName(), Obj->GetClass()->GetFName(), TagMap, InMemoryPackage->GetChunkIDs(), InMemoryPackage->GetPackageFlags());
+				FAssetData* AssetData = new (OutAssetData)FAssetData(PackageName, PackagePath, Obj->GetFName(), Obj->GetClass()->GetFName(), TagMap, InMemoryPackage->GetChunkIDs(), InMemoryPackage->GetPackageFlags());
 			}
 		};
 
@@ -633,6 +655,40 @@ const FAssetPackageData* FAssetRegistry::GetAssetPackageData(FName PackageName)
 	return State.GetAssetPackageData(PackageName);
 }
 
+FName FAssetRegistry::GetRedirectedObjectPath(const FName ObjectPath) const
+{
+	FString RedirectedPath = ObjectPath.ToString();
+	FAssetData DestinationData = GetAssetByObjectPath(ObjectPath);
+	TSet<FString> SeenPaths;
+	SeenPaths.Add(RedirectedPath);
+
+	// Need to follow chain of redirectors
+	while (DestinationData.IsRedirector())
+	{
+		if (DestinationData.GetTagValue("DestinationObject", RedirectedPath))
+		{
+			ConstructorHelpers::StripObjectClass(RedirectedPath);
+			if (SeenPaths.Contains(RedirectedPath))
+			{
+				// Recursive, bail
+				DestinationData = FAssetData();
+			}
+			else
+			{
+				SeenPaths.Add(RedirectedPath);
+				DestinationData = GetAssetByObjectPath(FName(*RedirectedPath), true);
+			}
+		}
+		else
+		{
+			// Can't extract
+			DestinationData = FAssetData();
+		}
+	}
+
+	return FName(*RedirectedPath);
+}
+
 bool FAssetRegistry::GetAncestorClassNames(FName ClassName, TArray<FName>& OutAncestorClassNames) const
 {
 	// Start with the cached inheritance map
@@ -1114,8 +1170,18 @@ void FAssetRegistry::AssetDeleted(UObject* DeletedAsset)
 			}
 		}
 
+		FAssetData AssetDataDeleted = FAssetData(DeletedAsset);
+
+#if WITH_EDITOR
+		if (bInitialSearchCompleted && AssetDataDeleted.IsRedirector())
+		{
+			// Need to remove from GRedirectCollector
+			GRedirectCollector.RemoveAssetPathRedirection(AssetDataDeleted.ObjectPath.ToString());
+		}
+#endif
+
 		// Let subscribers know that the asset was removed from the registry
-		AssetRemovedEvent.Broadcast(FAssetData(DeletedAsset));
+		AssetRemovedEvent.Broadcast(AssetDataDeleted);
 
 		// Notify listeners that an in-memory asset was just deleted
 		InMemoryAssetDeletedEvent.Broadcast(DeletedAsset);
@@ -1199,6 +1265,12 @@ void FAssetRegistry::Tick(float DeltaTime)
 {
 	double TickStartTime = FPlatformTime::Seconds();
 
+	if (DeltaTime < 0)
+	{
+		// Force a full flush
+		TickStartTime = -1;
+	}
+
 	// Gather results from the background search
 	bool bIsSearching = false;
 	TArray<double> SearchTimes;
@@ -1270,6 +1342,10 @@ void FAssetRegistry::Tick(float DeltaTime)
 	{
 		if ( !bInitialSearchCompleted )
 		{
+#if WITH_EDITOR
+			// update redirectors
+			UpdateRedirectCollector();
+#endif
 			UE_LOG(LogAssetRegistry, Verbose, TEXT("### Time spent amortizing search results: %0.4f seconds"), TotalAmortizeTime);
 			UE_LOG(LogAssetRegistry, Log, TEXT("Asset discovery search completed in %0.4f seconds"), FPlatformTime::Seconds() - FullSearchStartTime);
 
@@ -1286,7 +1362,6 @@ void FAssetRegistry::Tick(float DeltaTime)
 	}
 }
 
-
 bool FAssetRegistry::IsUsingWorldAssets()
 {
 	return !FParse::Param(FCommandLine::Get(), TEXT("DisableWorldAssets"));
@@ -1311,6 +1386,22 @@ void FAssetRegistry::Serialize(FArchive& Ar)
 	}
 }
 
+uint32 FAssetRegistry::GetAllocatedSize(bool bLogDetailed) const
+{
+	uint32 StateSize = State.GetAllocatedSize(bLogDetailed);
+
+	uint32 StaticSize = sizeof(FAssetRegistry) + CachedEmptyPackages.GetAllocatedSize() + CachedInheritanceMap.GetAllocatedSize() + EditSearchableNameDelegates.GetAllocatedSize() + ClassGeneratorNames.GetAllocatedSize() + SerializationOptions.CookFilterlistTagsByClass.GetAllocatedSize();
+	uint32 SearchSize = BackgroundAssetResults.GetAllocatedSize() + BackgroundPathResults.GetAllocatedSize() + BackgroundDependencyResults.GetAllocatedSize() + BackgroundCookedPackageNamesWithoutAssetDataResults.GetAllocatedSize() + SynchronouslyScannedPathsAndFiles.GetAllocatedSize() + CachedPathTree.GetAllocatedSize();
+
+	if (bLogDetailed)
+	{
+		UE_LOG(LogAssetRegistry, Log, TEXT("AssetRegistry Static Size: %dk"), StaticSize / 1024);
+		UE_LOG(LogAssetRegistry, Log, TEXT("AssetRegistry Search Size: %dk"), SearchSize / 1024);
+	}
+	
+	return StaticSize + StaticSize + SearchSize;
+}
+
 void FAssetRegistry::LoadPackageRegistryData(FArchive& Ar, TArray<FAssetData*> &AssetDataList ) const
 {
 	
@@ -1330,7 +1421,7 @@ void FAssetRegistry::LoadPackageRegistryData(FArchive& Ar, TArray<FAssetData*> &
 void FAssetRegistry::SaveRegistryData(FArchive& Ar, TMap<FName, FAssetData*>& Data, TArray<FName>* InMaps /* = nullptr */)
 {
 	FAssetRegistryState TempState;
-	InitializeTemporaryAssetRegistryState(TempState, SerializationOptions, Data);
+	InitializeTemporaryAssetRegistryState(TempState, SerializationOptions, false, Data);
 
 	TempState.Serialize(Ar, SerializationOptions);
 }
@@ -1357,11 +1448,11 @@ void FAssetRegistry::LoadRegistryData(FArchive& Ar, TMap<FName, FAssetData*>& Da
 	}
 }
 
-void FAssetRegistry::InitializeTemporaryAssetRegistryState(FAssetRegistryState& OutState, const FAssetRegistrySerializationOptions& Options, const TMap<FName, FAssetData*>& OverrideData) const
+void FAssetRegistry::InitializeTemporaryAssetRegistryState(FAssetRegistryState& OutState, const FAssetRegistrySerializationOptions& Options, bool bRefreshExisting, const TMap<FName, FAssetData*>& OverrideData) const
 {
 	const TMap<FName, FAssetData*>& DataToUse = OverrideData.Num() > 0 ? OverrideData : State.CachedAssetsByObjectPath;
 
-	OutState.InitializeFromExisting(DataToUse, State.CachedDependsNodes, State.CachedPackageData, Options);
+	OutState.InitializeFromExisting(DataToUse, State.CachedDependsNodes, State.CachedPackageData, Options, bRefreshExisting);
 }
 
 void FAssetRegistry::ScanPathsAndFilesSynchronous(const TArray<FString>& InPaths, const TArray<FString>& InSpecificFiles, bool bForceRescan, EAssetDataCacheMode AssetDataCacheMode)
@@ -1376,12 +1467,36 @@ void FAssetRegistry::ScanPathsAndFilesSynchronous(const TArray<FString>& InPaths
 	// Only scan paths that were not previously synchronously scanned, unless we were asked to force rescan.
 	TArray<FString> PathsToScan;
 	TArray<FString> FilesToScan;
+	bool bPathsRemoved = false;
+
 	for (const FString& Path : InPaths)
 	{
-		if (bForceRescan || !SynchronouslyScannedPathsAndFiles.Contains(Path))
+		bool bAlreadyScanned = false;
+		FString PathWithSlash = Path;
+		if (!PathWithSlash.EndsWith(TEXT("/"), ESearchCase::CaseSensitive))
+		{
+			// Add / if it's missing so the substring check is safe
+			PathWithSlash += TEXT("/");
+		}
+
+		// Check that it starts with /
+		for (const FString& ScannedPath : SynchronouslyScannedPathsAndFiles)
+		{
+			if (PathWithSlash.StartsWith(ScannedPath))
+			{
+				bAlreadyScanned = true;
+				break;
+			}
+		}
+
+		if (bForceRescan || !bAlreadyScanned)
 		{
 			PathsToScan.Add(Path);
-			SynchronouslyScannedPathsAndFiles.Add(Path);
+			SynchronouslyScannedPathsAndFiles.Add(PathWithSlash);
+		}
+		else
+		{
+			bPathsRemoved = true;
 		}
 	}
 
@@ -1392,6 +1507,16 @@ void FAssetRegistry::ScanPathsAndFilesSynchronous(const TArray<FString>& InPaths
 			FilesToScan.Add(SpecificFile);
 			SynchronouslyScannedPathsAndFiles.Add(SpecificFile);
 		}
+		else
+		{
+			bPathsRemoved = true;
+		}
+	}
+
+	// If we removed paths, we can't use the monolithic cache as this will replace it with invalid data
+	if (AssetDataCacheMode == EAssetDataCacheMode::UseMonolithicCache && bPathsRemoved)
+	{
+		AssetDataCacheMode = EAssetDataCacheMode::UseModularCache;
 	}
 
 	if ( PathsToScan.Num() > 0 || FilesToScan.Num() > 0 )
@@ -2031,7 +2156,7 @@ void FAssetRegistry::ProcessLoadedAssetsToUpdateCache(const double TickStartTime
 
 		FAssetData NewAssetData = FAssetData(LoadedAsset);
 
-		if (!NewAssetData.TagsAndValues.GetMap().OrderIndependentCompareEqual((*CachedData)->TagsAndValues.GetMap()))
+		if (NewAssetData.TagsAndValues.GetMap() != (*CachedData)->TagsAndValues.GetMap())
 		{
 			// We need to actually update disk cache
 			UpdateAssetData(*CachedData, NewAssetData);
@@ -2053,6 +2178,22 @@ void FAssetRegistry::ProcessLoadedAssetsToUpdateCache(const double TickStartTime
 	}
 }
 
+void FAssetRegistry::UpdateRedirectCollector()
+{
+	// Look for all redirectors in list
+	const TArray<const FAssetData*>& RedirectorAssets = State.GetAssetsByClassName(UObjectRedirector::StaticClass()->GetFName());
+
+	for (const FAssetData* AssetData : RedirectorAssets)
+	{
+		FName Destination = GetRedirectedObjectPath(AssetData->ObjectPath);
+
+		if (Destination != AssetData->ObjectPath)
+		{
+			GRedirectCollector.AddAssetPathRedirection(AssetData->ObjectPath.ToString(), Destination.ToString());
+		}
+	}
+}
+
 #endif // WITH_EDITOR
 
 
@@ -2071,8 +2212,8 @@ void FAssetRegistry::OnContentPathMounted( const FString& InAssetPath, const FSt
 
 	// Listen for directory changes in this content path
 #if WITH_EDITOR
-	// Commandlets and in-game don't listen for directory changes
-	if ( !IsRunningCommandlet() && GIsEditor)
+	// In-game doesn't listen for directory changes
+	if (GIsEditor)
 	{
 		FDirectoryWatcherModule& DirectoryWatcherModule = FModuleManager::LoadModuleChecked<FDirectoryWatcherModule>(TEXT("DirectoryWatcher"));
 		IDirectoryWatcher* DirectoryWatcher = DirectoryWatcherModule.Get();
@@ -2127,8 +2268,8 @@ void FAssetRegistry::OnContentPathDismounted(const FString& InAssetPath, const F
 
 	// Stop listening for directory changes in this content path
 #if WITH_EDITOR
-	// Commandlets and in-game don't listen for directory changes
-	if (!IsRunningCommandlet() && GIsEditor)
+	// In-game doesn't listen for directory changes
+	if (GIsEditor)
 	{
 		FDirectoryWatcherModule& DirectoryWatcherModule = FModuleManager::LoadModuleChecked<FDirectoryWatcherModule>(TEXT("DirectoryWatcher"));
 		IDirectoryWatcher* DirectoryWatcher = DirectoryWatcherModule.Get();
@@ -2354,11 +2495,11 @@ bool FAssetRegistry::SetPrimaryAssetIdForObjectPath(const FName ObjectPath, FPri
 
 	FAssetData* AssetData = *FoundAssetData;
 
-	TMap<FName, FString> TagsAndValues = AssetData->TagsAndValues.GetMap();
+	FAssetDataTagMap TagsAndValues = AssetData->TagsAndValues.GetMap();
 	TagsAndValues.Add(FPrimaryAssetId::PrimaryAssetTypeTag, PrimaryAssetId.PrimaryAssetType.ToString());
 	TagsAndValues.Add(FPrimaryAssetId::PrimaryAssetNameTag, PrimaryAssetId.PrimaryAssetName.ToString());
 
-	FAssetData NewAssetData = FAssetData(AssetData->PackageName, AssetData->PackagePath, AssetData->GroupNames, AssetData->AssetName, AssetData->AssetClass, TagsAndValues, AssetData->ChunkIDs, AssetData->PackageFlags);
+	FAssetData NewAssetData = FAssetData(AssetData->PackageName, AssetData->PackagePath, AssetData->AssetName, AssetData->AssetClass, TagsAndValues, AssetData->ChunkIDs, AssetData->PackageFlags);
 
 	UpdateAssetData(AssetData, NewAssetData);
 
diff --git a/Engine/Source/Runtime/AssetRegistry/Private/AssetRegistry.h b/Engine/Source/Runtime/AssetRegistry/Private/AssetRegistry.h
--- a/Engine/Source/Runtime/AssetRegistry/Private/AssetRegistry.h
+++ b/Engine/Source/Runtime/AssetRegistry/Private/AssetRegistry.h
@@ -38,6 +38,7 @@ public:
 	virtual bool GetReferencers(const FAssetIdentifier& AssetIdentifier, TArray<FAssetIdentifier>& OutReferencers, EAssetRegistryDependencyType::Type InReferenceType = EAssetRegistryDependencyType::All) const override;
 	virtual bool GetReferencers(FName PackageName, TArray<FName>& OutReferencers, EAssetRegistryDependencyType::Type InReferenceType = EAssetRegistryDependencyType::Packages) const override;
 	virtual const FAssetPackageData* GetAssetPackageData(FName PackageName) const override;
+	virtual FName GetRedirectedObjectPath(const FName ObjectPath) const override;
 	virtual bool GetAncestorClassNames(FName ClassName, TArray<FName>& OutAncestorClassNames) const override;
 	virtual void GetDerivedClassNames(const TArray<FName>& ClassNames, const TSet<FName>& ExcludedClassNames, TSet<FName>& OutDerivedClassNames) const override;
 	virtual void GetAllCachedPaths(TArray<FString>& OutPathList) const override;
@@ -55,8 +56,9 @@ public:
 	virtual void ScanFilesSynchronous(const TArray<FString>& InFilePaths, bool bForceRescan = false) override;
 	virtual void PrioritizeSearchPath(const FString& PathToPrioritize) override;
 	virtual void Serialize(FArchive& Ar) override;
+	virtual uint32 GetAllocatedSize(bool bLogDetailed = false) const override;
 	virtual void LoadPackageRegistryData(FArchive& Ar, TArray<FAssetData*>& Data) const override;
-	virtual void InitializeTemporaryAssetRegistryState(FAssetRegistryState& OutState, const FAssetRegistrySerializationOptions& Options, const TMap<FName, FAssetData*>& OverrideData = TMap<FName, FAssetData*>()) const override;
+	virtual void InitializeTemporaryAssetRegistryState(FAssetRegistryState& OutState, const FAssetRegistrySerializationOptions& Options, bool bRefreshExisting = false, const TMap<FName, FAssetData*>& OverrideData = TMap<FName, FAssetData*>()) const override;
 	virtual void InitializeSerializationOptions(FAssetRegistrySerializationOptions& Options, const FString& PlatformIniName = FString()) const override;
 
 	virtual void SaveRegistryData(FArchive& Ar, TMap<FName, FAssetData*>& Data, TArray<FName>* InMaps = nullptr) override;
@@ -174,6 +176,9 @@ private:
 
 	/** Process Loaded Assets to update cache */
 	void ProcessLoadedAssetsToUpdateCache(const double TickStartTime);
+
+	/** Update Redirect collector with redirects loaded from asset registry */
+	void UpdateRedirectCollector();
 #endif // WITH_EDITOR
 
 	/**
@@ -194,6 +199,9 @@ private:
 	 */
 	void OnContentPathDismounted( const FString& AssetPath, const FString& FileSystemPath );
 
+	/** Called to refresh the native classes list, called at end of engine initialization */
+	void RefreshNativeClasses();
+
 	/** Returns the names of all subclasses of the class whose name is ClassName */
 	void GetSubClasses(const TArray<FName>& InClassNames, const TSet<FName>& ExcludedClassNames, TSet<FName>& SubClassNames) const;
 	void GetSubClasses_Recursive(FName InClassName, TSet<FName>& SubClassNames, const TMap<FName, TSet<FName>>& ReverseInheritanceMap, const TSet<FName>& ExcludedClassNames) const;
