Copying //UE4/Dev-Framework to //UE4/Dev-Main (Source: //UE4/Dev-Framework @ 3431384)

diff --git a/Engine/Source/Runtime/AssetRegistry/Private/AssetRegistryState.cpp b/Engine/Source/Runtime/AssetRegistry/Private/AssetRegistryState.cpp
--- a/Engine/Source/Runtime/AssetRegistry/Private/AssetRegistryState.cpp
+++ b/Engine/Source/Runtime/AssetRegistry/Private/AssetRegistryState.cpp
@@ -17,48 +17,6 @@
 #include "NameTableArchive.h"
 #include "GenericPlatform/GenericPlatformChunkInstall.h"
 
-/**
- * Enum for tracking versions of the runtime asset registry that gets used for cooked asset registries
- * This enum is NOT used for the editor caching, that uses AssetDataGathererConstants::CacheSerializationVersion
- */
-enum class ERuntimeRegistryVersion
-{
-	PreVersioning,		   // From before file versioning was implemented
-	HardSoftDependencies,  // The first version of the runtime asset registry to include file versioning.
-	AddAssetRegistryState, // Added FAssetRegistryState and support for piecemeal serialization
-
-	LatestPlusOne,
-	Latest = (LatestPlusOne - 1),
-};
-
-/** Guid for the runtime asset registry file format, used for identification purposes */
-static const FGuid GRuntimeRegistryGuid(0x717F9EE7, 0xE9B0493A, 0x88B39132, 0x1B388107);
-
-/** Helper function for reading the header of the runtime asset registry. It deals with the header not
-	existing for old formats of the file */
-inline ERuntimeRegistryVersion ReadRuntimeRegistryVersion(FArchive& Ar)
-{
-	int64 InitialLocation = Ar.Tell();
-	FGuid Guid;
-	Ar << Guid;
-
-	ERuntimeRegistryVersion Version = ERuntimeRegistryVersion::PreVersioning;
-
-	if (Guid == GRuntimeRegistryGuid)
-	{
-		int32 VersionInt;
-		Ar << VersionInt;
-		Version = (ERuntimeRegistryVersion)VersionInt;
-	}
-	else
-	{
-		// This is an old format file, so skip back to where we started
-		Ar.Seek(InitialLocation);
-	}
-
-	return Version;
-}
-
 FAssetRegistryState::FAssetRegistryState()
 {
 	NumAssets = 0;
@@ -160,12 +118,26 @@ void FAssetRegistryState::Reset()
 	CachedPackageData.Empty();
 }
 
-void FAssetRegistryState::InitializeFromExisting(const TMap<FName, FAssetData*>& AssetDataMap, const TMap<FAssetIdentifier, FDependsNode*>& DependsNodeMap, const TMap<FName, FAssetPackageData*>& AssetPackageDataMap, const FAssetRegistrySerializationOptions& Options)
+void FAssetRegistryState::InitializeFromExisting(const TMap<FName, FAssetData*>& AssetDataMap, const TMap<FAssetIdentifier, FDependsNode*>& DependsNodeMap, const TMap<FName, FAssetPackageData*>& AssetPackageDataMap, const FAssetRegistrySerializationOptions& Options, bool bRefreshExisting)
 {
-	Reset();
+	if (!bRefreshExisting)
+	{
+		Reset();
+	}
 
 	for (const TPair<FName, FAssetData*>& Pair : AssetDataMap)
 	{
+		FAssetData* ExistingData = nullptr;
+
+		if (bRefreshExisting)
+		{
+			ExistingData = CachedAssetsByObjectPath.FindRef(Pair.Key);
+			if (!ExistingData)
+			{
+				continue;
+			}
+		}
+
 		if (Pair.Value)
 		{
 			// Filter asset registry tags now
@@ -176,19 +148,18 @@ void FAssetRegistryState::InitializeFromExisting(const TMap<FName, FAssetData*>&
 			const TSet<FName>* ClassSpecificFilterlist = Options.CookFilterlistTagsByClass.Find(AssetData.AssetClass);
 
 			// Exclude blacklisted tags or include only whitelisted tags, based on how we were configured in ini
-			TMap<FName, FString> LocalTagsAndValues;
-			for (auto TagIt = AssetData.TagsAndValues.GetMap().CreateConstIterator(); TagIt; ++TagIt)
+			FAssetDataTagMap LocalTagsAndValues;
+			for (const TPair<FName, FString>& TagPair : AssetData.TagsAndValues)
 			{
-				FName TagName = TagIt.Key();
-				const bool bInAllClasseslist = AllClassesFilterlist && (AllClassesFilterlist->Contains(TagName) || AllClassesFilterlist->Contains(WildcardName));
-				const bool bInClassSpecificlist = ClassSpecificFilterlist && (ClassSpecificFilterlist->Contains(TagName) || ClassSpecificFilterlist->Contains(WildcardName));
+				const bool bInAllClasseslist = AllClassesFilterlist && (AllClassesFilterlist->Contains(TagPair.Key) || AllClassesFilterlist->Contains(WildcardName));
+				const bool bInClassSpecificlist = ClassSpecificFilterlist && (ClassSpecificFilterlist->Contains(TagPair.Key) || ClassSpecificFilterlist->Contains(WildcardName));
 				if (Options.bUseAssetRegistryTagsWhitelistInsteadOfBlacklist)
 				{
 					// It's a whitelist, only include it if it is in the all classes list or in the class specific list
 					if (bInAllClasseslist || bInClassSpecificlist)
 					{
 						// It is in the whitelist. Keep it.
-						LocalTagsAndValues.Add(TagIt.Key(), TagIt.Value());
+						LocalTagsAndValues.Add(TagPair.Key, TagPair.Value);
 					}
 				}
 				else
@@ -197,47 +168,61 @@ void FAssetRegistryState::InitializeFromExisting(const TMap<FName, FAssetData*>&
 					if (!bInAllClasseslist && !bInClassSpecificlist)
 					{
 						// It isn't in the blacklist. Keep it.
-						LocalTagsAndValues.Add(TagIt.Key(), TagIt.Value());
+						LocalTagsAndValues.Add(TagPair.Key, TagPair.Value);
 					}
 				}
 			}
 
-			FAssetData* NewData = new FAssetData(AssetData.PackageName, AssetData.PackagePath, AssetData.GroupNames, AssetData.AssetName,
-				AssetData.AssetClass, LocalTagsAndValues, AssetData.ChunkIDs, AssetData.PackageFlags);
-
-			AddAssetData(NewData);
+			if (bRefreshExisting)
+			{
+				// Only modify tags
+				if (LocalTagsAndValues != ExistingData->TagsAndValues.GetMap())
+				{
+					FAssetData TempData = *ExistingData;
+					TempData.TagsAndValues = FAssetDataTagMapSharedView(MoveTemp(LocalTagsAndValues));
+					UpdateAssetData(ExistingData, TempData);
+				}
+			}
+			else
+			{
+				FAssetData* NewData = new FAssetData(AssetData.PackageName, AssetData.PackagePath, AssetData.AssetName,
+					AssetData.AssetClass, LocalTagsAndValues, AssetData.ChunkIDs, AssetData.PackageFlags);
+				AddAssetData(NewData);
+			}
 		}
 	}
 
-	for (const TPair<FAssetIdentifier, FDependsNode*>& Pair : DependsNodeMap)
+	if (!bRefreshExisting)
 	{
-		FDependsNode* OldNode = Pair.Value;
-		FDependsNode* NewNode = CreateOrFindDependsNode(Pair.Key);
+		for (const TPair<FName, FAssetPackageData*>& Pair : AssetPackageDataMap)
+		{
+			// Only add if also in asset data map
+			if (Pair.Value && CachedAssetsByPackageName.Find(Pair.Key))
+			{
+				FAssetPackageData* NewData = CreateOrGetAssetPackageData(Pair.Key);
+				*NewData = *Pair.Value;
+			}
+		}
+
+		for (const TPair<FAssetIdentifier, FDependsNode*>& Pair : DependsNodeMap)
+		{
+			FDependsNode* OldNode = Pair.Value;
+			FDependsNode* NewNode = CreateOrFindDependsNode(Pair.Key);
 
-		Pair.Value->IterateOverDependencies(
-			[this, &DependsNodeMap, OldNode, NewNode](FDependsNode* InDependency, EAssetRegistryDependencyType::Type InDependencyType) {
+			Pair.Value->IterateOverDependencies([this, &DependsNodeMap, OldNode, NewNode](FDependsNode* InDependency, EAssetRegistryDependencyType::Type InDependencyType) {
 				if (DependsNodeMap.Find(InDependency->GetIdentifier()))
 				{
 					// Only add if this node is in the incoming map
 					FDependsNode* NewDependency = CreateOrFindDependsNode(InDependency->GetIdentifier());
-					NewNode->AddDependency(NewDependency, InDependencyType);
+					NewNode->AddDependency(NewDependency, InDependencyType, true);
 					NewDependency->AddReferencer(NewNode);
 				}
 			});
-	}
-
-	for (const TPair<FName, FAssetPackageData*>& Pair : AssetPackageDataMap)
-	{
-		// Only add if also in asset data map
-		if (Pair.Value && CachedAssetsByPackageName.Find(Pair.Key))
-		{
-			FAssetPackageData* NewData = CreateOrGetAssetPackageData(Pair.Key);
-			*NewData = *Pair.Value;
 		}
 	}
 }
 
-void FAssetRegistryState::PruneAssetData(const TSet<FName>& RequiredPackages, const TSet<FName>& RemovePackages)
+void FAssetRegistryState::PruneAssetData(const TSet<FName>& RequiredPackages, const TSet<FName>& RemovePackages, bool bFilterAssetDataWithNoTags)
 {
 	// Generate list up front as the maps will get cleaned up
 	TArray<FAssetData*> AllAssetData;
@@ -253,6 +238,10 @@ void FAssetRegistryState::PruneAssetData(const TSet<FName>& RequiredPackages, co
 		{
 			RemoveAssetData(AssetData);
 		}
+		else if (bFilterAssetDataWithNoTags && AssetData->TagsAndValues.Num() == 0)
+		{
+			RemoveAssetData(AssetData);
+		}
 	}
 
 	// Remove any orphaned depends nodes. This will leave cycles in but those might represent useful data
@@ -547,11 +536,8 @@ bool FAssetRegistryState::Serialize(FArchive& OriginalAr, FAssetRegistrySerializ
 	{
 		check(CachedAssetsByObjectPath.Num() == NumAssets);
 
-		// Write asset registry header
-		FGuid LocalGuid = GRuntimeRegistryGuid;
-		LocalGuid.Serialize(OriginalAr);
-		int32 VersionInt = (int32)ERuntimeRegistryVersion::Latest;
-		OriginalAr << VersionInt;
+		FAssetRegistryVersion::Type Version = FAssetRegistryVersion::LatestVersion;
+		FAssetRegistryVersion::SerializeVersion(OriginalAr, Version);
 
 		// Set up name table archive
 		FNameTableArchiveWriter Ar(OriginalAr);
@@ -568,7 +554,7 @@ bool FAssetRegistryState::Serialize(FArchive& OriginalAr, FAssetRegistrySerializ
 		for (TPair<FName, FAssetData*>& Pair : CachedAssetsByObjectPath)
 		{
 			FAssetData& AssetData(*Pair.Value);
-			Ar << AssetData;
+			AssetData.SerializeForCache(Ar);
 		}
 
 		if (Options.bSerializeDependencies)
@@ -662,7 +648,7 @@ bool FAssetRegistryState::Serialize(FArchive& OriginalAr, FAssetRegistrySerializ
 			for (TPair<FName, FAssetPackageData*>& Pair : CachedPackageData)
 			{
 				Ar << Pair.Key;
-				Ar << *Pair.Value;
+				Pair.Value->SerializeForCache(Ar);
 			}
 		}
 		else
@@ -673,11 +659,13 @@ bool FAssetRegistryState::Serialize(FArchive& OriginalAr, FAssetRegistrySerializ
 	// load in by building the TMap
 	else
 	{
-		ERuntimeRegistryVersion Version = ReadRuntimeRegistryVersion(OriginalAr);
+		FAssetRegistryVersion::Type Version = FAssetRegistryVersion::LatestVersion;
+		FAssetRegistryVersion::SerializeVersion(OriginalAr, Version);
 
-		if (Version < ERuntimeRegistryVersion::AddAssetRegistryState)
+		if (Version < FAssetRegistryVersion::ChangedAssetData)
 		{
-			return ReadDeprecated(OriginalAr, Version);
+			// Cannot read states before this version
+			return false;
 		}
 
 		// Set up name table archive
@@ -697,7 +685,7 @@ bool FAssetRegistryState::Serialize(FArchive& OriginalAr, FAssetRegistrySerializ
 			FAssetData* NewAssetData = &PreallocatedAssetDataBuffer[AssetIndex];
 
 			// load it
-			Ar << *NewAssetData;
+			NewAssetData->SerializeForCache(Ar);
 
 			AddAssetData(NewAssetData);
 		}
@@ -754,11 +742,11 @@ bool FAssetRegistryState::Serialize(FArchive& OriginalAr, FAssetRegistrySerializ
 					{
 						if (InDependencyType == EAssetRegistryDependencyType::None)
 						{
-							NewDependsNodeData->AddReferencer(&PreallocatedDependsNodeDataBuffer[Index]);
+							NewDependsNodeData->AddReferencer(&PreallocatedDependsNodeDataBuffer[Index], true);
 						}
 						else
 						{
-							NewDependsNodeData->AddDependency(&PreallocatedDependsNodeDataBuffer[Index], InDependencyType);
+							NewDependsNodeData->AddDependency(&PreallocatedDependsNodeDataBuffer[Index], InDependencyType, true);
 						}
 					}
 				}
@@ -790,13 +778,13 @@ bool FAssetRegistryState::Serialize(FArchive& OriginalAr, FAssetRegistrySerializ
 			if (Options.bSerializePackageData)
 			{
 				FAssetPackageData& NewPackageData = PreallocatedPackageDataBuffer[PackageDataIndex];
-				Ar << NewPackageData;
+				NewPackageData.SerializeForCache(Ar);
 				CachedPackageData.Add(PackageName, &NewPackageData);
 			}
 			else
 			{
 				FAssetPackageData FakeData;
-				Ar << FakeData;
+				FakeData.SerializeForCache(Ar);
 			}
 		}
 	}
@@ -804,100 +792,84 @@ bool FAssetRegistryState::Serialize(FArchive& OriginalAr, FAssetRegistrySerializ
 	return !OriginalAr.IsError();
 }
 
-bool FAssetRegistryState::ReadDeprecated(FArchive& Ar, ERuntimeRegistryVersion Version)
+uint32 FAssetRegistryState::GetAllocatedSize(bool bLogDetailed) const
 {
-	// serialize number of objects
-	int32 LocalNumAssets = 0;
-	Ar << LocalNumAssets;
-
-	// allocate one single block for all asset data structs (to reduce tens of thousands of heap allocations)
-	FAssetData* PreallocatedAssetDataBuffer = new FAssetData[LocalNumAssets];
-	PreallocatedAssetDataBuffers.Add(PreallocatedAssetDataBuffer);
-
-	double LastPumpTime = FPlatformTime::Seconds();
-	for (int32 AssetIndex = 0; AssetIndex < LocalNumAssets; AssetIndex++)
-	{
-		// make a new asset data object
-		FAssetData* NewAssetData = &PreallocatedAssetDataBuffer[AssetIndex];
-
-		// load it
-		Ar << *NewAssetData;
+	uint32 TotalBytes = 0;
 
-		AddAssetData(NewAssetData);
-	}
+	uint32 MapMemory = CachedAssetsByObjectPath.GetAllocatedSize();
+	MapMemory += CachedAssetsByPackageName.GetAllocatedSize();
+	MapMemory += CachedAssetsByPath.GetAllocatedSize();
+	MapMemory += CachedAssetsByClass.GetAllocatedSize();
+	MapMemory += CachedAssetsByTag.GetAllocatedSize();
+	MapMemory += CachedDependsNodes.GetAllocatedSize();
+	MapMemory += CachedPackageData.GetAllocatedSize();
 
-	int32 LocalNumDependsNodes = LocalNumAssets;
-	if (Version == ERuntimeRegistryVersion::PreVersioning)
+	if (bLogDetailed)
 	{
-		Ar << LocalNumDependsNodes;
+		UE_LOG(LogAssetRegistry, Log, TEXT("Index Size: %dk"), MapMemory / 1024);
 	}
 
-	FDependsNode* PreallocatedDependsNodeDataBuffer = new FDependsNode[LocalNumDependsNodes];
-	PreallocatedDependsNodeDataBuffers.Add(PreallocatedDependsNodeDataBuffer);
-	CachedDependsNodes.Reserve(LocalNumDependsNodes);
+	uint32 AssetDataSize = 0, TagOverHead = 0, TotalTagSize = 0;
+	TMap<FName, uint32> TagSizes;
 
-	for (int32 DependsNodeIndex = 0; DependsNodeIndex < LocalNumDependsNodes; DependsNodeIndex++)
+	for (const TPair<FName, FAssetData*>& AssetDataPair : CachedAssetsByObjectPath)
 	{
-		FDependsNode* NewDependsNodeData = &PreallocatedDependsNodeDataBuffer[DependsNodeIndex];
-
-		if (DependsNodeIndex < LocalNumAssets)
-		{
-			// One of the packages
-			int32 AssetIndex = DependsNodeIndex;
+		const FAssetData& AssetData = *AssetDataPair.Value;
+		
+		AssetDataSize += sizeof(AssetData);
+		AssetDataSize += AssetData.ChunkIDs.GetAllocatedSize();
 
-			if (Version == ERuntimeRegistryVersion::PreVersioning)
-			{
-				Ar << AssetIndex;
-			}
+		TagOverHead += AssetData.TagsAndValues.GetAllocatedSize();
 
-			NewDependsNodeData->SetPackageName(PreallocatedAssetDataBuffer[AssetIndex].PackageName);
+		for (const TPair<FName, FString>& TagPair : AssetData.TagsAndValues)
+		{
+			uint32 StringSize = TagPair.Value.GetAllocatedSize();
+			
+			TotalTagSize += StringSize;
+			TagSizes.FindOrAdd(TagPair.Key) += StringSize;
 		}
-		
-		CachedDependsNodes.Add(NewDependsNodeData->GetIdentifier(), NewDependsNodeData);
 	}
 
-	for (int32 DependsNodeIndex = 0; DependsNodeIndex < LocalNumDependsNodes; DependsNodeIndex++)
+	if (bLogDetailed)
 	{
-		int32 LocalNumHardDependencies = 0;
-		int32 LocalNumSoftDependencies = 0;
-		int32 LocalNumReferencers = 0;
-		Ar << LocalNumHardDependencies;
-		if (Version >= ERuntimeRegistryVersion::HardSoftDependencies)
+		UE_LOG(LogAssetRegistry, Log, TEXT("AssetData Count: %d"), CachedAssetsByObjectPath.Num());
+		UE_LOG(LogAssetRegistry, Log, TEXT("AssetData Static Size: %dk"), AssetDataSize / 1024);
+		UE_LOG(LogAssetRegistry, Log, TEXT("AssetData Tag Overhead: %dk"), TagOverHead / 1024);
+		
+		for (const TPair<FName, uint32>& SizePair : TagSizes)
 		{
-			Ar << LocalNumSoftDependencies;
+			UE_LOG(LogAssetRegistry, Log, TEXT("Tag %s Size: %dk"), *SizePair.Key.ToString(), SizePair.Value / 1024);
 		}
+	}
 
-		Ar << LocalNumReferencers;
-
-		FDependsNode* NewDependsNodeData = &PreallocatedDependsNodeDataBuffer[DependsNodeIndex];
-		NewDependsNodeData->Reserve(LocalNumHardDependencies, LocalNumSoftDependencies, 0, 0, LocalNumReferencers);
+	uint32 DependNodesSize = 0, DependenciesSize = 0;
 
-		for (int32 DependencyIndex = 0; DependencyIndex < LocalNumHardDependencies; ++DependencyIndex)
-		{
-			int32 Index = 0;
-			Ar << Index;
+	for (const TPair<FAssetIdentifier, FDependsNode*>& DependsNodePair : CachedDependsNodes)
+	{
+		const FDependsNode& DependsNode = *DependsNodePair.Value;
+		DependNodesSize += sizeof(DependsNode);
 
-			NewDependsNodeData->AddDependency(&PreallocatedDependsNodeDataBuffer[Index], EAssetRegistryDependencyType::Hard);
-		}
+		DependenciesSize += DependsNode.GetAllocatedSize();
+	}
 
-		for (int32 DependencyIndex = 0; DependencyIndex < LocalNumSoftDependencies; ++DependencyIndex)
-		{
-			int32 Index = 0;
-			Ar << Index;
+	if (bLogDetailed)
+	{
+		UE_LOG(LogAssetRegistry, Log, TEXT("Dependency Node Count: %d"), CachedDependsNodes.Num());
+		UE_LOG(LogAssetRegistry, Log, TEXT("Dependency Node Static Size: %dk"), DependNodesSize / 1024);
+		UE_LOG(LogAssetRegistry, Log, TEXT("Dependency Arrays Size: %dk"), DependenciesSize / 1024);
+	}
 
-			NewDependsNodeData->AddDependency(&PreallocatedDependsNodeDataBuffer[Index], EAssetRegistryDependencyType::Soft);
-		}
+	uint32 PackageDataSize = CachedPackageData.Num() * sizeof(FAssetPackageData);
 
-		for (int32 ReferencerIndex = 0; ReferencerIndex < LocalNumReferencers; ++ReferencerIndex)
-		{
-			int32 Index = 0;
-			Ar << Index;
+	TotalBytes = MapMemory + AssetDataSize + TagOverHead + TotalTagSize + DependNodesSize + DependenciesSize + PackageDataSize;
 
-			NewDependsNodeData->AddReferencer(&PreallocatedDependsNodeDataBuffer[Index]);
-		}
+	if (bLogDetailed)
+	{
+		UE_LOG(LogAssetRegistry, Log, TEXT("PackageData Count: %d"), CachedPackageData.Num());
+		UE_LOG(LogAssetRegistry, Log, TEXT("PackageData Static Size: %dk"), PackageDataSize / 1024);
+		UE_LOG(LogAssetRegistry, Log, TEXT("Total State Size: %dk"), TotalBytes / 1024);
 	}
-
-	return !Ar.IsError();
+	return TotalBytes;
 }
 
 FDependsNode* FAssetRegistryState::ResolveRedirector(FDependsNode* InDependency,
@@ -1093,14 +1065,14 @@ bool FAssetRegistryState::RemoveAssetData(FAssetData* AssetData)
 		TArray<FAssetData*>* OldClassAssets = CachedAssetsByClass.Find(AssetData->AssetClass);
 
 		CachedAssetsByObjectPath.Remove(AssetData->ObjectPath);
-		OldPackageAssets->Remove(AssetData);
-		OldPathAssets->Remove(AssetData);
-		OldClassAssets->Remove(AssetData);
+		OldPackageAssets->RemoveSingleSwap(AssetData);
+		OldPathAssets->RemoveSingleSwap(AssetData);
+		OldClassAssets->RemoveSingleSwap(AssetData);
 
 		for (auto TagIt = AssetData->TagsAndValues.CreateConstIterator(); TagIt; ++TagIt)
 		{
 			TArray<FAssetData*>* OldTagAssets = CachedAssetsByTag.Find(TagIt.Key());
-			OldTagAssets->Remove(AssetData);
+			OldTagAssets->RemoveSingleSwap(AssetData);
 		}
 
 		// We need to update the cached dependencies references cache so that they know we no
diff --git a/Engine/Source/Runtime/AssetRegistry/Private/DependsNode.cpp b/Engine/Source/Runtime/AssetRegistry/Private/DependsNode.cpp
--- a/Engine/Source/Runtime/AssetRegistry/Private/DependsNode.cpp
+++ b/Engine/Source/Runtime/AssetRegistry/Private/DependsNode.cpp
@@ -72,11 +72,18 @@ void FDependsNode::GetReferencers(TArray<FDependsNode*>& OutReferencers, EAssetR
 	}
 }
 
-void FDependsNode::AddDependency(FDependsNode* InDependency, EAssetRegistryDependencyType::Type InDependencyType)
+void FDependsNode::AddDependency(FDependsNode* InDependency, EAssetRegistryDependencyType::Type InDependencyType, bool bGuaranteedUnique)
 {
-	IterateOverDependencyArrays([InDependency](TArray<FDependsNode*>& InArray, EAssetRegistryDependencyType::Type)
+	IterateOverDependencyArrays([InDependency,&bGuaranteedUnique](TArray<FDependsNode*>& InArray, EAssetRegistryDependencyType::Type)
 	{
-		InArray.AddUnique(InDependency);
+		if (bGuaranteedUnique)
+		{
+			InArray.Add(InDependency);
+		}
+		else
+		{
+			InArray.AddUnique(InDependency);
+		}
 	}, InDependencyType);
 }
 
diff --git a/Engine/Source/Runtime/AssetRegistry/Private/DependsNode.h b/Engine/Source/Runtime/AssetRegistry/Private/DependsNode.h
--- a/Engine/Source/Runtime/AssetRegistry/Private/DependsNode.h
+++ b/Engine/Source/Runtime/AssetRegistry/Private/DependsNode.h
@@ -34,9 +34,9 @@ public:
 	/** Sets the entire identifier */
 	void SetIdentifier(const FAssetIdentifier& InIdentifier) { Identifier = InIdentifier;  }
 	/** Add a dependency to this node */
-	void AddDependency(FDependsNode* InDependency, EAssetRegistryDependencyType::Type InDependencyType = EAssetRegistryDependencyType::Hard);
+	void AddDependency(FDependsNode* InDependency, EAssetRegistryDependencyType::Type InDependencyType = EAssetRegistryDependencyType::Hard, bool bGuaranteedUnique = false);
 	/** Add a referencer to this node */
-	void AddReferencer(FDependsNode* InReferencer) { Referencers.AddUnique(InReferencer); }
+	void AddReferencer(FDependsNode* InReferencer, bool bGuaranteedUnique = false) { bGuaranteedUnique ? Referencers.Add(InReferencer) : Referencers.AddUnique(InReferencer); }
 	/** Remove a dependency from this node */
 	void RemoveDependency(FDependsNode* InDependency);
 	/** Remove a referencer from this node */
@@ -47,6 +47,11 @@ public:
 	void RemoveManageReferencesToNode();
 	/** Returns number of connections this node has, both references and dependencies */
 	int32 GetConnectionCount() const;
+	/** Returns amount of memory used by the arrays */
+	uint32 GetAllocatedSize(void) const
+	{
+		return HardDependencies.GetAllocatedSize() + SoftDependencies.GetAllocatedSize() + NameDependencies.GetAllocatedSize() + ManageDependencies.GetAllocatedSize() + Referencers.GetAllocatedSize();
+	}
 
 	/** Iterate over all the dependencies of this node, filtered by the supplied type parameter, and call the supplied lambda parameter on the record */
 	template <class T>
diff --git a/Engine/Source/Runtime/AssetRegistry/Private/DiskCachedAssetData.h b/Engine/Source/Runtime/AssetRegistry/Private/DiskCachedAssetData.h
--- a/Engine/Source/Runtime/AssetRegistry/Private/DiskCachedAssetData.h
+++ b/Engine/Source/Runtime/AssetRegistry/Private/DiskCachedAssetData.h
@@ -20,23 +20,27 @@ public:
 		: Timestamp(InTimestamp)
 	{}
 
-	/** Operator for serialization */
-	friend FArchive& operator<<(FArchive& Ar, FDiskCachedAssetData& DiskCachedAssetData)
+	/**
+	 * Serialize as part of the registry cache. This is not meant to be serialized as part of a package so  it does not handle versions normally
+	 * To version this data change FAssetRegistryVersion or CacheSerializationVersion
+	 */
+	void SerializeForCache(FArchive& Ar)
 	{
-		Ar << DiskCachedAssetData.Timestamp;
-		if (Ar.IsError())
+		Ar << Timestamp;
+	
+		int32 AssetDataCount = AssetDataList.Num();
+		Ar << AssetDataCount;
+
+		if (Ar.IsLoading())
 		{
-			return Ar;
+			AssetDataList.SetNum(AssetDataCount);
 		}
 
-		Ar << DiskCachedAssetData.AssetDataList;
-		if (Ar.IsError())
+		for (int32 i = 0; i < AssetDataCount; i++)
 		{
-			return Ar;
+			AssetDataList[i].SerializeForCache(Ar);
 		}
 
-		Ar << DiskCachedAssetData.DependencyData;
-		
-		return Ar;
+		DependencyData.SerializeForCache(Ar);
 	}
 };
diff --git a/Engine/Source/Runtime/AssetRegistry/Private/NameTableArchive.cpp b/Engine/Source/Runtime/AssetRegistry/Private/NameTableArchive.cpp
--- a/Engine/Source/Runtime/AssetRegistry/Private/NameTableArchive.cpp
+++ b/Engine/Source/Runtime/AssetRegistry/Private/NameTableArchive.cpp
@@ -156,6 +156,31 @@ int64 FNameTableArchiveReader::TotalSize()
 	return 0;
 }
 
+const FCustomVersionContainer& FNameTableArchiveReader::GetCustomVersions() const
+{
+	if (ProxyAr)
+	{
+		return ProxyAr->GetCustomVersions();
+	}
+	return FArchive::GetCustomVersions();
+}
+
+void FNameTableArchiveReader::SetCustomVersions(const FCustomVersionContainer& NewVersions)
+{
+	if (ProxyAr)
+	{
+		ProxyAr->SetCustomVersions(NewVersions);
+	}
+}
+
+void FNameTableArchiveReader::ResetCustomVersions()
+{
+	if (ProxyAr)
+	{
+		ProxyAr->ResetCustomVersions();
+	}
+}
+
 FArchive& FNameTableArchiveReader::operator<<( FName& Name )
 {
 	int32 NameIndex;
@@ -323,6 +348,31 @@ int64 FNameTableArchiveWriter::TotalSize()
 	return 0.f;
 }
 
+const FCustomVersionContainer& FNameTableArchiveWriter::GetCustomVersions() const
+{
+	if (ProxyAr)
+	{
+		return ProxyAr->GetCustomVersions();
+	}
+	return FArchive::GetCustomVersions();
+}
+
+void FNameTableArchiveWriter::SetCustomVersions(const FCustomVersionContainer& NewVersions)
+{
+	if (ProxyAr)
+	{
+		ProxyAr->SetCustomVersions(NewVersions);
+	}
+}
+
+void FNameTableArchiveWriter::ResetCustomVersions()
+{
+	if (ProxyAr)
+	{
+		ProxyAr->ResetCustomVersions();
+	}
+}
+
 FArchive& FNameTableArchiveWriter::operator<<( FName& Name )
 {
 	int32* NameIndexPtr = NameMap.Find(Name);
diff --git a/Engine/Source/Runtime/AssetRegistry/Private/NameTableArchive.h b/Engine/Source/Runtime/AssetRegistry/Private/NameTableArchive.h
--- a/Engine/Source/Runtime/AssetRegistry/Private/NameTableArchive.h
+++ b/Engine/Source/Runtime/AssetRegistry/Private/NameTableArchive.h
@@ -35,11 +35,14 @@ public:
 	~FNameTableArchiveReader();
 	
 	// Farchive implementation to redirect requests to the wrapped archive
-	virtual void Serialize(void* V, int64 Length);
-	virtual bool Precache(int64 PrecacheOffset, int64 PrecacheSize);
-	virtual void Seek(int64 InPos);
-	virtual int64 Tell();
-	virtual int64 TotalSize();
+	virtual void Serialize(void* V, int64 Length) override;
+	virtual bool Precache(int64 PrecacheOffset, int64 PrecacheSize) override;
+	virtual void Seek(int64 InPos) override;
+	virtual int64 Tell() override;
+	virtual int64 TotalSize() override;
+	virtual const FCustomVersionContainer& GetCustomVersions() const override;
+	virtual void SetCustomVersions(const FCustomVersionContainer& NewVersions) override;
+	virtual void ResetCustomVersions() override;
 	virtual FArchive& operator<<(FName& Name);
 
 private:
@@ -63,11 +66,14 @@ public:
 	~FNameTableArchiveWriter();
 
 	// Farchive implementation to redirect requests to the wrapped archive
-	virtual void Serialize(void* V, int64 Length);
-	virtual bool Precache(int64 PrecacheOffset, int64 PrecacheSize);
-	virtual void Seek(int64 InPos);
-	virtual int64 Tell();
-	virtual int64 TotalSize();
+	virtual void Serialize(void* V, int64 Length) override;
+	virtual bool Precache(int64 PrecacheOffset, int64 PrecacheSize) override;
+	virtual void Seek(int64 InPos) override;
+	virtual int64 Tell() override;
+	virtual int64 TotalSize() override;
+	virtual const FCustomVersionContainer& GetCustomVersions() const override;
+	virtual void SetCustomVersions(const FCustomVersionContainer& NewVersions) override;
+	virtual void ResetCustomVersions() override;
 	virtual FArchive& operator<<(FName& Name);
 
 private:
diff --git a/Engine/Source/Runtime/AssetRegistry/Private/PackageDependencyData.h b/Engine/Source/Runtime/AssetRegistry/Private/PackageDependencyData.h
--- a/Engine/Source/Runtime/AssetRegistry/Private/PackageDependencyData.h
+++ b/Engine/Source/Runtime/AssetRegistry/Private/PackageDependencyData.h
@@ -24,16 +24,17 @@ public:
 	 */
 	FName GetImportPackageName(int32 ImportIndex);
 
-	/** Operator for serialization */
-	friend FArchive& operator<<(FArchive& Ar, FPackageDependencyData& DependencyData)
+	/**
+	 * Serialize as part of the registry cache. This is not meant to be serialized as part of a package so  it does not handle versions normally
+	 * To version this data change FAssetRegistryVersion or CacheSerializationVersion
+	 */
+	void SerializeForCache(FArchive& Ar)
 	{
-		// serialize out the asset info, this is tied to CacheSerializationVersion
-		Ar << DependencyData.PackageName;
-		Ar << DependencyData.PackageData;
-		Ar << DependencyData.ImportMap;
-		Ar << DependencyData.StringAssetReferencesMap;
-		Ar << DependencyData.SearchableNamesMap;
+		Ar << PackageName;
+		Ar << ImportMap;
+		Ar << StringAssetReferencesMap;
+		Ar << SearchableNamesMap;
 		
-		return Ar;
+		PackageData.SerializeForCache(Ar);
 	}
 };
diff --git a/Engine/Source/Runtime/AssetRegistry/Private/PackageReader.cpp b/Engine/Source/Runtime/AssetRegistry/Private/PackageReader.cpp
--- a/Engine/Source/Runtime/AssetRegistry/Private/PackageReader.cpp
+++ b/Engine/Source/Runtime/AssetRegistry/Private/PackageReader.cpp
@@ -159,11 +159,11 @@ bool FPackageReader::ReadAssetRegistryData (TArray<FAssetData*>& AssetDataList)
 		if (bLegacyPackage || bNoMapAsset)
 		{
 			FString AssetName = FPackageName::GetLongPackageAssetName(PackageName);
-			AssetDataList.Add(new FAssetData(FName(*PackageName), FName(*PackagePath), FName(), MoveTemp(FName(*AssetName)), FName(TEXT("World")), TMap<FName, FString>(), PackageFileSummary.ChunkIDs, PackageFileSummary.PackageFlags));
+			AssetDataList.Add(new FAssetData(FName(*PackageName), FName(*PackagePath), MoveTemp(FName(*AssetName)), FName(TEXT("World")), FAssetDataTagMap(), PackageFileSummary.ChunkIDs, PackageFileSummary.PackageFlags));
 		}
 	}
 
-	// UAsset files only have one object, but legacy or map packages may have more.
+	// UAsset files usually only have one asset, maps and redirectors have multiple
 	for(int32 ObjectIdx = 0; ObjectIdx < ObjectCount; ++ObjectIdx)
 	{
 		FString ObjectPath;
@@ -173,7 +173,7 @@ bool FPackageReader::ReadAssetRegistryData (TArray<FAssetData*>& AssetDataList)
 		*this << ObjectClassName;
 		*this << TagCount;
 
-		TMap<FName, FString> TagsAndValues;
+		FAssetDataTagMap TagsAndValues;
 		TagsAndValues.Reserve(TagCount);
 
 		for(int32 TagIdx = 0; TagIdx < TagCount; ++TagIdx)
@@ -183,21 +183,26 @@ bool FPackageReader::ReadAssetRegistryData (TArray<FAssetData*>& AssetDataList)
 			*this << Key;
 			*this << Value;
 
-			TagsAndValues.Add(FName(*Key), Value);
+			if (!Key.IsEmpty() && !Value.IsEmpty())
+			{
+				TagsAndValues.Add(FName(*Key), Value);
+			}
 		}
 
-		FString GroupNames;
-		FString AssetName;
-
-		if ( ObjectPath.Contains(TEXT("."), ESearchCase::CaseSensitive))
+		if (ObjectPath.StartsWith(TEXT("/"), ESearchCase::CaseSensitive))
 		{
-			ObjectPath.Split(TEXT("."), &GroupNames, &AssetName, ESearchCase::CaseSensitive, ESearchDir::FromEnd);
+			// This should never happen, it means that package A has an export with an outer of package B
+			UE_LOG(LogAssetRegistry, Warning, TEXT("Package %s has invalid export %s, resave source package!"), *PackageName, *ObjectPath);
+			continue;
 		}
-		else
+
+		if (!ensureMsgf(!ObjectPath.Contains(TEXT("."), ESearchCase::CaseSensitive), TEXT("Cannot make FAssetData for sub object %s!"), *ObjectPath))
 		{
-			AssetName = ObjectPath;
+			continue;
 		}
 
+		FString AssetName = ObjectPath;
+
 		// Before world were RF_Public, other non-public assets were added to the asset data table in map packages.
 		// Here we simply skip over them
 		if ( bIsMapPackage && PackageFileSummary.GetFileVersionUE4() < VER_UE4_PUBLIC_WORLDS )
@@ -209,7 +214,7 @@ bool FPackageReader::ReadAssetRegistryData (TArray<FAssetData*>& AssetDataList)
 		}
 
 		// Create a new FAssetData for this asset and update it with the gathered data
-		AssetDataList.Add(new FAssetData(FName(*PackageName), FName(*PackagePath), MoveTemp(FName(*GroupNames)), MoveTemp(FName(*AssetName)), MoveTemp(FName(*ObjectClassName)), MoveTemp(TagsAndValues), PackageFileSummary.ChunkIDs, PackageFileSummary.PackageFlags));
+		AssetDataList.Add(new FAssetData(FName(*PackageName), FName(*PackagePath), MoveTemp(FName(*AssetName)), MoveTemp(FName(*ObjectClassName)), MoveTemp(TagsAndValues), PackageFileSummary.ChunkIDs, PackageFileSummary.PackageFlags));
 	}
 
 	return true;
@@ -254,17 +259,13 @@ bool FPackageReader::ReadAssetDataFromThumbnailCache(TArray<FAssetData*>& AssetD
 		FString GroupNames;
 		FString AssetName;
 
-		if ( ObjectPathWithoutPackageName.Contains(TEXT("."), ESearchCase::CaseSensitive) )
-		{
-			ObjectPathWithoutPackageName.Split(TEXT("."), &GroupNames, &AssetName, ESearchCase::CaseSensitive, ESearchDir::FromEnd);
-		}
-		else
+		if (!ensureMsgf(!ObjectPathWithoutPackageName.Contains(TEXT("."), ESearchCase::CaseSensitive), TEXT("Cannot make FAssetData for sub object %s!"), *ObjectPathWithoutPackageName))
 		{
-			AssetName = ObjectPathWithoutPackageName;
+			continue;
 		}
 
 		// Create a new FAssetData for this asset and update it with the gathered data
-		AssetDataList.Add(new FAssetData(FName(*PackageName), FName(*PackagePath), MoveTemp(FName(*GroupNames)), MoveTemp(FName(*AssetName)), MoveTemp(FName(*AssetClassName)), TMap<FName, FString>(), PackageFileSummary.ChunkIDs, PackageFileSummary.PackageFlags));
+		AssetDataList.Add(new FAssetData(FName(*PackageName), FName(*PackagePath), MoveTemp(FName(*ObjectPathWithoutPackageName)), MoveTemp(FName(*AssetClassName)), FAssetDataTagMap(), PackageFileSummary.ChunkIDs, PackageFileSummary.PackageFlags));
 	}
 
 	return true;
@@ -295,10 +296,6 @@ bool FPackageReader::ReadAssetRegistryDataIfCookedPackage(TArray<FAssetData*>& A
 			{
 				if (Export.bIsAsset)
 				{
-					FString GroupNames; // Not used for anything
-					TMap<FName, FString> Tags; // Not used for anything
-					TArray<int32> ChunkIDs; // Not used for anything
-
 					// We need to get the class name from the import/export maps
 					FName ObjectClassName;
 					if (Export.ClassIndex.IsNull())
@@ -316,7 +313,7 @@ bool FPackageReader::ReadAssetRegistryDataIfCookedPackage(TArray<FAssetData*>& A
 						ObjectClassName = ClassImport.ObjectName;
 					}
 
-					AssetDataList.Add(new FAssetData(FName(*PackageName), FName(*PackagePath), FName(*GroupNames), Export.ObjectName, ObjectClassName, Tags, ChunkIDs, GetPackageFlags()));
+					AssetDataList.Add(new FAssetData(FName(*PackageName), FName(*PackagePath), Export.ObjectName, ObjectClassName, FAssetDataTagMap(), TArray<int32>(), GetPackageFlags()));
 					bFoundAtLeastOneAsset = true;
 				}
 			}
diff --git a/Engine/Source/Runtime/AssetRegistry/Private/PathTree.h b/Engine/Source/Runtime/AssetRegistry/Private/PathTree.h
--- a/Engine/Source/Runtime/AssetRegistry/Private/PathTree.h
+++ b/Engine/Source/Runtime/AssetRegistry/Private/PathTree.h
@@ -19,6 +19,11 @@ public:
 	/** Recursively gathers all child paths from the specified base path relative to this node */
 	bool GetSubPaths(FName BasePath, TSet<FName>& OutPaths, bool bRecurse = true) const;
 
+	uint32 GetAllocatedSize(void) const
+	{
+		return ParentPathToChildPaths.GetAllocatedSize() + ChildPathToParentPath.GetAllocatedSize();
+	}
+
 private:
 	/** A one-to-many mapping between a parent path and its child paths. */
 	TMap<FName, TSet<FName>> ParentPathToChildPaths;
diff --git a/Engine/Source/Runtime/AssetRegistry/Public/AssetBundleData.h b/Engine/Source/Runtime/AssetRegistry/Public/AssetBundleData.h
--- a/Engine/Source/Runtime/AssetRegistry/Public/AssetBundleData.h
+++ b/Engine/Source/Runtime/AssetRegistry/Public/AssetBundleData.h
@@ -75,7 +75,24 @@ struct ASSETREGISTRY_API FAssetBundleData
 	/** Adds multiple assets at once */
 	void AddBundleAssets(FName BundleName, const TArray<FStringAssetReference>& AssetPaths);
 
+	/** A fast set of asset bundle assets, will destroy copied in path list */
+	void SetBundleAssets(FName BundleName, TArray<FStringAssetReference>&& AssetPaths);
+
 	/** Resets the data to defaults */
 	void Reset();
 
+	/** Override Import/Export to not write out empty structs */
+	bool ExportTextItem(FString& ValueStr, FAssetBundleData const& DefaultValue, UObject* Parent, int32 PortFlags, UObject* ExportRootScope) const;
+	bool ImportTextItem(const TCHAR*& Buffer, int32 PortFlags, UObject* Parent, FOutputDevice* ErrorText);
+
+};
+
+template<>
+struct TStructOpsTypeTraits<FAssetBundleData> : public TStructOpsTypeTraitsBase2<FAssetBundleData>
+{
+	enum
+	{
+		WithExportTextItem = true,
+		WithImportTextItem = true,
+	};
 };
diff --git a/Engine/Source/Runtime/AssetRegistry/Public/AssetData.h b/Engine/Source/Runtime/AssetRegistry/Public/AssetData.h
--- a/Engine/Source/Runtime/AssetRegistry/Public/AssetData.h
+++ b/Engine/Source/Runtime/AssetRegistry/Public/AssetData.h
@@ -12,30 +12,53 @@
 #include "Misc/PackageName.h"
 #include "Misc/SecureHash.h"
 #include "UObject/LinkerLoad.h"
-#include "SharedMapView.h"
+#include "AssetDataTagMap.h"
 #include "PrimaryAssetId.h"
 
 ASSETREGISTRY_API DECLARE_LOG_CATEGORY_EXTERN(LogAssetData, Log, All);
 
+/** Version used for serializing asset registry caches, both runtime and editor */
+struct ASSETREGISTRY_API FAssetRegistryVersion
+{
+	enum Type
+	{
+		PreVersioning = 0,		// From before file versioning was implemented
+		HardSoftDependencies,	// The first version of the runtime asset registry to include file versioning.
+		AddAssetRegistryState,	// Added FAssetRegistryState and support for piecemeal serialization
+		ChangedAssetData,		// AssetData serialization format changed, versions before this are not readable
+
+		// -----<new versions can be added above this line>-------------------------------------------------
+		VersionPlusOne,
+		LatestVersion = VersionPlusOne - 1
+	};
+
+	/** The GUID for this custom version number */
+	const static FGuid GUID;
+
+	/** Read/write the custom version to the archive, should call at the very beginning */
+	static bool SerializeVersion(FArchive& Ar, FAssetRegistryVersion::Type& Version);
+
+private:
+	FAssetRegistryVersion() {}
+};
+
 /** A class to hold important information about an assets found by the Asset Registry */
 class FAssetData
 {
 public:
 
-	/** The object path for the asset in the form 'Package.GroupNames.AssetName' */
+	/** The object path for the asset in the form PackageName.AssetName. Only top level objects in a package can have AssetData */
 	FName ObjectPath;
-	/** The name of the package in which the asset is found */
+	/** The name of the package in which the asset is found, this is the full long package name such as /Game/Path/Package */
 	FName PackageName;
-	/** The path to the package in which the asset is found */
+	/** The path to the package in which the asset is found, this is /Game/Path with the Package stripped off */
 	FName PackagePath;
-	/** The '.' delimited list of group names in which the asset is found. NAME_None if there were no groups */
-	FName GroupNames;
-	/** The name of the asset without the package or groups */
+	/** The name of the asset without the package */
 	FName AssetName;
 	/** The name of the asset's class */
 	FName AssetClass;
-	/** The map of values for properties that were marked AssetRegistrySearchable */
-	TSharedMapView<FName, FString> TagsAndValues;
+	/** The map of values for properties that were marked AssetRegistrySearchable or added by GetAssetRegistryTags */
+	FAssetDataTagMapSharedView TagsAndValues;
 	/** The IDs of the chunks this asset is located in for streaming install.  Empty if not assigned to a chunk */
 	TArray<int32> ChunkIDs;
 	/** Asset package flags */
@@ -47,65 +70,56 @@ public:
 	{}
 
 	/** Constructor */
-	FAssetData(FName InPackageName, FName InPackagePath, FName InGroupNames, FName InAssetName, FName InAssetClass, TMap<FName, FString> InTags, TArray<int32> InChunkIDs, uint32 InPackageFlags)
+	FAssetData(FName InPackageName, FName InPackagePath, FName InAssetName, FName InAssetClass, FAssetDataTagMap InTags = FAssetDataTagMap(), TArray<int32> InChunkIDs = TArray<int32>(), uint32 InPackageFlags = 0)
 		: PackageName(InPackageName)
 		, PackagePath(InPackagePath)
-		, GroupNames(InGroupNames)
 		, AssetName(InAssetName)
 		, AssetClass(InAssetClass)
-		, TagsAndValues(MakeSharedMapView(MoveTemp(InTags)))
+		, TagsAndValues(MoveTemp(InTags))
 		, ChunkIDs(MoveTemp(InChunkIDs))
 		, PackageFlags(InPackageFlags)
 	{
 		FString ObjectPathStr = PackageName.ToString() + TEXT(".");
 
-		if ( GroupNames != NAME_None )
-		{
-			ObjectPathStr += GroupNames.ToString() + TEXT(".");
-		}
-
 		ObjectPathStr += AssetName.ToString();
 
 		ObjectPath = FName(*ObjectPathStr);
 	}
 
-	/** Constructor taking a UObject */
-	FAssetData(const UObject* InAsset)
+	/** Constructor taking a UObject. By default trying to create one for a blueprint class will create one for the UBlueprint instead, but this can be overridden */
+	FAssetData(const UObject* InAsset, bool bAllowBlueprintClass = false)
 	{
-		if ( InAsset != NULL )
+		if ( InAsset != nullptr )
 		{
 			const UClass* InClass = Cast<UClass>(InAsset);
-			if (InClass && InClass->ClassGeneratedBy)
+			if (InClass && InClass->ClassGeneratedBy && !bAllowBlueprintClass)
 			{
-				// this UObject was generated by another class, use the generator to identify
-				// the source asset:
+				// For Blueprints, the AssetData refers to the UBlueprint and not the UBlueprintGeneratedClass
 				InAsset = InClass->ClassGeneratedBy;
 			}
 
 			const UPackage* Outermost = InAsset->GetOutermost();
 
-			// Find the group names
-			FString GroupNamesStr;
-			FString AssetNameStr;
-			InAsset->GetPathName(Outermost).Split(TEXT("."), &GroupNamesStr, &AssetNameStr, ESearchCase::CaseSensitive);
-
 			PackageName = Outermost->GetFName();
 			PackagePath = FName(*FPackageName::GetLongPackagePath(Outermost->GetName()));
-			GroupNames = FName(*GroupNamesStr);
 			AssetName = InAsset->GetFName();
 			AssetClass = InAsset->GetClass()->GetFName();
 			ObjectPath = FName(*InAsset->GetPathName());
 
-			TArray<UObject::FAssetRegistryTag> TagList;
-			InAsset->GetAssetRegistryTags(TagList);
+			TArray<UObject::FAssetRegistryTag> ObjectTags;
+			InAsset->GetAssetRegistryTags(ObjectTags);
 
-			TMap<FName, FString> TagsAndValuesMap;
-			for ( auto TagIt = TagList.CreateConstIterator(); TagIt; ++TagIt )
+			FAssetDataTagMap NewTagsAndValues;
+			for (UObject::FAssetRegistryTag& AssetRegistryTag : ObjectTags)
 			{
-				TagsAndValuesMap.Add(TagIt->Name, TagIt->Value);
+				if (AssetRegistryTag.Name != NAME_None && !AssetRegistryTag.Value.IsEmpty())
+				{
+					// Don't add empty tags
+					NewTagsAndValues.Add(AssetRegistryTag.Name, AssetRegistryTag.Value);
+				}
 			}
-			TagsAndValues = MakeSharedMapView(MoveTemp(TagsAndValuesMap));
 
+			TagsAndValues = FAssetDataTagMapSharedView(MoveTemp(NewTagsAndValues));
 			ChunkIDs = Outermost->GetChunkIDs();
 			PackageFlags = Outermost->GetPackageFlags();
 		}
@@ -294,7 +308,6 @@ public:
 		UE_LOG(LogAssetData, Log, TEXT("    ============================="));
 		UE_LOG(LogAssetData, Log, TEXT("        PackageName: %s"), *PackageName.ToString());
 		UE_LOG(LogAssetData, Log, TEXT("        PackagePath: %s"), *PackagePath.ToString());
-		UE_LOG(LogAssetData, Log, TEXT("        GroupNames: %s"), *GroupNames.ToString());
 		UE_LOG(LogAssetData, Log, TEXT("        AssetName: %s"), *AssetName.ToString());
 		UE_LOG(LogAssetData, Log, TEXT("        AssetClass: %s"), *AssetClass.ToString());
 		UE_LOG(LogAssetData, Log, TEXT("        TagsAndValues: %d"), TagsAndValues.Num());
@@ -339,38 +352,24 @@ public:
 		return (T*)Asset;
 	}
 
-	/** Operator for serialization */
-	friend FArchive& operator<<(FArchive& Ar, FAssetData& AssetData)
+	/** 
+	 * Serialize as part of the registry cache. This is not meant to be serialized as part of a package so  it does not handle versions normally
+	 * To version this data change FAssetRegistryVersion
+	 */
+	void SerializeForCache(FArchive& Ar)
 	{
-		// serialize out the asset info
-		Ar << AssetData.ObjectPath;
-		Ar << AssetData.PackagePath;
-		Ar << AssetData.AssetClass;
-		Ar << AssetData.GroupNames;
-
-		// these are derived from ObjectPath, probably can skip serializing at the expense of runtime string manipulation
-		Ar << AssetData.PackageName;
-		Ar << AssetData.AssetName;
+		// Serialize out the asset info
+		Ar << ObjectPath;
+		Ar << PackagePath;
+		Ar << AssetClass;
 
-		Ar << const_cast<TMap<FName, FString>&>(AssetData.TagsAndValues.GetMap());
-
-		if (Ar.UE4Ver() >= VER_UE4_CHANGED_CHUNKID_TO_BE_AN_ARRAY_OF_CHUNKIDS)
-		{
-			Ar << AssetData.ChunkIDs;
-		}
-		else if (Ar.UE4Ver() >= VER_UE4_ADDED_CHUNKID_TO_ASSETDATA_AND_UPACKAGE)
-		{
-			// loading old assetdata.  we weren't using this value yet, so throw it away
-			int ChunkID = -1;
-			Ar << ChunkID;
-		}
+		// These are derived from ObjectPath, we manually serialize them because they get pooled
+		Ar << PackageName;
+		Ar << AssetName;
 
-		if (Ar.UE4Ver() >= VER_UE4_COOKED_ASSETS_IN_EDITOR_SUPPORT)
-		{
-			Ar << AssetData.PackageFlags;
-		}
-
-		return Ar;
+		Ar << TagsAndValues;
+		Ar << ChunkIDs;
+		Ar << PackageFlags;
 	}
 
 private:
@@ -509,13 +508,15 @@ public:
 	{
 	}
 
-	friend FArchive& operator<<(FArchive& Ar, FAssetPackageData& PackageData)
+	/**
+	 * Serialize as part of the registry cache. This is not meant to be serialized as part of a package so  it does not handle versions normally
+	 * To version this data change FAssetRegistryVersion
+	 */
+	void SerializeForCache(FArchive& Ar)
 	{
-		Ar << PackageData.DiskSize;
-		Ar << PackageData.PackageGuid;
-		Ar << PackageData.PackageSourceHash;
-
-		return Ar;
+		Ar << DiskSize;
+		Ar << PackageGuid;
+		Ar << PackageSourceHash;
 	}
 };
 
@@ -665,6 +666,7 @@ struct FAssetIdentifier
 		return Hash;
 	}
 
+	/** Identifiers may be serialized as part of the registry cache, or in other contexts. If you make changes here you must also change FAssetRegistryVersion */
 	friend FArchive& operator<<(FArchive& Ar, FAssetIdentifier& AssetIdentifier)
 	{
 		// Serialize bitfield of which elements to serialize, in general many are empty
diff --git a/Engine/Source/Runtime/AssetRegistry/Public/AssetDataTagMap.h b/Engine/Source/Runtime/AssetRegistry/Public/AssetDataTagMap.h
new file mode 100644
--- /dev/null
+++ b/Engine/Source/Runtime/AssetRegistry/Public/AssetDataTagMap.h
@@ -0,0 +1,166 @@
+// Copyright 1998-2017 Epic Games, Inc. All Rights Reserved.
+
+#pragma once
+
+#include "CoreMinimal.h"
+#include "SortedMap.h"
+
+/** Type of tag map that can be used during construction */
+typedef TSortedMap<FName, FString, FDefaultAllocator, FNameSortIndexes> FAssetDataTagMap;
+
+/** Wrapper of shared pointer to a map */
+class FAssetDataTagMapSharedView
+{
+	
+public:
+	/** Default constructor - empty map */
+	FAssetDataTagMapSharedView()
+	{
+	}
+
+	/** Constructor from an existing map pointer */
+	FAssetDataTagMapSharedView(TSharedPtr<FAssetDataTagMap> InMap)
+		: Map(InMap)
+	{
+	}
+
+	/** Constructor from an existing map pointer */
+	FAssetDataTagMapSharedView(FAssetDataTagMap&& InMap)
+	{
+		if (InMap.Num())
+		{
+			Map = MakeShareable(new FAssetDataTagMap(MoveTemp(InMap)));
+		}
+	}
+
+	/** Find a value by key (nullptr if not found) */
+	const FString* Find(FAssetDataTagMap::KeyConstPointerType Key) const
+	{
+		return GetMap().Find(Key);
+	}
+
+	/** Find a value by key (abort if not found) */
+	const FString& FindChecked(FAssetDataTagMap::KeyConstPointerType Key) const
+	{
+		return GetMap().FindChecked(Key);
+	}
+
+	/** Find a value by key (default value if not found) */
+	FString FindRef(FAssetDataTagMap::KeyConstPointerType Key) const
+	{
+		return GetMap().FindRef(Key);
+	}
+
+	/** Determine whether a key is present in the map */
+	bool Contains(FAssetDataTagMap::KeyConstPointerType Key) const
+	{
+		return GetMap().Contains(Key);
+	}
+
+	/** Retrieve size of map */
+	int32 Num() const
+	{
+		return GetMap().Num();
+	}
+
+	/** Populate an array with all the map's keys */
+	template <class FAllocator>
+	int32 GetKeys(TArray<FName, FAllocator>& OutKeys) const
+	{
+		return GetMap().GetKeys(OutKeys);
+	}
+
+	/** Populate an array with all the map's keys */
+	template <class FAllocator>
+	void GenerateKeyArray(TArray<FName, FAllocator>& OutKeys) const
+	{
+		GetMap().GenerateKeyArray(OutKeys);
+	}
+
+	/** Populate an array with all the map's values */
+	template <class FAllocator>
+	void GenerateValueArray(TArray<FName, FAllocator>& OutValues) const
+	{
+		GetMap().GenerateValueArray(OutValues);
+	}
+
+	/** Iterate all key-value pairs */
+	FAssetDataTagMap::TConstIterator CreateConstIterator() const
+	{
+		return GetMap().CreateConstIterator();
+	}
+
+	/** Const access to the underlying map, mainly for taking a copy */
+	const FAssetDataTagMap& GetMap() const
+	{
+		static FAssetDataTagMap EmptyMap;
+
+		if (Map.IsValid())
+		{
+			return *Map;
+		}
+
+		return EmptyMap;
+	}
+
+	/** Returns amount of extra memory used by this structure, including shared ptr overhead */
+	uint32 GetAllocatedSize() const
+	{
+		uint32 AllocatedSize = 0;
+		if (Map.IsValid())
+		{
+			AllocatedSize += sizeof(FAssetDataTagMap); // Map itself
+			AllocatedSize += (sizeof(int32) * 2); // SharedPtr overhead
+			AllocatedSize += Map->GetAllocatedSize();
+		}
+
+		return AllocatedSize;
+	}
+private:
+	friend class FAssetData;
+
+	FORCEINLINE friend FArchive& operator<<(FArchive& Ar, FAssetDataTagMapSharedView& SharedView)
+	{
+		if (Ar.IsSaving())
+		{
+			if (SharedView.Map.IsValid())
+			{
+				Ar << *SharedView.Map;
+			}
+			else
+			{
+				FAssetDataTagMap TempMap;
+				Ar << TempMap;
+			}
+		}
+		else
+		{
+			// Serialize into temporary map, if it isn't empty move memory into new map
+			FAssetDataTagMap TempMap;
+			Ar << TempMap;
+
+			if (TempMap.Num())
+			{
+				SharedView.Map = MakeShareable(new FAssetDataTagMap(MoveTemp(TempMap)));
+			}
+		}
+
+		return Ar;
+	}
+
+	/** Range for iterator access - DO NOT USE DIRECTLY */
+	friend FAssetDataTagMap::RangedForConstIteratorType begin(const FAssetDataTagMapSharedView& View)
+	{
+		return begin(View.GetMap());
+	}
+
+	/** Range for iterator access - DO NOT USE DIRECTLY */
+	friend FAssetDataTagMap::RangedForConstIteratorType end(const FAssetDataTagMapSharedView& View)
+	{
+		return end(View.GetMap());
+	}
+
+	/** Pointer to map being wrapped, it is created on demand */
+	TSharedPtr<FAssetDataTagMap> Map;
+};
+
diff --git a/Engine/Source/Runtime/AssetRegistry/Public/AssetRegistryState.h b/Engine/Source/Runtime/AssetRegistry/Public/AssetRegistryState.h
--- a/Engine/Source/Runtime/AssetRegistry/Public/AssetRegistryState.h
+++ b/Engine/Source/Runtime/AssetRegistry/Public/AssetRegistryState.h
@@ -8,7 +8,6 @@
 
 class FDependsNode;
 struct FARFilter;
-enum class ERuntimeRegistryVersion;
 
 /** Load/Save options used to modify how the cache is serialized. These are read out of the AssetRegistry section of Engine.ini and can be changed per platform. */
 struct FAssetRegistrySerializationOptions
@@ -31,6 +30,9 @@ struct FAssetRegistrySerializationOptions
 	/** True if CookFilterlistTagsByClass is a whitelist. False if it is a blacklist. */
 	bool bUseAssetRegistryTagsWhitelistInsteadOfBlacklist;
 
+	/** True if we want to only write out asset data if it has valid tags. This saves memory by not saving data for things like textures */
+	bool bFilterAssetDataWithNoTags;
+
 	/** The map of classname to tag set of tags that are allowed in cooked builds. This is either a whitelist or blacklist depending on bUseAssetRegistryTagsWhitelistInsteadOfBlacklist */
 	TMap<FName, TSet<FName>> CookFilterlistTagsByClass;
 
@@ -41,12 +43,14 @@ struct FAssetRegistrySerializationOptions
 		, bSerializeManageDependencies(false)
 		, bSerializePackageData(false)
 		, bUseAssetRegistryTagsWhitelistInsteadOfBlacklist(false)
+		, bFilterAssetDataWithNoTags(false)
 	{}
 
 	/** Options used to read/write the DevelopmentAssetRegistry, which includes all data */
 	void ModifyForDevelopment()
 	{
 		bSerializeAssetRegistry = bSerializeDependencies = bSerializeSearchableNameDependencies = bSerializeManageDependencies = bSerializePackageData = true;
+		bFilterAssetDataWithNoTags = false;
 	}
 };
 
@@ -112,10 +116,10 @@ public:
 	}
 
 	/**
-	 * Gets the asset data for the specified object path
+	 * Gets the asset data for the specified package name
 	 *
-	 * @param ObjectPath the path of the object to be looked up
-	 * @return the assets data, null if not found
+	 * @param PackageName the path of the package to be looked up
+	 * @return an array of AssetData*, empty if nothing found
 	 */
 	const TArray<const FAssetData*>& GetAssetsByPackageName(const FName PackageName) const
 	{
@@ -129,6 +133,42 @@ public:
 		return InvalidArray;
 	}
 
+	/**
+	 * Gets the asset data for the specified asset class
+	 *
+	 * @param ClassName the class name of the assets to look for
+	 * @return An array of AssetData*, empty if nothing found
+	 */
+	const TArray<const FAssetData*>& GetAssetsByClassName(const FName ClassName) const
+	{
+		static TArray<const FAssetData*> InvalidArray;
+		const TArray<FAssetData*>* FoundAssetArray = CachedAssetsByClass.Find(ClassName);
+		if (FoundAssetArray)
+		{
+			return reinterpret_cast<const TArray<const FAssetData*>&>(*FoundAssetArray);
+		}
+
+		return InvalidArray;
+	}
+
+	/**
+	 * Gets the asset data for the specified asset tag
+	 *
+	 * @param TagName the tag name to search for 
+	 * @return An array of AssetData*, empty if nothing found
+	 */
+	const TArray<const FAssetData*>& GetAssetsByTagName(const FName TagName) const
+	{
+		static TArray<const FAssetData*> InvalidArray;
+		const TArray<FAssetData*>* FoundAssetArray = CachedAssetsByTag.Find(TagName);
+		if (FoundAssetArray)
+		{
+			return reinterpret_cast<const TArray<const FAssetData*>&>(*FoundAssetArray);
+		}
+
+		return InvalidArray;
+	}
+
 	/** Returns const version of internal ObjectPath->AssetData map for fast iteration */
 	const TMap<FName, const FAssetData*>& GetObjectPathToAssetDataMap() const
 	{
@@ -163,18 +203,26 @@ public:
 	void Reset();
 
 	/** Initializes cache from existing set of asset data and depends nodes */
-	void InitializeFromExisting(const TMap<FName, FAssetData*>& AssetDataMap, const TMap<FAssetIdentifier, FDependsNode*>& DependsNodeMap, const TMap<FName, FAssetPackageData*>& AssetPackageDataMap, const FAssetRegistrySerializationOptions& Options);
-	void InitializeFromExisting(const FAssetRegistryState& Existing, const FAssetRegistrySerializationOptions& Options)
+	void InitializeFromExisting(const TMap<FName, FAssetData*>& AssetDataMap, const TMap<FAssetIdentifier, FDependsNode*>& DependsNodeMap, const TMap<FName, FAssetPackageData*>& AssetPackageDataMap, const FAssetRegistrySerializationOptions& Options, bool bRefreshExisting = false);
+	void InitializeFromExisting(const FAssetRegistryState& Existing, const FAssetRegistrySerializationOptions& Options, bool bRefreshExisting = false)
 	{
-		InitializeFromExisting(Existing.CachedAssetsByObjectPath, Existing.CachedDependsNodes, Existing.CachedPackageData, Options);
+		InitializeFromExisting(Existing.CachedAssetsByObjectPath, Existing.CachedDependsNodes, Existing.CachedPackageData, Options, bRefreshExisting);
 	}
 
-	/** Prunes an asset cache, this removes asset data, nodes, and package data that isn't needed. If Required is empty it will only apply ignored */
-	void PruneAssetData(const TSet<FName>& RequiredPackages, const TSet<FName>& RemovePackages);
+	/** 
+	 * Prunes an asset cache, this removes asset data, nodes, and package data that isn't needed. 
+	 * @param RequiredPackages If set, only these packages will be maintained. If empty it will keep all unless filtered by other parameters
+	 * @param RemovePackages These packages will be removed from the current set
+	 * @param bFilterAssetDataWithNoTags If true, any AssetData with no Tags will be removed
+	 */
+	void PruneAssetData(const TSet<FName>& RequiredPackages, const TSet<FName>& RemovePackages, bool bFilterAssetDataWithNoTags = false);
 
 	/** Serialize the registry to/from a file, skipping editor only data */
 	bool Serialize(FArchive& Ar, FAssetRegistrySerializationOptions& Options);
 
+	/** Returns memory size of entire registry, optionally logging sizes */
+	uint32 GetAllocatedSize(bool bLogDetailed = false) const;
+
 	/** Checks a filter to make sure there are no illegal entries */
 	static bool IsFilterValid(const FARFilter& Filter, bool bAllowRecursion);
 
@@ -191,9 +239,6 @@ private:
 	/** Removes the depends node and updates the dependencies to no longer contain it as as a referencer. */
 	bool RemoveDependsNode(const FAssetIdentifier& Identifier);
 
-	/** Reads old format asset registry state */
-	bool ReadDeprecated(FArchive& Ar, ERuntimeRegistryVersion Version);
-
 	/** The map of ObjectPath names to asset data for assets saved to disk */
 	TMap<FName, FAssetData*> CachedAssetsByObjectPath;
 
diff --git a/Engine/Source/Runtime/AssetRegistry/Public/IAssetRegistry.h b/Engine/Source/Runtime/AssetRegistry/Public/IAssetRegistry.h
--- a/Engine/Source/Runtime/AssetRegistry/Public/IAssetRegistry.h
+++ b/Engine/Source/Runtime/AssetRegistry/Public/IAssetRegistry.h
@@ -157,6 +157,9 @@ public:
 	/** Finds Package data for a package name. This data is only updated on save and can only be accessed for valid packages */
 	virtual const FAssetPackageData* GetAssetPackageData(FName PackageName) const = 0;
 
+	/** Uses the asset registry to look for ObjectRedirectors. This will follow the chain of redirectors. It will return the original path if no redirectors are found */
+	virtual FName GetRedirectedObjectPath(const FName ObjectPath) const = 0;
+
 	/** Returns true if the specified ClassName's ancestors could be found. If so, OutAncestorClassNames is a list of all its ancestors */
 	virtual bool GetAncestorClassNames(FName ClassName, TArray<FName>& OutAncestorClassNames) const = 0;
 
@@ -302,14 +305,18 @@ public:
 	/** Serialize the registry to/from a file, skipping editor only data */
 	virtual void Serialize(FArchive& Ar) = 0;
 
+	/** Returns memory size of entire registry, optionally logging sizes */
+	virtual uint32 GetAllocatedSize(bool bLogDetailed = false) const = 0;
+
 	/**
 	 * Fills in a AssetRegistryState with a copy of the data in the internal cache, overriding some
 	 *
-	 * @param OutState		This will be filled in with a copy of the asset data, platform data, and dependency data
-	 * @param Options		Serialization options that will be used to write this later
-	 * @param OverrideData	Map of ObjectPath to AssetData. If non empty, it will use this map of AssetData, and will filter Platform/Dependency data to only include this set
+	 * @param OutState			This will be filled in with a copy of the asset data, platform data, and dependency data
+	 * @param Options			Serialization options that will be used to write this later
+	 * @param bRefreshExisting	If true, will not delete or add packages in OutState and will just update things that already exist
+	 * @param OverrideData		Map of ObjectPath to AssetData. If non empty, it will use this map of AssetData, and will filter Platform/Dependency data to only include this set
 	 */
-	virtual void InitializeTemporaryAssetRegistryState(FAssetRegistryState& OutState, const FAssetRegistrySerializationOptions& Options, const TMap<FName, FAssetData*>& OverrideData = TMap<FName, FAssetData*>()) const = 0;
+	virtual void InitializeTemporaryAssetRegistryState(FAssetRegistryState& OutState, const FAssetRegistrySerializationOptions& Options, bool bRefreshExisting = false, const TMap<FName, FAssetData*>& OverrideData = TMap<FName, FAssetData*>()) const = 0;
 
 	/** Fills in FAssetRegistrySerializationOptions from ini, optionally using a target platform ini name */
 	virtual void InitializeSerializationOptions(FAssetRegistrySerializationOptions& Options, const FString& PlatformIniName = FString()) const = 0;
diff --git a/Engine/Source/Runtime/CinematicCamera/Private/CineCameraComponent.cpp b/Engine/Source/Runtime/CinematicCamera/Private/CineCameraComponent.cpp
--- a/Engine/Source/Runtime/CinematicCamera/Private/CineCameraComponent.cpp
+++ b/Engine/Source/Runtime/CinematicCamera/Private/CineCameraComponent.cpp
@@ -251,9 +251,9 @@ void UCineCameraComponent::UpdateDebugFocusPlane()
 #if WITH_EDITORONLY_DATA
 	if (FocusSettings.bDrawDebugFocusPlane && DebugFocusPlaneMesh && DebugFocusPlaneComponent)
 	{
-		FVector const CamLocation = ComponentToWorld.GetLocation();
-		FVector const CamDir = ComponentToWorld.GetRotation().Vector();
-		FVector const FocusPoint = ComponentToWorld.GetLocation() + CamDir * GetDesiredFocusDistance(CamLocation);
+		FVector const CamLocation = GetComponentTransform().GetLocation();
+		FVector const CamDir = GetComponentTransform().GetRotation().Vector();
+		FVector const FocusPoint = GetComponentTransform().GetLocation() + CamDir * GetDesiredFocusDistance(CamLocation);
 		DebugFocusPlaneComponent->SetWorldLocation(FocusPoint);
 	}
 #endif
diff --git a/Engine/Source/Runtime/ClothingSystemRuntime/Private/ClothingSimulation.cpp b/Engine/Source/Runtime/ClothingSystemRuntime/Private/ClothingSimulation.cpp
--- a/Engine/Source/Runtime/ClothingSystemRuntime/Private/ClothingSimulation.cpp
+++ b/Engine/Source/Runtime/ClothingSystemRuntime/Private/ClothingSimulation.cpp
@@ -130,7 +130,7 @@ void FClothingSimulationBase::SkinPhysicsMesh(UClothingAsset* InAsset, const FCl
 void FClothingSimulationBase::FillContext(USkeletalMeshComponent* InComponent, IClothingSimulationContext* InOutContext)
 {
 	FClothingSimulationContextBase* BaseContext = static_cast<FClothingSimulationContextBase*>(InOutContext);
-	BaseContext->ComponentToWorld = InComponent->ComponentToWorld;
+	BaseContext->ComponentToWorld = InComponent->GetComponentTransform();
 	BaseContext->PredictedLod = InComponent->PredictedLODLevel;
 	InComponent->GetWindForCloth_GameThread(BaseContext->WindVelocity, BaseContext->WindAdaption);
 	
diff --git a/Engine/Source/Runtime/ClothingSystemRuntime/Private/ClothingSimulationNv.cpp b/Engine/Source/Runtime/ClothingSystemRuntime/Private/ClothingSimulationNv.cpp
--- a/Engine/Source/Runtime/ClothingSystemRuntime/Private/ClothingSimulationNv.cpp
+++ b/Engine/Source/Runtime/ClothingSystemRuntime/Private/ClothingSimulationNv.cpp
@@ -608,7 +608,7 @@ void FClothingSimulationNv::GetSimulationData(TMap<int32, FClothSimulData>& OutD
 
 	bool bFoundSimData = false;
 
-	const FTransform& OwnerTransform = InOwnerComponent->ComponentToWorld;
+	const FTransform& OwnerTransform = InOwnerComponent->GetComponentTransform();
 	
 	for(int32 ActorIdx = 0; ActorIdx < NumActors; ++ActorIdx)
 	{
diff --git a/Engine/Source/Runtime/Core/Private/Containers/Algo/AlgosTest.cpp b/Engine/Source/Runtime/Core/Private/Containers/Algo/AlgosTest.cpp
--- a/Engine/Source/Runtime/Core/Private/Containers/Algo/AlgosTest.cpp
+++ b/Engine/Source/Runtime/Core/Private/Containers/Algo/AlgosTest.cpp
@@ -249,5 +249,35 @@ bool FAlgosTest::RunTest(const FString& Parameters)
 		}
 	}
 
+	// binary search
+	{
+		// Verify static array case
+		int StaticArray[] = { 2,4,6,6,6,8 };
+
+		check(Algo::BinarySearch(StaticArray, 6) == 2);
+		check(Algo::BinarySearch(StaticArray, 5) == INDEX_NONE);
+
+		check(Algo::LowerBound(StaticArray, 6) == 2);
+		check(Algo::LowerBound(StaticArray, 5) == 2);
+		check(Algo::UpperBound(StaticArray, 6) == 5);
+		check(Algo::LowerBound(StaticArray, 7) == 5);
+		check(Algo::LowerBound(StaticArray, 9) == 6);
+
+		// Dynamic array case
+		TArray<int32> IntArray = { 2,2,4,4,6,6,6,8,8 };
+
+		check(Algo::BinarySearch(IntArray, 6) == 4);
+		check(Algo::BinarySearch(IntArray, 5) == INDEX_NONE);
+
+		check(Algo::LowerBound(IntArray, 2) == 0);
+		check(Algo::UpperBound(IntArray, 2) == 2);
+		check(Algo::LowerBound(IntArray, 6) == 4);
+		check(Algo::UpperBound(IntArray, 6) == 7);
+		check(Algo::LowerBound(IntArray, 5) == 4);
+		check(Algo::UpperBound(IntArray, 5) == 4);
+		check(Algo::LowerBound(IntArray, 7) == 7);
+		check(Algo::LowerBound(IntArray, 9) == 9);
+	}
+
 	return true;
 }
diff --git a/Engine/Source/Runtime/Core/Private/Containers/ContainersTest.cpp b/Engine/Source/Runtime/Core/Private/Containers/ContainersTest.cpp
--- a/Engine/Source/Runtime/Core/Private/Containers/ContainersTest.cpp
+++ b/Engine/Source/Runtime/Core/Private/Containers/ContainersTest.cpp
@@ -4,14 +4,15 @@
 #include "Misc/AssertionMacros.h"
 #include "Containers/Array.h"
 #include "Containers/Map.h"
+#include "Containers/SortedMap.h"
+#include "Containers/ArrayView.h"
 #include "Misc/AutomationTest.h"
-
-
-IMPLEMENT_SIMPLE_AUTOMATION_TEST(FContainersTest, "System.Core.Misc.Containers", EAutomationTestFlags::ApplicationContextMask | EAutomationTestFlags::SmokeFilter)
+#include "Stats/StatsMisc.h"
+#include "Math/RandomStream.h"
 
 #define MAX_TEST_OBJECTS      65
 #define MAX_TEST_OBJECTS_STEP 1
-
+#define RANDOM_SEED 12345
 
 namespace
 {
@@ -270,60 +271,268 @@ namespace
 			++CIter;
 		}
 
-		check(end(                  Cont) == Iter);
-		check(end((const Container&)Cont) == CIter);
+		check(!Iter);
+		check(!CIter);
 	}
-}
 
-bool FContainersTest::RunTest( const FString& Parameters )
-{
-	typedef TMap<int32, FContainerTestValueType> TestMapType;
+	template <typename KeyType>
+	KeyType GenerateTestKey(int32 Input)
+	{
+		return (KeyType)Input;
+	}
+
+	template <>
+	FName GenerateTestKey<FName>(int32 Input)
+	{
+		// Don't use _foo as we want to test the slower compare path
+		return FName(*FString::Printf(TEXT("TestName%d"), Input));
+	}
+
+	template <>
+	FString GenerateTestKey<FString>(int32 Input)
+	{
+		return FString::Printf(TEXT("TestString%d"), Input);
+	}
+
+	template <typename ContainerType, typename KeyType>
+	void RunContainerTests()
+	{
+		ContainerType Cont;
+
+		ContainerTestStats.Reset();
+		// Subtract one to account for temporaries that will be created during an Add
+		for (int32 Count = 0; Count < MAX_TEST_OBJECTS - 1; Count += MAX_TEST_OBJECTS_STEP)
+		{
+			for (int32 N = 0; N != Count; ++N)
+			{
+				Cont.Add(GenerateTestKey<KeyType>(N), FContainerTestValueType(TEXT("New Value")));
+				CheckContainerNum(Cont);
+				CheckContainerEnds(Cont);
+				CheckContainerElements(Cont);
+			}
+
+			for (int32 N = 0; N != Count; ++N)
+			{
+				Cont.Remove(GenerateTestKey<KeyType>(N));
+				CheckContainerNum(Cont);
+				CheckContainerEnds(Cont);
+				CheckContainerElements(Cont);
+			}
+
+			for (int32 N = 0; N != Count; ++N)
+			{
+				Cont.Add(GenerateTestKey<KeyType>((Count - 1) - N), FContainerTestValueType(TEXT("New Value")));
+				CheckContainerNum(Cont);
+				CheckContainerEnds(Cont);
+				CheckContainerElements(Cont);
+			}
+
+			for (int32 N = 0; N != Count; ++N)
+			{
+				Cont.Remove(GenerateTestKey<KeyType>(N));
+				CheckContainerNum(Cont);
+				CheckContainerEnds(Cont);
+				CheckContainerElements(Cont);
+			}
+		}
+	}
+
+	template <typename ContainerType, typename KeyType>
+	void RunPerformanceTest(const FString& Description, int32 NumObjects, int32 NumOperations)
+	{
+		ContainerTestStats.Reset();
+		
+		ContainerType Cont;
+		FRandomStream RandomStream(RANDOM_SEED);
+
+		// Prep keys, not part of performance test
+		TArray<KeyType> KeyArray;
+		KeyArray.Reserve(NumObjects);
+
+		for (int32 i = 0; i < NumObjects; i++)
+		{
+			KeyArray.Add(GenerateTestKey<KeyType>(i));
+		}
+
+		for (int32 i = 0; i < NumObjects; i++)
+		{
+			int32 SwapIndex = RandomStream.RandRange(0, NumObjects - 1);
+			if (i != SwapIndex)
+			{
+				KeyArray.Swap(i, SwapIndex);
+			}
+		}
 
-	ContainerTestStats.Reset();
+		FScopeLogTime LogTimePtr(*FString::Printf(TEXT("%s objects=%d count=%d"), *Description, NumObjects, NumOperations), nullptr, FScopeLogTime::ScopeLog_Milliseconds);
 
-	TestMapType TestMap;
+		// Add elements in stably randomized order
+		for (int32 i = 0; i < NumObjects; i++)
+		{
+			Cont.Add(KeyArray[i], FString(TEXT("New Value")));
+		}
+		
+		// Now do searches
+		for (int32 i = 0; i < NumOperations; i++)
+		{
+			KeyType& Key = KeyArray[RandomStream.RandRange(0, NumObjects - 1)];
 
-	// Subtract one to account for temporaries that will be created during an Add
-	for (int32 Count = 0; Count < MAX_TEST_OBJECTS - 1; Count += MAX_TEST_OBJECTS_STEP)
+			FString* FoundValue = Cont.Find(Key);
+			check(FoundValue);
+		}
+	}
+
+	template <typename ContainerType, typename KeyType>
+	void RunSetPerformanceTest(const FString& Description, int32 NumObjects, int32 NumOperations)
 	{
-		for (int32 N = 0; N != Count; ++N)
+		ContainerTestStats.Reset();
+
+		ContainerType Cont;
+		FRandomStream RandomStream(RANDOM_SEED);
+
+		// Prep keys, not part of performance test
+		TArray<KeyType> KeyArray;
+		KeyArray.Reserve(NumObjects);
+
+		for (int32 i = 0; i < NumObjects; i++)
 		{
-			TestMap.Add(N, FContainerTestValueType(TEXT("New Value")));
-			CheckContainerNum     (TestMap);
-			CheckContainerEnds    (TestMap);
-			CheckContainerElements(TestMap);
+			KeyArray.Add(GenerateTestKey<KeyType>(i));
 		}
 
-		for (int32 N = 0; N != Count; ++N)
+		for (int32 i = 0; i < NumObjects; i++)
 		{
-			TestMap.Remove(N);
-			CheckContainerNum     (TestMap);
-			CheckContainerEnds    (TestMap);
-			CheckContainerElements(TestMap);
+			int32 SwapIndex = RandomStream.RandRange(0, NumObjects - 1);
+			if (i != SwapIndex)
+			{
+				KeyArray.Swap(i, SwapIndex);
+			}
 		}
 
-		for (int32 N = 0; N != Count; ++N)
+		FScopeLogTime LogTimePtr(*FString::Printf(TEXT("%s objects=%d count=%d"), *Description, NumObjects, NumOperations), nullptr, FScopeLogTime::ScopeLog_Milliseconds);
+
+		// Add elements in stably randomized order
+		for (int32 i = 0; i < NumObjects; i++)
 		{
-			TestMap.Add((Count - 1) - N, FContainerTestValueType(TEXT("New Value")));
-			CheckContainerNum     (TestMap);
-			CheckContainerEnds    (TestMap);
-			CheckContainerElements(TestMap);
+			Cont.Add(KeyArray[i]);
 		}
 
-		for (int32 N = 0; N != Count; ++N)
+		// Now do searches
+		for (int32 i = 0; i < NumOperations; i++)
 		{
-			TestMap.Remove(N);
-			CheckContainerNum     (TestMap);
-			CheckContainerEnds    (TestMap);
-			CheckContainerElements(TestMap);
+			KeyType& Key = KeyArray[RandomStream.RandRange(0, NumObjects - 1)];
+
+			bool FoundValue = Cont.Contains(Key);
+			check(FoundValue);
 		}
 	}
 
+	template<typename TValueType>
+	struct FCaseSensitiveLookupKeyFuncs : BaseKeyFuncs<TValueType, FString>
+	{
+		static FORCEINLINE const FString& GetSetKey(const TPair<FString, TValueType>& Element)
+		{
+			return Element.Key;
+		}
+		static FORCEINLINE bool Matches(const FString& A, const FString& B)
+		{
+			return A.Equals(B, ESearchCase::CaseSensitive);
+		}
+		static FORCEINLINE uint32 GetKeyHash(const FString& Key)
+		{
+			return FCrc::StrCrc32<TCHAR>(*Key);
+		}
+	};
+}
+
+IMPLEMENT_SIMPLE_AUTOMATION_TEST(FContainersSmokeTest, "System.Core.Containers.Smoke", EAutomationTestFlags::ApplicationContextMask | EAutomationTestFlags::SmokeFilter)
+bool FContainersSmokeTest::RunTest( const FString& Parameters )
+{
+	RunContainerTests<TMap<int32, FContainerTestValueType>, int32>();
+
 	return true;
 }
 
-//Array view tests
-#include "Containers/ArrayView.h"
+IMPLEMENT_SIMPLE_AUTOMATION_TEST(FContainersFullTest, "System.Core.Containers.Full", EAutomationTestFlags::ApplicationContextMask | EAutomationTestFlags::EngineFilter)
+bool FContainersFullTest::RunTest(const FString& Parameters)
+{
+	RunContainerTests<TMap<int32, FContainerTestValueType>, int32>();
+	RunContainerTests<TMap<FName, FContainerTestValueType>, FName>();
+	RunContainerTests<TMap<FString, FContainerTestValueType>, FString>();
+	RunContainerTests<TMap<int32, FContainerTestValueType, TInlineSetAllocator<64>>, int32>();
+	RunContainerTests<TMap<FString, FContainerTestValueType, FDefaultSetAllocator, FCaseSensitiveLookupKeyFuncs<FContainerTestValueType>>, FString>();
+
+	RunContainerTests<TSortedMap<int32, FContainerTestValueType>, int32>();
+	RunContainerTests<TSortedMap<FName, FContainerTestValueType>, FName>();
+	RunContainerTests<TSortedMap<FString, FContainerTestValueType>, FString>();
+	RunContainerTests<TSortedMap<FString, FContainerTestValueType, TInlineAllocator<64>>, FString>();
+
+	// Verify use of FName index sorter with SortedMap
+
+	TSortedMap<FName, int32, FDefaultAllocator, FNameSortIndexes> NameMap;
+	NameMap.Add(NAME_NameProperty);
+	NameMap.Add(NAME_FloatProperty);
+	NameMap.Add(NAME_None);
+	NameMap.Add(NAME_IntProperty);
+
+	auto It = NameMap.CreateConstIterator();
+
+	check(It->Key == NAME_None); ++It;
+	check(It->Key == NAME_IntProperty); ++It;
+	check(It->Key == NAME_FloatProperty); ++It;
+	check(It->Key == NAME_NameProperty); ++It;
+	check(!It);
+
+	return true;
+}
+
+IMPLEMENT_SIMPLE_AUTOMATION_TEST(FContainerPerformanceTest, "System.Core.Containers.Performance", EAutomationTestFlags::ApplicationContextMask | EAutomationTestFlags::EngineFilter)
+bool FContainerPerformanceTest::RunTest(const FString& Parameters)
+{
+	RunPerformanceTest<TMap<int32, FString>, int32>(TEXT("TMap int32"), 1, 1000000);
+	RunPerformanceTest<TMap<int32, FString>, int32>(TEXT("TMap int32"), 10, 1000000);
+	RunPerformanceTest<TMap<int32, FString>, int32>(TEXT("TMap int32"), 100, 1000000);
+	RunPerformanceTest<TMap<int32, FString>, int32>(TEXT("TMap int32"), 1000, 1000000);
+	RunPerformanceTest<TMap<int32, FString>, int32>(TEXT("TMap int32"), 10000, 1000000);
+
+	RunPerformanceTest<TMap<FName, FString>, FName>(TEXT("TMap FName"), 1, 1000000);
+	RunPerformanceTest<TMap<FName, FString>, FName>(TEXT("TMap FName"), 10, 1000000);
+	RunPerformanceTest<TMap<FName, FString>, FName>(TEXT("TMap FName"), 100, 1000000);
+	RunPerformanceTest<TMap<FName, FString>, FName>(TEXT("TMap FName"), 1000, 1000000);
+	RunPerformanceTest<TMap<FName, FString>, FName>(TEXT("TMap FName"), 10000, 1000000);
+
+	RunPerformanceTest<TMap<FString, FString>, FString>(TEXT("TMap FString"), 1, 1000000);
+	RunPerformanceTest<TMap<FString, FString>, FString>(TEXT("TMap FString"), 10, 1000000);
+	RunPerformanceTest<TMap<FString, FString>, FString>(TEXT("TMap FString"), 100, 1000000);
+	RunPerformanceTest<TMap<FString, FString>, FString>(TEXT("TMap FString"), 1000, 1000000);
+	RunPerformanceTest<TMap<FString, FString>, FString>(TEXT("TMap FString"), 10000, 1000000);
+
+	RunPerformanceTest<TSortedMap<int32, FString>, int32>(TEXT("TSortedMap int32"), 1, 1000000);
+	RunPerformanceTest<TSortedMap<int32, FString>, int32>(TEXT("TSortedMap int32"), 10, 1000000);
+	RunPerformanceTest<TSortedMap<int32, FString>, int32>(TEXT("TSortedMap int32"), 100, 1000000);
+	RunPerformanceTest<TSortedMap<int32, FString>, int32>(TEXT("TSortedMap int32"), 1000, 1000000);
+	RunPerformanceTest<TSortedMap<int32, FString>, int32>(TEXT("TSortedMap int32"), 10000, 1000000);
+
+	RunPerformanceTest<TSortedMap<FName, FString>, FName>(TEXT("TSortedMap FName"), 1, 1000000);
+	RunPerformanceTest<TSortedMap<FName, FString>, FName>(TEXT("TSortedMap FName"), 10, 1000000);
+	RunPerformanceTest<TSortedMap<FName, FString>, FName>(TEXT("TSortedMap FName"), 100, 1000000);
+	RunPerformanceTest<TSortedMap<FName, FString>, FName>(TEXT("TSortedMap FName"), 1000, 1000000);
+	RunPerformanceTest<TSortedMap<FName, FString>, FName>(TEXT("TSortedMap FName"), 10000, 1000000);
+
+	RunPerformanceTest<TSortedMap<FString, FString>, FString>(TEXT("TSortedMap FString"), 1, 1000000);
+	RunPerformanceTest<TSortedMap<FString, FString>, FString>(TEXT("TSortedMap FString"), 10, 1000000);
+	RunPerformanceTest<TSortedMap<FString, FString>, FString>(TEXT("TSortedMap FString"), 100, 1000000);
+	RunPerformanceTest<TSortedMap<FString, FString>, FString>(TEXT("TSortedMap FString"), 1000, 1000000);
+	RunPerformanceTest<TSortedMap<FString, FString>, FString>(TEXT("TSortedMap FString"), 10000, 1000000);
+
+	RunSetPerformanceTest<TSet<FName>, FName>(TEXT("TSet FName"), 1, 1000000);
+	RunSetPerformanceTest<TSet<FName>, FName>(TEXT("TSet FName"), 10, 1000000);
+	RunSetPerformanceTest<TSet<FName>, FName>(TEXT("TSet FName"), 100, 1000000);
+
+	RunSetPerformanceTest<TArray<FName>, FName>(TEXT("TArray FName"), 1, 1000000);
+	RunSetPerformanceTest<TArray<FName>, FName>(TEXT("TArray FName"), 10, 1000000);
+	RunSetPerformanceTest<TArray<FName>, FName>(TEXT("TArray FName"), 100, 1000000);
+
+	return true;
+}
 
 namespace ArrayViewTests
 {
diff --git a/Engine/Source/Runtime/Core/Private/Misc/CoreDelegates.cpp b/Engine/Source/Runtime/Core/Private/Misc/CoreDelegates.cpp
--- a/Engine/Source/Runtime/Core/Private/Misc/CoreDelegates.cpp
+++ b/Engine/Source/Runtime/Core/Private/Misc/CoreDelegates.cpp
@@ -53,6 +53,8 @@ FCoreDelegates::FPakSigningKeysDelegate& FCoreDelegates::GetPakSigningKeysDelega
 #endif	//WITH_EDITOR
 FSimpleMulticastDelegate FCoreDelegates::OnShutdownAfterError;
 FSimpleMulticastDelegate FCoreDelegates::OnInit;
+FSimpleMulticastDelegate FCoreDelegates::OnPostEngineInit;
+FSimpleMulticastDelegate FCoreDelegates::OnFEngineLoopInitComplete;
 FSimpleMulticastDelegate FCoreDelegates::OnExit;
 FSimpleMulticastDelegate FCoreDelegates::OnPreExit;
 FSimpleMulticastDelegate FCoreDelegates::ColorPickerChanged;
@@ -109,7 +111,6 @@ FCoreDelegates::FOnAsyncLoadingFlush FCoreDelegates::OnAsyncLoadingFlush;
 FCoreDelegates::FOnAsyncLoadPackage FCoreDelegates::OnAsyncLoadPackage;
 FCoreDelegates::FRenderingThreadChanged FCoreDelegates::PostRenderingThreadCreated;
 FCoreDelegates::FRenderingThreadChanged FCoreDelegates::PreRenderingThreadDestroyed;
-FSimpleMulticastDelegate FCoreDelegates::OnFEngineLoopInitComplete;
 FCoreDelegates::FImageIntegrityChanged  FCoreDelegates::OnImageIntegrityChanged;
 
 FCoreDelegates::FApplicationReceivedOnScreenOrientationChangedNotificationDelegate FCoreDelegates::ApplicationReceivedScreenOrientationChangedNotificationDelegate;
diff --git a/Engine/Source/Runtime/Core/Private/Misc/CoreGlobals.cpp b/Engine/Source/Runtime/Core/Private/Misc/CoreGlobals.cpp
--- a/Engine/Source/Runtime/Core/Private/Misc/CoreGlobals.cpp
+++ b/Engine/Source/Runtime/Core/Private/Misc/CoreGlobals.cpp
@@ -96,9 +96,6 @@ bool GCompilingBlueprint = false;
 /** True if we're reconstructing blueprint instances. Should never be true on cooked builds */
 bool GIsReconstructingBlueprintInstances = false;
 
-/** Force blueprints to not compile on load */
-bool GForceDisableBlueprintCompileOnLoad = false;
-
 /** True if actors and objects are being re-instanced. */
 bool GIsReinstancing = false;
 
diff --git a/Engine/Source/Runtime/Core/Private/Windows/WindowsAsyncIO.h b/Engine/Source/Runtime/Core/Private/Windows/WindowsAsyncIO.h
--- a/Engine/Source/Runtime/Core/Private/Windows/WindowsAsyncIO.h
+++ b/Engine/Source/Runtime/Core/Private/Windows/WindowsAsyncIO.h
@@ -144,7 +144,7 @@ public:
 		}
 #endif
 		check(AlignedOffset <= Offset);
-		int64 BytesRead;
+		int64 BytesRead = 0L;
 		if (!GetOverlappedResult(FileHandle, &OverlappedIO, (LPDWORD)&BytesRead, TRUE))
 		{
 			uint32 ErrorCode = GetLastError();
diff --git a/Engine/Source/Runtime/Core/Public/Containers/Array.h b/Engine/Source/Runtime/Core/Public/Containers/Array.h
--- a/Engine/Source/Runtime/Core/Public/Containers/Array.h
+++ b/Engine/Source/Runtime/Core/Public/Containers/Array.h
@@ -28,6 +28,7 @@
  * - A method IndexType Num() const that returns the number of items in the container.
  * - A method bool IsValidIndex(IndexType index) which returns whether a given index is valid in the container.
  * - A method T& operator\[\](IndexType index) which returns a reference to a contained object by index.
+ * - A method void RemoveAt(IndexType index) which removes the element at index
  */
 template< typename ContainerType, typename ElementType, typename IndexType>
 class TIndexedContainerIterator
@@ -89,29 +90,21 @@ public:
 		return Tmp -= Offset;
 	}
 
-	/** @name Element access */
-	//@{
 	ElementType& operator* () const
 	{
 		return Container[ Index ];
 	}
 
-	ElementType* operator-> () const
+	ElementType* operator->() const
 	{
 		return &Container[ Index ];
 	}
-	//@}
 
 	/** conversion to "bool" returning true if the iterator has not reached the last element. */
 	FORCEINLINE explicit operator bool() const
 	{
 		return Container.IsValidIndex(Index);
 	}
-	/** inverse of the "bool" operator */
-	FORCEINLINE bool operator !() const 
-	{
-		return !(bool)*this;
-	}
 
 	/** Returns an index to the current element. */
 	IndexType GetIndex() const
@@ -125,6 +118,19 @@ public:
 		Index = 0;
 	}
 
+	/** Sets iterator to the last element. */
+	void SetToEnd()
+	{
+		Index = Container.Num();
+	}
+	
+	/** Removes current element in array. This invalidates the current iterator value and it must be incremented */
+	void RemoveCurrent()
+	{
+		Container.RemoveAt(Index);
+		Index--;
+	}
+
 	FORCEINLINE friend bool operator==(const TIndexedContainerIterator& Lhs, const TIndexedContainerIterator& Rhs) { return &Lhs.Container == &Rhs.Container && Lhs.Index == Rhs.Index; }
 	FORCEINLINE friend bool operator!=(const TIndexedContainerIterator& Lhs, const TIndexedContainerIterator& Rhs) { return &Lhs.Container != &Rhs.Container || Lhs.Index != Rhs.Index; }
 
diff --git a/Engine/Source/Runtime/Core/Public/Containers/ContainersFwd.h b/Engine/Source/Runtime/Core/Public/Containers/ContainersFwd.h
--- a/Engine/Source/Runtime/Core/Public/Containers/ContainersFwd.h
+++ b/Engine/Source/Runtime/Core/Public/Containers/ContainersFwd.h
@@ -12,3 +12,5 @@ template<typename T> class TTransArray;
 template<typename KeyType, typename ValueType, bool bInAllowDuplicateKeys> struct TDefaultMapHashableKeyFuncs;
 template<typename KeyType, typename ValueType, typename SetAllocator = FDefaultSetAllocator, typename KeyFuncs = TDefaultMapHashableKeyFuncs<KeyType, ValueType, false> > class TMap;
 template<typename KeyType, typename ValueType, typename SetAllocator = FDefaultSetAllocator, typename KeyFuncs = TDefaultMapHashableKeyFuncs<KeyType, ValueType, true > > class TMultiMap;
+template <typename T = void > struct TLess;
+template<typename KeyType, typename ValueType, typename ArrayAllocator = FDefaultAllocator, typename SortPredicate = TLess<KeyType> > class TSortedMap;
