Fixed warning about missing asset registries.

#tests ran locally

diff --git a/Engine/Source/Runtime/AssetRegistry/Private/AssetRegistry.cpp b/Engine/Source/Runtime/AssetRegistry/Private/AssetRegistry.cpp
--- a/Engine/Source/Runtime/AssetRegistry/Private/AssetRegistry.cpp
+++ b/Engine/Source/Runtime/AssetRegistry/Private/AssetRegistry.cpp
@@ -112,15 +112,16 @@ UAssetRegistryImpl::UAssetRegistryImpl(const FObjectInitializer& ObjectInitializ
 		for ( TSharedRef<IPlugin> ContentPlugin : ContentPlugins )
 		{
 			if ( ContentPlugin->CanContainContent() )
-		{
-				if (FFileHelper::LoadFileToArray(SerializedAssetData, *(ContentPlugin->GetBaseDir() / TEXT("AssetRegistry.bin"))))
 			{
-				SerializedAssetData.Seek(0);
-				Serialize(SerializedAssetData);
+				FString PluginAssetRegistry = ContentPlugin->GetBaseDir() / TEXT("AssetRegistry.bin");
+				if (IFileManager::Get().FileExists(*PluginAssetRegistry) && FFileHelper::LoadFileToArray(SerializedAssetData, *PluginAssetRegistry))
+				{
+					SerializedAssetData.Seek(0);
+					Serialize(SerializedAssetData);
+				}
 			}
 		}
 	}
-	}
 
 	// Report startup time. This does not include DirectoryWatcher startup time.
 	UE_LOG(LogAssetRegistry, Log, TEXT( "FAssetRegistry took %0.4f seconds to start up" ), FPlatformTime::Seconds() - StartupStartTime );
