Added IAssetRegistry::GetAssetsByTags to optimize tag look-up when you can have several potential values

diff --git a/Engine/Source/Editor/MovieSceneTools/Private/TrackEditors/CameraShakeTrackEditor.cpp b/Engine/Source/Editor/MovieSceneTools/Private/TrackEditors/CameraShakeTrackEditor.cpp
--- a/Engine/Source/Editor/MovieSceneTools/Private/TrackEditors/CameraShakeTrackEditor.cpp
+++ b/Engine/Source/Editor/MovieSceneTools/Private/TrackEditors/CameraShakeTrackEditor.cpp
@@ -140,10 +140,7 @@ void FCameraShakeTrackEditor::BuildObjectBindingTrackMenu(FMenuBuilder& MenuBuil
 		{
 			FARFilter Filter;
 			Filter.ClassNames.Add(UBlueprint::StaticClass()->GetFName());
-
-			TMultiMap<FName, FString> TagsAndValuesFilter;
-			TagsAndValuesFilter.Add(FBlueprintTags::ParentClassPath, CameraShakeClassPath);
-			Filter.TagsAndValues = TagsAndValuesFilter;
+			Filter.TagsAndValues.Add(FBlueprintTags::ParentClassPath, CameraShakeClassPath);
 
 			AssetRegistryModule.Get().GetAssets(Filter, AssetDataList);
 		}
diff --git a/Engine/Source/Editor/UnrealEd/Private/Fbx/FbxSkeletalMeshImport.cpp b/Engine/Source/Editor/UnrealEd/Private/Fbx/FbxSkeletalMeshImport.cpp
--- a/Engine/Source/Editor/UnrealEd/Private/Fbx/FbxSkeletalMeshImport.cpp
+++ b/Engine/Source/Editor/UnrealEd/Private/Fbx/FbxSkeletalMeshImport.cpp
@@ -1694,8 +1694,7 @@ USkeletalMesh* UnFbx::FFbxImporter::ImportSkeletalMesh(FImportSkeletalMeshArgs &
 						
 						FARFilter ARFilter;
 						ARFilter.ClassNames.Add(*USkeletalMesh::StaticClass()->GetName());
-						FString& Value = ARFilter.TagsAndValues.Add(TEXT("Skeleton"));
-						Value = FAssetData(Skeleton).GetExportTextName();
+						ARFilter.TagsAndValues.Add(TEXT("Skeleton"), FAssetData(Skeleton).GetExportTextName());
 
 						IAssetRegistry& AssetRegistry = AssetRegistryModule.Get();
 						if (AssetRegistry.GetAssets(ARFilter, SkeletalMeshAssetData))
diff --git a/Engine/Source/Runtime/AssetRegistry/Private/AssetRegistry.cpp b/Engine/Source/Runtime/AssetRegistry/Private/AssetRegistry.cpp
--- a/Engine/Source/Runtime/AssetRegistry/Private/AssetRegistry.cpp
+++ b/Engine/Source/Runtime/AssetRegistry/Private/AssetRegistry.cpp
@@ -505,10 +505,23 @@ bool UAssetRegistryImpl::GetAssetsByClass(FName ClassName, TArray<FAssetData>& O
 	return GetAssets(Filter, OutAssetData);
 }
 
+bool UAssetRegistryImpl::GetAssetsByTags(const TArray<FName>& AssetTags, TArray<FAssetData>& OutAssetData) const
+{
+	FARFilter Filter;
+	for (const FName AssetTag : AssetTags)
+	{
+		Filter.TagsAndValues.Add(AssetTag);
+	}
+	return GetAssets(Filter, OutAssetData);
+}
+
 bool UAssetRegistryImpl::GetAssetsByTagValues(const TMultiMap<FName, FString>& AssetTagsAndValues, TArray<FAssetData>& OutAssetData) const
 {
 	FARFilter Filter;
-	Filter.TagsAndValues = AssetTagsAndValues;
+	for (const auto& AssetTagsAndValue : AssetTagsAndValues)
+	{
+		Filter.TagsAndValues.Add(AssetTagsAndValue.Key, AssetTagsAndValue.Value);
+	}
 	return GetAssets(Filter, OutAssetData);
 }
 
@@ -586,13 +599,13 @@ bool UAssetRegistryImpl::GetAssets(const FARFilter& InFilter, TArray<FAssetData>
 					for (auto FilterTagIt = Filter.TagsAndValues.CreateConstIterator(); FilterTagIt; ++FilterTagIt)
 					{
 						const FName Tag = FilterTagIt.Key();
-						const FString& Value = FilterTagIt.Value();
+						const TOptional<FString>& Value = FilterTagIt.Value();
 
 						for (UObject::FAssetRegistryTag& AssetRegistryTag : ObjectTags)
 						{
 							if (AssetRegistryTag.Name == Tag)
 							{
-								if (AssetRegistryTag.Value == Value)
+								if (!Value.IsSet() || AssetRegistryTag.Value == Value.GetValue())
 								{
 									bMatch = true;
 								}
diff --git a/Engine/Source/Runtime/AssetRegistry/Private/AssetRegistry.h b/Engine/Source/Runtime/AssetRegistry/Private/AssetRegistry.h
--- a/Engine/Source/Runtime/AssetRegistry/Private/AssetRegistry.h
+++ b/Engine/Source/Runtime/AssetRegistry/Private/AssetRegistry.h
@@ -37,6 +37,7 @@ public:
 	virtual bool GetAssetsByPackageName(FName PackageName, TArray<FAssetData>& OutAssetData, bool bIncludeOnlyOnDiskAssets = false) const override;
 	virtual bool GetAssetsByPath(FName PackagePath, TArray<FAssetData>& OutAssetData, bool bRecursive = false, bool bIncludeOnlyOnDiskAssets = false) const override;
 	virtual bool GetAssetsByClass(FName ClassName, TArray<FAssetData>& OutAssetData, bool bSearchSubClasses = false) const override;
+	virtual bool GetAssetsByTags(const TArray<FName>& AssetTags, TArray<FAssetData>& OutAssetData) const override;
 	virtual bool GetAssetsByTagValues(const TMultiMap<FName, FString>& AssetTagsAndValues, TArray<FAssetData>& OutAssetData) const override;
 	virtual bool GetAssets(const FARFilter& Filter, TArray<FAssetData>& OutAssetData) const override;
 	virtual FAssetData GetAssetByObjectPath( const FName ObjectPath, bool bIncludeOnlyOnDiskAssets = false ) const override;
diff --git a/Engine/Source/Runtime/AssetRegistry/Private/AssetRegistryState.cpp b/Engine/Source/Runtime/AssetRegistry/Private/AssetRegistryState.cpp
--- a/Engine/Source/Runtime/AssetRegistry/Private/AssetRegistryState.cpp
+++ b/Engine/Source/Runtime/AssetRegistry/Private/AssetRegistryState.cpp
@@ -382,7 +382,7 @@ bool FAssetRegistryState::GetAssets(const FARFilter& Filter, const TSet<FName>&
 		for (auto FilterTagIt = Filter.TagsAndValues.CreateConstIterator(); FilterTagIt; ++FilterTagIt)
 		{
 			const FName Tag = FilterTagIt.Key();
-			const FString& Value = FilterTagIt.Value();
+			const TOptional<FString>& Value = FilterTagIt.Value();
 
 			const TArray<FAssetData*>* TagAssets = CachedAssetsByTag.Find(Tag);
 
@@ -393,7 +393,7 @@ bool FAssetRegistryState::GetAssets(const FARFilter& Filter, const TSet<FName>&
 					if (AssetData != nullptr)
 					{
 						const FString* TagValue = AssetData->TagsAndValues.Find(Tag);
-						if (TagValue != nullptr && *TagValue == Value)
+						if (TagValue != nullptr && (!Value.IsSet() || *TagValue == Value.GetValue()))
 						{
 							TagAndValuesFilter.Add(AssetData);
 						}
diff --git a/Engine/Source/Runtime/AssetRegistry/Public/ARFilter.h b/Engine/Source/Runtime/AssetRegistry/Public/ARFilter.h
--- a/Engine/Source/Runtime/AssetRegistry/Public/ARFilter.h
+++ b/Engine/Source/Runtime/AssetRegistry/Public/ARFilter.h
@@ -32,7 +32,7 @@ struct FARFilter
 	TArray<FName> ClassNames;
 
 	/** The filter component for properties marked with the AssetRegistrySearchable flag */
-	TMultiMap<FName, FString> TagsAndValues;
+	TMultiMap<FName, TOptional<FString>> TagsAndValues;
 
 	/** Only if bRecursiveClasses is true, the results will exclude classes (and subclasses) in this list */
 	UPROPERTY(transient, BlueprintReadWrite, Category=AssetRegistry)
diff --git a/Engine/Source/Runtime/AssetRegistry/Public/IAssetRegistry.h b/Engine/Source/Runtime/AssetRegistry/Public/IAssetRegistry.h
--- a/Engine/Source/Runtime/AssetRegistry/Public/IAssetRegistry.h
+++ b/Engine/Source/Runtime/AssetRegistry/Public/IAssetRegistry.h
@@ -103,11 +103,19 @@ public:
 	UFUNCTION(BlueprintCallable, Category = "AssetRegistry")
 	virtual bool GetAssetsByClass(FName ClassName, TArray<FAssetData>& OutAssetData, bool bSearchSubClasses = false) const = 0;
 
+	/**
+	 * Gets asset data for all assets with the supplied tags, regardless of their value
+	 *
+	 * @param AssetTags the tags associated with the assets requested
+	 * @param OutAssetData the list of assets with any of the given tags
+	 */
+	virtual bool GetAssetsByTags(const TArray<FName>& AssetTags, TArray<FAssetData>& OutAssetData) const = 0;
+
 	/**
 	 * Gets asset data for all assets with the supplied tags and values
 	 *
 	 * @param AssetTagsAndValues the tags and values associated with the assets requested
-	 * @param OutAssetData the list of assets in this path
+	 * @param OutAssetData the list of assets with any of the given tags and values
 	 */
 	virtual bool GetAssetsByTagValues(const TMultiMap<FName, FString>& AssetTagsAndValues, TArray<FAssetData>& OutAssetData) const = 0;
 
