Bug fixes in AssetRegistryState::RemoveAssetData().
Only remove dependencies and package data when there are no other known assets in the package.
Return true when data is removed, the function was previously always returning false.

diff --git a/Engine/Source/Runtime/AssetRegistry/Private/AssetRegistryState.cpp b/Engine/Source/Runtime/AssetRegistry/Private/AssetRegistryState.cpp
--- a/Engine/Source/Runtime/AssetRegistry/Private/AssetRegistryState.cpp
+++ b/Engine/Source/Runtime/AssetRegistry/Private/AssetRegistryState.cpp
@@ -1271,15 +1271,19 @@ bool FAssetRegistryState::RemoveAssetData(FAssetData* AssetData, bool bRemoveDep
 			OldTagAssets->RemoveSingleSwap(AssetData);
 		}
 
-		// We need to update the cached dependencies references cache so that they know we no
-		// longer exist and so don't reference them.
-		if (bRemoveDependencyData)
+		// Only remove dependencies and package data if there are no other known assets in the package
+		if (OldPackageAssets->Num() == 0)
 		{
-			RemoveDependsNode(AssetData->PackageName);
-		}
+			// We need to update the cached dependencies references cache so that they know we no
+			// longer exist and so don't reference them.
+			if (bRemoveDependencyData)
+			{
+				RemoveDependsNode(AssetData->PackageName);
+			}
 
-		// Remove the package data as well
-		RemovePackageData(AssetData->PackageName);
+			// Remove the package data as well
+			RemovePackageData(AssetData->PackageName);
+		}
 
 		// if the assets were preallocated in a block, we can't delete them one at a time, only the whole chunk in the destructor
 		if (PreallocatedAssetDataBuffers.Num() == 0)
@@ -1287,6 +1291,7 @@ bool FAssetRegistryState::RemoveAssetData(FAssetData* AssetData, bool bRemoveDep
 			delete AssetData;
 		}
 		NumAssets--;
+		bRemoved = true;
 	}
 
 	return bRemoved;
