Re-enable FAssetRegistryState::InitializeFromExistingAndPrune() cooker optimization.
In PruneAssetData(), prune dependency nodes for packages without package data (cleaning up deleted assets from the runtime asset registry).
Fix bugs in InitializeFromExistingAndPrune() so that it produces identical results to calling InitializeFromExisting() + PruneAssetData().
Refactor InitializeFromExistingAndPrune() and PruneAssetData() to make the implementations more similar.
Improve Dump() output.

Results:
SavingAssetRegistry scoped timer down from 540 sec to 70 sec (a 8x speedup) when cooking WindowsClient with ~200 000 packages and saving 25 asset registry files.

diff --git a/Engine/Source/Editor/UnrealEd/Private/Commandlets/AssetRegistryGenerator.cpp b/Engine/Source/Editor/UnrealEd/Private/Commandlets/AssetRegistryGenerator.cpp
--- a/Engine/Source/Editor/UnrealEd/Private/Commandlets/AssetRegistryGenerator.cpp
+++ b/Engine/Source/Editor/UnrealEd/Private/Commandlets/AssetRegistryGenerator.cpp
@@ -1018,10 +1018,7 @@ bool FAssetRegistryGenerator::SaveAssetRegistry(const FString& SandboxPath, bool
 		{
 			// Prune out the development only packages, and any assets that belong in a different chunk asset registry
 			FAssetRegistryState NewState;
-			// New InitializeFromExistingAndPrune temporarily disabled until bugs have been fixed
-			// NewState.InitializeFromExistingAndPrune(State, CookedPackages, TSet<FName>(), ChunkBucketElement.Value, SaveOptions);
-			NewState.InitializeFromExisting(State, SaveOptions);
-			NewState.PruneAssetData(CookedPackages, TSet<FName>(), ChunkBucketElement.Value, SaveOptions);
+			NewState.InitializeFromExistingAndPrune(State, CookedPackages, TSet<FName>(), ChunkBucketElement.Value, SaveOptions);
 
 			InjectEncryptionData(NewState);
 
diff --git a/Engine/Source/Runtime/AssetRegistry/Private/AssetRegistryState.cpp b/Engine/Source/Runtime/AssetRegistry/Private/AssetRegistryState.cpp
--- a/Engine/Source/Runtime/AssetRegistry/Private/AssetRegistryState.cpp
+++ b/Engine/Source/Runtime/AssetRegistry/Private/AssetRegistryState.cpp
@@ -249,89 +249,126 @@ void FAssetRegistryState::FilterTags(const FAssetDataTagMapSharedView& InTagsAnd
 void FAssetRegistryState::InitializeFromExistingAndPrune(const FAssetRegistryState & ExistingState, const TSet<FName>& RequiredPackages, const TSet<FName>& RemovePackages, const TSet<int32> ChunksToKeep, const FAssetRegistrySerializationOptions & Options)
 {
 	const bool bIsFilteredByChunkId = ChunksToKeep.Num() != 0;
-	const bool bIsFilteredByRequiredPackage = RequiredPackages.Num() != 0;
+	const bool bIsFilteredByRequiredPackages = RequiredPackages.Num() != 0;
 	const bool bIsFilteredByRemovedPackages = RemovePackages.Num() != 0;
-	const bool bIsFilteredByNoTags = Options.bFilterAssetDataWithNoTags;
+
+	TSet<FName> RequiredDependNodePackages;
 
 	// Duplicate asset data entries
 	for (const TPair<FName, FAssetData*>& AssetPair : ExistingState.CachedAssetsByObjectPath)
 	{
 		const FAssetData* AssetData = AssetPair.Value;
-		if (AssetData)
+
+		bool bRemoveAssetData = false;
+		bool bRemoveDependencyData = true;
+
+		if (bIsFilteredByChunkId &&
+			!AssetData->ChunkIDs.ContainsByPredicate([&](int32 ChunkId) { return ChunksToKeep.Contains(ChunkId); }))
+		{
+			bRemoveAssetData = true;
+		}
+		else if (bIsFilteredByRequiredPackages && !RequiredPackages.Contains(AssetData->PackageName))
+		{
+			bRemoveAssetData = true;
+		}
+		else if (bIsFilteredByRemovedPackages && RemovePackages.Contains(AssetData->PackageName))
 		{
-			bool bDiscard = false;
-			bDiscard |= bIsFilteredByChunkId && !AssetData->ChunkIDs.ContainsByPredicate([&](int32 ChunkId) { return ChunksToKeep.Contains(ChunkId); });
-			bDiscard |= bIsFilteredByRequiredPackage && !RequiredPackages.Contains(AssetData->PackageName);
-			bDiscard |= bIsFilteredByRemovedPackages && RemovePackages.Contains(AssetData->PackageName);
-			bDiscard |= bIsFilteredByNoTags && AssetData->TagsAndValues.Num() == 0 && !FPackageName::IsLocalizedPackage(AssetData->PackageName.ToString());
+			bRemoveAssetData = true;
+		}
+		else if (Options.bFilterAssetDataWithNoTags && AssetData->TagsAndValues.Num() == 0 &&
+			!FPackageName::IsLocalizedPackage(AssetData->PackageName.ToString()))
+		{
+			bRemoveAssetData = true;
+			bRemoveDependencyData = Options.bFilterDependenciesWithNoTags;
+		}
 
-			if (bDiscard)
+		if (bRemoveAssetData)
+		{
+			if (!bRemoveDependencyData)
 			{
-				continue;
+				RequiredDependNodePackages.Add(AssetData->PackageName);
 			}
+			continue;
+		}
 
-			FAssetDataTagMap NewTagsAndValues;
-			FAssetRegistryState::FilterTags(AssetData->TagsAndValues, NewTagsAndValues, Options.CookFilterlistTagsByClass.Find(AssetData->AssetClass), Options);
+		FAssetDataTagMap NewTagsAndValues;
+		FAssetRegistryState::FilterTags(AssetData->TagsAndValues, NewTagsAndValues, Options.CookFilterlistTagsByClass.Find(AssetData->AssetClass), Options);
 
-			FAssetData* NewAssetData = new FAssetData(AssetData->PackageName, AssetData->PackagePath, AssetData->AssetName,
-				AssetData->AssetClass, NewTagsAndValues, AssetData->ChunkIDs, AssetData->PackageFlags);
-			// Add asset to new state
-			AddAssetData(NewAssetData);
-			// Also add it's package if it doesn't exist
-			CreateOrGetAssetPackageData(NewAssetData->PackageName);
-		}
+		FAssetData* NewAssetData = new FAssetData(AssetData->PackageName, AssetData->PackagePath, AssetData->AssetName,
+			AssetData->AssetClass, NewTagsAndValues, AssetData->ChunkIDs, AssetData->PackageFlags);
+		// Add asset to new state
+		AddAssetData(NewAssetData);
 	}
 
-	TSet<FAssetIdentifier> ScriptPackages;
-
-	// Build list of script packages for depend nodes copying
-	for (const FName& Name : RequiredPackages)
+	// Create package data for all script and required packages
+	for (const TPair<FName, FAssetPackageData*>& Pair : ExistingState.CachedPackageData)
 	{
-		bool bIsScriptPackage = FPackageName::IsScriptPackage(Name.ToString());
-		const FAssetPackageData* OldData = ExistingState.GetAssetPackageData(Name);
-		if (bIsScriptPackage) 
+		if (Pair.Value)
 		{
-			if (OldData)
+			// Only add if also in asset data map, or script package
+			if (CachedAssetsByPackageName.Find(Pair.Key) ||
+				FPackageName::IsScriptPackage(Pair.Key.ToString()))
 			{
-				ScriptPackages.Add(Name);
-
-				FAssetPackageData* NewData = CreateOrGetAssetPackageData(Name);
-				*NewData = *OldData;
+				FAssetPackageData* NewData = CreateOrGetAssetPackageData(Pair.Key);
+				*NewData = *Pair.Value;
 			}
 		}
-		else if (CachedAssetsByPackageName.Find(Name))
-		{
-			FAssetPackageData* NewData = CreateOrGetAssetPackageData(Name);
-			*NewData = *OldData;
-		}
 	}
 
-
-	// Duplicate dependency nodes
+	// Find valid dependency nodes for all script and required packages
+	TSet<FDependsNode*> ValidDependsNodes;
+	ValidDependsNodes.Reserve(ExistingState.CachedDependsNodes.Num());
 	for (const TPair<FAssetIdentifier, FDependsNode*>& Pair : ExistingState.CachedDependsNodes)
 	{
-		FDependsNode* OldNode = Pair.Value;
+		FDependsNode* Node = Pair.Value;
+		const FAssetIdentifier& Id = Node->GetIdentifier();
+		bool bRemoveDependsNode = false;
 
-		// Only add if this node is already part of the filtered asset set
-		if (!CachedAssetsByPackageName.Contains(OldNode->GetIdentifier().PackageName))
+		if (Options.bFilterSearchableNames && Id.IsValue())
 		{
-			continue;
+			bRemoveDependsNode = true;
+		}
+		else if (Id.IsPackage() &&
+			!CachedAssetsByPackageName.Contains(Id.PackageName) &&
+			!RequiredDependNodePackages.Contains(Id.PackageName) &&
+			!FPackageName::IsScriptPackage(Id.PackageName.ToString()))
+		{
+			bRemoveDependsNode = true;
+		}
+
+		if (!bRemoveDependsNode)
+		{
+			ValidDependsNodes.Add(Node);
 		}
+	}
 
-		FDependsNode* NewNode = CreateOrFindDependsNode(Pair.Key);
+	// Duplicate dependency nodes
+	for (FDependsNode* OldNode : ValidDependsNodes)
+	{
+		FDependsNode* NewNode = CreateOrFindDependsNode(OldNode->GetIdentifier());
 		NewNode->Reserve(OldNode);
 	
 		OldNode->IterateOverDependencies([&, OldNode, NewNode](FDependsNode* InDependency, EAssetRegistryDependencyType::Type InDependencyType) {
-			const FAssetIdentifier& Identifier = InDependency->GetIdentifier();
-			if (CachedAssetsByPackageName.Find(Identifier.PackageName) || ScriptPackages.Contains(Identifier))
+			if (ValidDependsNodes.Contains(InDependency))
 			{
 				// Only add link if it's part of the filtered asset set
-				FDependsNode* NewDependency = CreateOrFindDependsNode(Identifier);
+				FDependsNode* NewDependency = CreateOrFindDependsNode(InDependency->GetIdentifier());
 				NewNode->AddDependency(NewDependency, InDependencyType, true);
 				NewDependency->AddReferencer(NewNode);
 			}
 		});
 	}
+
+	// Remove any orphaned depends nodes. This will leave cycles in but those might represent useful data
+	TArray<FDependsNode*> AllDependsNodes;
+	CachedDependsNodes.GenerateValueArray(AllDependsNodes);
+	for (FDependsNode* DependsNode : AllDependsNodes)
+	{
+		if (DependsNode->GetConnectionCount() == 0)
+		{
+			RemoveDependsNode(DependsNode->GetIdentifier());
+		}
+	}
 }
 
 void FAssetRegistryState::InitializeFromExisting(const TMap<FName, FAssetData*>& AssetDataMap, const TMap<FAssetIdentifier, FDependsNode*>& DependsNodeMap, const TMap<FName, FAssetPackageData*>& AssetPackageDataMap, const FAssetRegistrySerializationOptions& Options, EInitializationMode InInitializationMode)
@@ -433,57 +470,77 @@ void FAssetRegistryState::PruneAssetData(const TSet<FName>& RequiredPackages, co
 
 void FAssetRegistryState::PruneAssetData(const TSet<FName>& RequiredPackages, const TSet<FName>& RemovePackages, const TSet<int32> ChunksToKeep, const FAssetRegistrySerializationOptions& Options)
 {
+	const bool bIsFilteredByChunkId = ChunksToKeep.Num() != 0;
+	const bool bIsFilteredByRequiredPackages = RequiredPackages.Num() != 0;
+	const bool bIsFilteredByRemovedPackages = RemovePackages.Num() != 0;
+
+	TSet<FName> RequiredDependNodePackages;
+
 	// Generate list up front as the maps will get cleaned up
 	TArray<FAssetData*> AllAssetData;
 	CachedAssetsByObjectPath.GenerateValueArray(AllAssetData);
 
 	for (FAssetData* AssetData : AllAssetData)
 	{
-		bool bFilteredByChunkID = ChunksToKeep.Num() != 0;
+		bool bRemoveAssetData = false;
+		bool bRemoveDependencyData = true;
 
-		if (bFilteredByChunkID)
+		if (bIsFilteredByChunkId &&
+			!AssetData->ChunkIDs.ContainsByPredicate([&](int32 ChunkId) { return ChunksToKeep.Contains(ChunkId); }))
 		{
-			for (int32 ChunkToKeep : ChunksToKeep)
-			{
-				if (AssetData->ChunkIDs.Contains(ChunkToKeep ))
-				{
-					bFilteredByChunkID = false;
-					break;
-				}
-			}
+			bRemoveAssetData = true;
 		}
-
-		if (bFilteredByChunkID || (RequiredPackages.Num() > 0 && !RequiredPackages.Contains(AssetData->PackageName)))
+		else if (bIsFilteredByRequiredPackages && !RequiredPackages.Contains(AssetData->PackageName))
 		{
-			RemoveAssetData(AssetData);
+			bRemoveAssetData = true;
 		}
-		else if (RemovePackages.Contains(AssetData->PackageName))
+		else if (bIsFilteredByRemovedPackages && RemovePackages.Contains(AssetData->PackageName))
 		{
-			RemoveAssetData(AssetData);
+			bRemoveAssetData = true;
 		}
-		else if (Options.bFilterAssetDataWithNoTags && AssetData->TagsAndValues.Num() == 0 && !FPackageName::IsLocalizedPackage(AssetData->PackageName.ToString()))
+		else if (Options.bFilterAssetDataWithNoTags && AssetData->TagsAndValues.Num() == 0 &&
+			!FPackageName::IsLocalizedPackage(AssetData->PackageName.ToString()))
 		{
-			// Optionally remove dependency data
-			RemoveAssetData(AssetData, Options.bFilterDependenciesWithNoTags);
+			bRemoveAssetData = true;
+			bRemoveDependencyData = Options.bFilterDependenciesWithNoTags;
+		}
+
+		if (bRemoveAssetData)
+		{
+			if (!bRemoveDependencyData)
+			{
+				RequiredDependNodePackages.Add(AssetData->PackageName);
+			}
+			RemoveAssetData(AssetData, bRemoveDependencyData);
 		}
 	}
 
 	TArray<FDependsNode*> AllDependsNodes;
 	CachedDependsNodes.GenerateValueArray(AllDependsNodes);
 
-	if (Options.bFilterSearchableNames)
+	for (FDependsNode* DependsNode : AllDependsNodes)
 	{
-		// Remove searchable names if specified
-		for (FDependsNode* DependsNode : AllDependsNodes)
+		const FAssetIdentifier& Id = DependsNode->GetIdentifier();
+		bool bRemoveDependsNode = false;
+
+		if (Options.bFilterSearchableNames && Id.IsValue())
 		{
-			if (DependsNode->GetIdentifier().IsValue())
-			{
-				RemoveDependsNode(DependsNode->GetIdentifier());
-			}
+			bRemoveDependsNode = true;
+		}
+		else if (Id.IsPackage() &&
+			!CachedAssetsByPackageName.Contains(Id.PackageName) &&
+			!RequiredDependNodePackages.Contains(Id.PackageName) &&
+			!FPackageName::IsScriptPackage(Id.PackageName.ToString()))
+		{
+			bRemoveDependsNode = true;
+		}
+		
+		if (bRemoveDependsNode)
+		{
+			RemoveDependsNode(Id);
 		}
-
-		CachedDependsNodes.GenerateValueArray(AllDependsNodes);
 	}
+	CachedDependsNodes.GenerateValueArray(AllDependsNodes);
 
 	// Remove any orphaned depends nodes. This will leave cycles in but those might represent useful data
 	for (FDependsNode* DependsNode : AllDependsNodes)
@@ -1612,124 +1669,114 @@ bool FAssetRegistryState::IsFilterValid(const FARFilter& Filter, bool bAllowRecu
 
 void FAssetRegistryState::Dump(const TArray<FString>& Arguments, TArray<FString>& OutLines) const
 {
-	struct FAssetRefSorter
-	{
-		FORCEINLINE bool operator()(const FAssetData& A, const FAssetData& B) const
-		{
-			return (A.AssetName < B.AssetName);
-		}
-	};
+	OutLines.Reserve(14 + CachedAssetsByObjectPath.Num() * 5 + CachedDependsNodes.Num() + CachedPackageData.Num());
 
 	if (Arguments.Contains(TEXT("ObjectPath")))
 	{
-		OutLines.Add(FString::Printf(TEXT("CachedAssetsByObjectPath (%d): "), CachedAssetsByObjectPath.Num()));
+		OutLines.Add(FString::Printf(TEXT("--- Begin CachedAssetsByObjectPath ---")));
+
 		for (const TPair<FName, FAssetData*>& Pair : CachedAssetsByObjectPath)
 		{
-			OutLines.Add(FString::Printf(TEXT("	%s: %s"), *Pair.Key.ToString(), *Pair.Value->PackageName.ToString()));
+			OutLines.Add(FString::Printf(TEXT("	%s"), *Pair.Key.ToString()));
 		}
+
+		OutLines.Add(
+			FString::Printf(TEXT("--- End CachedAssetsByObjectPath : %d entries ---"), CachedAssetsByObjectPath.Num()));
 	}
 
-	if (Arguments.Contains(TEXT("PackageName")))
+	auto PrintAssetDataMap = [&OutLines](TCHAR* Name, const TMap<FName, TArray<FAssetData*>>& AssetMap)
 	{
-		OutLines.Add(FString::Printf(TEXT("CachedAssetsByPackageName (%d): "), CachedAssetsByPackageName.Num()));
-		for (const TPair<FName, TArray<FAssetData*>>& Pair : CachedAssetsByPackageName)
+		OutLines.Add(FString::Printf(TEXT("--- Begin %s ---"), Name));
+
+		TArray<FName> Keys;
+		AssetMap.GenerateKeyArray(Keys);
+		Keys.Sort();
+
+		TArray<FAssetData*> Items;
+		Items.Reserve(1024);
+
+		int32 ValidCount = 0;
+		for (const FName& Key : Keys)
 		{
-			if (Pair.Value.Num() == 0)
+			const TArray<FAssetData*> AssetArray = AssetMap.FindChecked(Key);
+			if (AssetArray.Num() == 0)
 			{
 				continue;
 			}
+			++ValidCount;
 
-			TArray<FAssetData*> Datas(Pair.Value);
-			Datas.HeapSort(FAssetRefSorter());
-			FString Tmp;
-			for (const FAssetData* data : Datas)
+			Items.Reset();
+			Items.Append(AssetArray);
+			Items.Sort([](const FAssetData& A, const FAssetData& B)
+				{ return A.ObjectPath.ToString() < B.ObjectPath.ToString(); }
+				);
+
+			OutLines.Add(FString::Printf(TEXT("	%s : %d item(s)"), *Key.ToString(), Items.Num()));
+			for (const FAssetData* Data : Items)
 			{
-				Tmp.Append(FString::Printf(TEXT("%s, "), *data->PackageName.ToString()));
+				OutLines.Add(FString::Printf(TEXT("	 %s"), *Data->ObjectPath.ToString()));
 			}
-			OutLines.Add(FString::Printf(TEXT("	%s: %s"), *Pair.Key.ToString(), *Tmp));
 		}
+
+		OutLines.Add(FString::Printf(TEXT("--- End %s : %d entries ---"), Name, ValidCount));
+	};
+
+	if (Arguments.Contains(TEXT("PackageName")))
+	{
+		PrintAssetDataMap(TEXT("CachedAssetsByPackageName"), CachedAssetsByPackageName);
 	}
 
 	if (Arguments.Contains(TEXT("Path")))
 	{
-		OutLines.Add(FString::Printf(TEXT("CachedAssetsByPath (%d): "), CachedAssetsByPath.Num()));
-		for (const TPair<FName, TArray<FAssetData*>>& Pair : CachedAssetsByPath)
-		{
-			if (Pair.Value.Num() == 0)
-			{
-				continue;
-			}
-
-			TArray<FAssetData*> Datas(Pair.Value);
-			Datas.HeapSort(FAssetRefSorter());
-			FString Tmp;
-			for (const FAssetData* data : Datas)
-			{
-				Tmp.Append(FString::Printf(TEXT("%s, "), *data->PackageName.ToString()));
-			}
-			OutLines.Add(FString::Printf(TEXT("	%s: %s"), *Pair.Key.ToString(), *Tmp));
-		}
+		PrintAssetDataMap(TEXT("CachedAssetsByPath"), CachedAssetsByPath);
 	}
 
 	if (Arguments.Contains(TEXT("Class")))
 	{
-		OutLines.Add(FString::Printf(TEXT("CachedAssetsByClass (%d): "), CachedAssetsByClass.Num()));
-		for (const TPair<FName, TArray<FAssetData*>>& Pair : CachedAssetsByClass)
-		{
-			if (Pair.Value.Num() == 0)
-			{
-				continue;
-			}
-
-			TArray<FAssetData*> Datas(Pair.Value);
-			Datas.HeapSort(FAssetRefSorter());
-			FString Tmp;
-			for (const FAssetData* data : Datas)
-			{
-				Tmp.Append(FString::Printf(TEXT("%s, "), *data->PackageName.ToString()));
-			}
-			OutLines.Add(FString::Printf(TEXT("	%s: %s"), *Pair.Key.ToString(), *Tmp));
-		}
+		PrintAssetDataMap(TEXT("CachedAssetsByClass"), CachedAssetsByClass);
 	}
 
 	if (Arguments.Contains(TEXT("Tag")))
 	{
-		OutLines.Add(FString::Printf(TEXT("CachedAssetsByTag (%d): "), CachedAssetsByTag.Num()));
-		for (const TPair<FName, TArray<FAssetData*>>& Pair : CachedAssetsByTag)
-		{
-			if (Pair.Value.Num() == 0)
-			{
-				continue;
-			}
-			TArray<FAssetData*> Datas(Pair.Value);
-			Datas.HeapSort(FAssetRefSorter());
-			FString Tmp;
-			for (const FAssetData* data : Datas)
-			{
-				Tmp.Append(FString::Printf(TEXT("%s, "), *data->PackageName.ToString()));
-			}
-			OutLines.Add(FString::Printf(TEXT("	%s: %s"), *Pair.Key.ToString(), *Tmp));
-		}
+		PrintAssetDataMap(TEXT("CachedAssetsByTag"), CachedAssetsByTag);
 	}
 
 	if (Arguments.Contains(TEXT("Dependencies")))
 	{
-		OutLines.Add(FString::Printf(TEXT("CachedDependsNodes (%d): "), CachedDependsNodes.Num()));
-		for (const TPair<FAssetIdentifier, FDependsNode*>& Pair : CachedDependsNodes)
+		OutLines.Add(FString::Printf(TEXT("--- Begin CachedDependsNodes ---")));
+
+		TArray<FDependsNode*> Nodes;
+		CachedDependsNodes.GenerateValueArray(Nodes);
+		Nodes.Sort([](const FDependsNode& A, const FDependsNode& B)
+			{ return A.GetIdentifier().ToString() < B.GetIdentifier().ToString(); }
+			);
+
+		for (const FDependsNode* Node : Nodes)
 		{
-			const FDependsNode* Node = Pair.Value;
-			OutLines.Add(FString::Printf(TEXT("	%s: %s (%d connection)"), *Pair.Key.PackageName.ToString(), *Node->GetIdentifier().PackageName.ToString(), Node->GetConnectionCount()));
+			OutLines.Add(
+				FString::Printf(TEXT("	%s : %d connection(s)"),
+					*Node->GetIdentifier().ToString(), Node->GetConnectionCount()));
 		}
+
+		OutLines.Add(FString::Printf(TEXT("--- End CachedDependsNodes : %d entries ---"), CachedDependsNodes.Num()));
 	}
 
 	if (Arguments.Contains(TEXT("PackageData")))
 	{
-		OutLines.Add(FString::Printf(TEXT("CachedPackageData (%d): "), CachedPackageData.Num()));
-		for (const TPair<FName, FAssetPackageData*>& Pair : CachedPackageData)
+		OutLines.Add(FString::Printf(TEXT("--- Begin CachedPackageData ---")));
+
+		TArray<FName> Keys;
+		CachedPackageData.GenerateKeyArray(Keys);
+		Keys.Sort();
+
+		for (const FName& Key : Keys)
 		{
-			const FAssetPackageData* PackageData = Pair.Value;
-			OutLines.Add(FString::Printf(TEXT("	%s: %s %d bytes"), *Pair.Key.ToString(), *PackageData->PackageGuid.ToString(), PackageData->DiskSize));
+			const FAssetPackageData* PackageData = CachedPackageData.FindChecked(Key);
+			OutLines.Add(FString::Printf(TEXT("	%s : %s : %d bytes"),
+					*Key.ToString(), *PackageData->PackageGuid.ToString(), PackageData->DiskSize));
 		}
+
+		OutLines.Add(FString::Printf(TEXT("--- End CachedPackageData : %d entries ---"), CachedPackageData.Num()));
 	}
 }
 
