Fix clang error for cl 5819128: no known conversion from 'const char16_t [20]' to 'TCHAR *' (aka 'char16_t *') for 1st argument

diff --git a/Engine/Source/Runtime/AssetRegistry/Private/AssetRegistryState.cpp b/Engine/Source/Runtime/AssetRegistry/Private/AssetRegistryState.cpp
--- a/Engine/Source/Runtime/AssetRegistry/Private/AssetRegistryState.cpp
+++ b/Engine/Source/Runtime/AssetRegistry/Private/AssetRegistryState.cpp
@@ -1684,9 +1684,9 @@ void FAssetRegistryState::Dump(const TArray<FString>& Arguments, TArray<FString>
 			FString::Printf(TEXT("--- End CachedAssetsByObjectPath : %d entries ---"), CachedAssetsByObjectPath.Num()));
 	}
 
-	auto PrintAssetDataMap = [&OutLines](TCHAR* Name, const TMap<FName, TArray<FAssetData*>>& AssetMap)
+	auto PrintAssetDataMap = [&OutLines](FString Name, const TMap<FName, TArray<FAssetData*>>& AssetMap)
 	{
-		OutLines.Add(FString::Printf(TEXT("--- Begin %s ---"), Name));
+		OutLines.Add(FString::Printf(TEXT("--- Begin %s ---"), *Name));
 
 		TArray<FName> Keys;
 		AssetMap.GenerateKeyArray(Keys);
@@ -1718,7 +1718,7 @@ void FAssetRegistryState::Dump(const TArray<FString>& Arguments, TArray<FString>
 			}
 		}
 
-		OutLines.Add(FString::Printf(TEXT("--- End %s : %d entries ---"), Name, ValidCount));
+		OutLines.Add(FString::Printf(TEXT("--- End %s : %d entries ---"), *Name, ValidCount));
 	};
 
 	if (Arguments.Contains(TEXT("PackageName")))
