Prevent crash after Force Deleting an Interface blueprint

diff --git a/Engine/Source/Editor/BlueprintGraph/Private/BlueprintActionDatabase.cpp b/Engine/Source/Editor/BlueprintGraph/Private/BlueprintActionDatabase.cpp
--- a/Engine/Source/Editor/BlueprintGraph/Private/BlueprintActionDatabase.cpp
+++ b/Engine/Source/Editor/BlueprintGraph/Private/BlueprintActionDatabase.cpp
@@ -561,6 +561,11 @@ static bool BlueprintActionDatabaseImpl::IsBlueprintInterfaceFunction(const UFun
 			for (int32 InterfaceIndex = 0; (InterfaceIndex < BpOuter->ImplementedInterfaces.Num()) && !bIsBpInterfaceFunc; ++InterfaceIndex)
 			{
 				FBPInterfaceDescription& InterfaceDesc = BpOuter->ImplementedInterfaces[InterfaceIndex];
+				if(!InterfaceDesc.Interface)
+				{
+					continue;
+				}
+
 				bIsBpInterfaceFunc = (InterfaceDesc.Interface->FindFunctionByName(FuncName) != nullptr);
 			}
 		}
diff --git a/Engine/Source/Runtime/AssetRegistry/Private/AssetRegistry.cpp b/Engine/Source/Runtime/AssetRegistry/Private/AssetRegistry.cpp
--- a/Engine/Source/Runtime/AssetRegistry/Private/AssetRegistry.cpp
+++ b/Engine/Source/Runtime/AssetRegistry/Private/AssetRegistry.cpp
@@ -2776,7 +2776,7 @@ void UAssetRegistryImpl::UpdateTemporaryCaches() const
 			for (int32 i = 0; i < Class->Interfaces.Num(); ++i)
 			{
 				UClass* InterfaceClass = Class->Interfaces[i].Class;
-				if (ensureMsgf(InterfaceClass, TEXT("Invalid inteface class (index %d) detected on '%s' - possibly deleted?"), i, *Class->GetName()))
+				if (InterfaceClass) // could be nulled out by ForceDelete of a blueprint interface
 				{
 					TSet<FName>& ChildClasses = MutableThis->TempReverseInheritanceMap.FindOrAdd(InterfaceClass->GetFName());
 					ChildClasses.Add(ClassName);
