Exposed Asset Registry GetDependencies and GetReferencers to scripting

diff --git a/Engine/Source/Runtime/AssetRegistry/Public/IAssetRegistry.h b/Engine/Source/Runtime/AssetRegistry/Public/IAssetRegistry.h
--- a/Engine/Source/Runtime/AssetRegistry/Public/IAssetRegistry.h
+++ b/Engine/Source/Runtime/AssetRegistry/Public/IAssetRegistry.h
@@ -54,6 +54,51 @@ namespace EAssetSetManagerFlags
 	};
 }
 
+USTRUCT(BlueprintType)
+struct FAssetRegistryDependencyOptions
+{
+	GENERATED_BODY()
+
+	void SetFromFlags(const EAssetRegistryDependencyType::Type InFlags)
+	{
+		bIncludeSoftPackageReferences = (InFlags & EAssetRegistryDependencyType::Soft);
+		bIncludeHardPackageReferences = (InFlags & EAssetRegistryDependencyType::Hard);
+		bIncludeSearchableNames = (InFlags & EAssetRegistryDependencyType::SearchableName);
+		bIncludeSoftManagementReferences = (InFlags & EAssetRegistryDependencyType::SoftManage);
+		bIncludeHardManagementReferences = (InFlags & EAssetRegistryDependencyType::HardManage);
+	}
+
+	EAssetRegistryDependencyType::Type GetAsFlags() const
+	{
+		uint32 Flags = EAssetRegistryDependencyType::None;
+		Flags |= (bIncludeSoftPackageReferences ? EAssetRegistryDependencyType::Soft : EAssetRegistryDependencyType::None);
+		Flags |= (bIncludeHardPackageReferences ? EAssetRegistryDependencyType::Hard : EAssetRegistryDependencyType::None);
+		Flags |= (bIncludeSearchableNames ? EAssetRegistryDependencyType::SearchableName : EAssetRegistryDependencyType::None);
+		Flags |= (bIncludeSoftManagementReferences ? EAssetRegistryDependencyType::SoftManage : EAssetRegistryDependencyType::None);
+		Flags |= (bIncludeHardManagementReferences ? EAssetRegistryDependencyType::HardManage : EAssetRegistryDependencyType::None);
+		return (EAssetRegistryDependencyType::Type)Flags;
+	}
+
+	/** Dependencies which don't need to be loaded for the object to be used (i.e. soft object paths) */
+	UPROPERTY(EditAnywhere, BlueprintReadWrite, Category = "AssetRegistry")
+	bool bIncludeSoftPackageReferences = true;
+
+	/** Dependencies which are required for correct usage of the source asset, and must be loaded at the same time */
+	UPROPERTY(EditAnywhere, BlueprintReadWrite, Category = "AssetRegistry")
+	bool bIncludeHardPackageReferences = true;
+
+	/** References to specific SearchableNames inside a package */
+	UPROPERTY(EditAnywhere, BlueprintReadWrite, Category = "AssetRegistry")
+	bool bIncludeSearchableNames = false;
+
+	/** Indirect management references, these are set through recursion for Primary Assets that manage packages or other primary assets */
+	UPROPERTY(EditAnywhere, BlueprintReadWrite, Category = "AssetRegistry")
+	bool bIncludeSoftManagementReferences = false;
+
+	/** Reference that says one object directly manages another object, set when Primary Assets manage things explicitly */
+	UPROPERTY(EditAnywhere, BlueprintReadWrite, Category = "AssetRegistry")
+	bool bIncludeHardManagementReferences = false;
+};
 
 UINTERFACE(MinimalApi, BlueprintType, meta = (CannotImplementInterfaceInBlueprint))
 class UAssetRegistry : public UInterface
@@ -167,6 +212,19 @@ public:
 	 */
 	virtual bool GetDependencies(FName PackageName, TArray<FName>& OutDependencies, EAssetRegistryDependencyType::Type InDependencyType = EAssetRegistryDependencyType::Packages) const = 0;
 
+	/**
+	 * Gets a list of paths to objects that are referenced by the supplied package. (On disk references ONLY)
+	 *
+	 * @param PackageName		the name of the package for which to gather dependencies
+	 * @param DependencyOptions	which kinds of dependencies to include in the output list
+	 * @param OutDependencies	a list of packages that are referenced by the package whose path is PackageName
+	 */
+	UFUNCTION(BlueprintCallable, BlueprintPure=false, Category = "AssetRegistry", meta=(DisplayName="Get Dependencies", ScriptName="GetDependencies"))
+	virtual bool K2_GetDependencies(FName PackageName, const FAssetRegistryDependencyOptions& DependencyOptions, TArray<FName>& OutDependencies) const
+	{
+		return GetDependencies(PackageName, OutDependencies, DependencyOptions.GetAsFlags());
+	}
+
 	/**
 	 * Gets a list of packages and searchable names that reference the supplied package or name. (On disk references ONLY)
 	 *
@@ -185,6 +243,19 @@ public:
 	 */
 	virtual bool GetReferencers(FName PackageName, TArray<FName>& OutReferencers, EAssetRegistryDependencyType::Type InReferenceType = EAssetRegistryDependencyType::Packages) const = 0;
 
+	/**
+	 * Gets a list of packages that reference the supplied package. (On disk references ONLY)
+	 *
+	 * @param PackageName		the name of the package for which to gather dependencies
+	 * @param ReferenceOptions	which kinds of references to include in the output list
+	 * @param OutReferencers	a list of packages that reference the package whose path is PackageName
+	 */
+	UFUNCTION(BlueprintCallable, BlueprintPure=false, Category = "AssetRegistry", meta=(DisplayName="Get Referencers", ScriptName="GetReferencers"))
+	virtual bool K2_GetReferencers(FName PackageName, const FAssetRegistryDependencyOptions& ReferenceOptions, TArray<FName>& OutReferencers) const
+	{
+		return GetReferencers(PackageName, OutReferencers, ReferenceOptions.GetAsFlags());
+	}
+
 	/** Finds Package data for a package name. This data is only updated on save and can only be accessed for valid packages */
 	virtual const FAssetPackageData* GetAssetPackageData(FName PackageName) const = 0;
 
diff --git a/Engine/Source/Runtime/CoreUObject/Public/Misc/AssetRegistryInterface.h b/Engine/Source/Runtime/CoreUObject/Public/Misc/AssetRegistryInterface.h
--- a/Engine/Source/Runtime/CoreUObject/Public/Misc/AssetRegistryInterface.h
+++ b/Engine/Source/Runtime/CoreUObject/Public/Misc/AssetRegistryInterface.h
@@ -24,6 +24,8 @@ namespace EAssetRegistryDependencyType
 
 		// Reference that says one object directly manages another object, set when Primary Assets manage things explicitly
 		HardManage = 0x10,
+
+		// Note: Also update FAssetRegistryDependencyOptions when adding more flags
 	};
 
 	static const Type None = (Type)(0);
@@ -33,7 +35,7 @@ namespace EAssetRegistryDependencyType
 }
 
 /**
-* HotReload module interface
+* Asset Registry module interface
 */
 class IAssetRegistryInterface : public IModuleInterface
 {
