Made the potentially expensive Get and Filter functions on the asset registry impure

diff --git a/Engine/Source/Runtime/AssetRegistry/Public/IAssetRegistry.h b/Engine/Source/Runtime/AssetRegistry/Public/IAssetRegistry.h
--- a/Engine/Source/Runtime/AssetRegistry/Public/IAssetRegistry.h
+++ b/Engine/Source/Runtime/AssetRegistry/Public/IAssetRegistry.h
@@ -116,7 +116,7 @@ public:
 	 * @param PackagePath the path to query asset data in
 	 * @param bRecursive if true, the supplied path will be tested recursively
 	 */
-	UFUNCTION(BlueprintCallable, Category = "AssetRegistry")
+	UFUNCTION(BlueprintCallable, BlueprintPure=false, Category = "AssetRegistry")
 	virtual bool HasAssets(const FName PackagePath, const bool bRecursive = false) const = 0;
 
 	/**
@@ -125,7 +125,7 @@ public:
 	 * @param PackageName the package name for the requested assets
 	 * @param OutAssetData the list of assets in this path
 	 */
-	UFUNCTION(BlueprintCallable, Category="AssetRegistry")
+	UFUNCTION(BlueprintCallable, BlueprintPure=false, Category="AssetRegistry")
 	virtual bool GetAssetsByPackageName(FName PackageName, TArray<FAssetData>& OutAssetData, bool bIncludeOnlyOnDiskAssets = false) const = 0;
 
 	/**
@@ -135,7 +135,7 @@ public:
 	 * @param OutAssetData the list of assets in this path
 	 * @param bRecursive if true, all supplied paths will be searched recursively
 	 */
-	UFUNCTION(BlueprintCallable, Category = "AssetRegistry")
+	UFUNCTION(BlueprintCallable, BlueprintPure=false, Category = "AssetRegistry")
 	virtual bool GetAssetsByPath(FName PackagePath, TArray<FAssetData>& OutAssetData, bool bRecursive = false, bool bIncludeOnlyOnDiskAssets = false) const = 0;
 
 	/**
@@ -145,7 +145,7 @@ public:
 	 * @param OutAssetData the list of assets in this path
 	 * @param bSearchSubClasses if true, all subclasses of the passed in class will be searched as well
 	 */
-	UFUNCTION(BlueprintCallable, Category = "AssetRegistry")
+	UFUNCTION(BlueprintCallable, BlueprintPure=false, Category = "AssetRegistry")
 	virtual bool GetAssetsByClass(FName ClassName, TArray<FAssetData>& OutAssetData, bool bSearchSubClasses = false) const = 0;
 
 	/**
@@ -172,7 +172,7 @@ public:
 	 * @param Filter filter to apply to the assets in the AssetRegistry
 	 * @param OutAssetData the list of assets in this path
 	 */
-	UFUNCTION(BlueprintCallable, Category="AssetRegistry")
+	UFUNCTION(BlueprintCallable, BlueprintPure=false, Category="AssetRegistry")
 	virtual bool GetAssets(const FARFilter& Filter, TArray<FAssetData>& OutAssetData) const = 0;
 
 	/**
@@ -182,7 +182,7 @@ public:
 	 * @param bIncludeOnlyOnDiskAssets if true, in-memory objects will be ignored. The call will be faster.
 	 * @return the assets data;Will be invalid if object could not be found
 	 */
-	UFUNCTION(BlueprintCallable, Category="AssetRegistry")
+	UFUNCTION(BlueprintCallable, BlueprintPure=false, Category="AssetRegistry")
 	virtual FAssetData GetAssetByObjectPath( const FName ObjectPath, bool bIncludeOnlyOnDiskAssets = false ) const = 0;
 
 	/**
@@ -191,7 +191,7 @@ public:
 	 *
 	 * @param OutAssetData the list of assets in this path
 	 */
-	UFUNCTION(BlueprintCallable, Category = "AssetRegistry")
+	UFUNCTION(BlueprintCallable, BlueprintPure=false, Category = "AssetRegistry")
 	virtual bool GetAllAssets(TArray<FAssetData>& OutAssetData, bool bIncludeOnlyOnDiskAssets = false) const = 0;
 
 	/**
@@ -278,19 +278,19 @@ public:
 	virtual void GetDerivedClassNames(const TArray<FName>& ClassNames, const TSet<FName>& ExcludedClassNames, TSet<FName>& OutDerivedClassNames) const = 0;
 
 	/** Gets a list of all paths that are currently cached */
-	UFUNCTION(BlueprintCallable, Category = "AssetRegistry")
+	UFUNCTION(BlueprintCallable, BlueprintPure=false, Category = "AssetRegistry")
 	virtual void GetAllCachedPaths(TArray<FString>& OutPathList) const = 0;
 
 	/** Gets a list of all paths that are currently cached below the passed-in base path */
-	UFUNCTION(BlueprintCallable, Category = "AssetRegistry")
+	UFUNCTION(BlueprintCallable, BlueprintPure=false, Category = "AssetRegistry")
 	virtual void GetSubPaths(const FString& InBasePath, TArray<FString>& OutPathList, bool bInRecurse) const = 0;
 
 	/** Trims items out of the asset data list that do not pass the supplied filter */
-	UFUNCTION(BlueprintCallable, Category = "AssetRegistry")
+	UFUNCTION(BlueprintCallable, BlueprintPure=false, Category = "AssetRegistry")
 	virtual void RunAssetsThroughFilter(TArray<FAssetData>& AssetDataList, const FARFilter& Filter) const = 0;
 
 	/** Trims items out of the asset data list that do not pass the supplied filter */
-	UFUNCTION(BlueprintCallable, Category = "AssetRegistry")
+	UFUNCTION(BlueprintCallable, BlueprintPure=false, Category = "AssetRegistry")
 	virtual void UseFilterToExcludeAssets(TArray<FAssetData>& AssetDataList, const FARFilter& Filter) const = 0;
 
 	/** Modifies passed in filter to make it safe for use on FAssetRegistryState. This expands recursive paths and classes */
