#UE4 Editor/Program speed improvement for DependsNodes without increasing non-editor runtime size

diff --git a/Engine/Source/Runtime/AssetRegistry/Private/DependsNode.cpp b/Engine/Source/Runtime/AssetRegistry/Private/DependsNode.cpp
--- a/Engine/Source/Runtime/AssetRegistry/Private/DependsNode.cpp
+++ b/Engine/Source/Runtime/AssetRegistry/Private/DependsNode.cpp
@@ -55,27 +55,27 @@ void FDependsNode::GetReferencers(TArray<FDependsNode*>& OutReferencers, EAssetR
 			// If type is specified, filter. We don't use the iterate wrapper for performance
 			if (InDependencyType & EAssetRegistryDependencyType::Hard)
 			{
-				bShouldAdd = Referencer->HardDependencies.Contains(this);
+				bShouldAdd = Referencer->HardDependencies.Contains(const_cast<FDependsNode*>(this));
 			}
 
 			if (InDependencyType & EAssetRegistryDependencyType::Soft && !bShouldAdd)
 			{
-				bShouldAdd = Referencer->SoftDependencies.Contains(this);
+				bShouldAdd = Referencer->SoftDependencies.Contains(const_cast<FDependsNode*>(this));
 			}
 
 			if (InDependencyType & EAssetRegistryDependencyType::HardManage && !bShouldAdd)
 			{
-				bShouldAdd = Referencer->HardManageDependencies.Contains(this);
+				bShouldAdd = Referencer->HardManageDependencies.Contains(const_cast<FDependsNode*>(this));
 			}
 
 			if (InDependencyType & EAssetRegistryDependencyType::SoftManage && !bShouldAdd)
 			{
-				bShouldAdd = Referencer->SoftManageDependencies.Contains(this);
+				bShouldAdd = Referencer->SoftManageDependencies.Contains(const_cast<FDependsNode*>(this));
 			}
 
 			if (InDependencyType & EAssetRegistryDependencyType::SearchableName && !bShouldAdd)
 			{
-				bShouldAdd = Referencer->NameDependencies.Contains(this);
+				bShouldAdd = Referencer->NameDependencies.Contains(const_cast<FDependsNode*>(this));
 			}
 		}
 		else
@@ -92,51 +92,111 @@ void FDependsNode::GetReferencers(TArray<FDependsNode*>& OutReferencers, EAssetR
 
 void FDependsNode::AddDependency(FDependsNode* InDependency, EAssetRegistryDependencyType::Type InDependencyType, bool bGuaranteedUnique)
 {
-	IterateOverDependencyArrays([InDependency,&bGuaranteedUnique](TArray<FDependsNode*>& InArray, EAssetRegistryDependencyType::Type)
+	IterateOverDependencyLists([InDependency,&bGuaranteedUnique](FDependsNodeList& InList, EAssetRegistryDependencyType::Type)
 	{
+#if USE_DEPENDS_NODE_LIST_SETS
+		InList.Add(InDependency);
+#else
 		if (bGuaranteedUnique)
 		{
-			InArray.Add(InDependency);
+			InList.Add(InDependency);
 		}
 		else
 		{
-			InArray.AddUnique(InDependency);
+			InList.AddUnique(InDependency);
 		}
+#endif
 	}, InDependencyType);
 }
 
+void FDependsNode::AddReferencer(FDependsNode* InReferencer, bool bGuaranteedUnique)
+{
+#if USE_DEPENDS_NODE_LIST_SETS
+	Referencers.Add(InReferencer);
+#else
+	if (bGuaranteedUnique)
+	{
+		Referencers.Add(InReferencer);
+	}
+	else
+	{
+		Referencers.AddUnique(InReferencer);
+	}
+#endif
+}
+
 void FDependsNode::RemoveDependency(FDependsNode* InDependency)
 {
-	IterateOverDependencyArrays([InDependency](TArray<FDependsNode*>& InArray, EAssetRegistryDependencyType::Type)
+	IterateOverDependencyLists([InDependency](FDependsNodeList& InList, EAssetRegistryDependencyType::Type)
 	{
-		InArray.RemoveSingleSwap(InDependency, false);
+#if USE_DEPENDS_NODE_LIST_SETS
+		InList.Remove(InDependency);
+#else
+		InList.RemoveSingleSwap(InDependency, false);
+#endif
 	});
 }
 
+void FDependsNode::RemoveReferencer(FDependsNode* InReferencer)
+{
+#if USE_DEPENDS_NODE_LIST_SETS
+	Referencers.Remove(InReferencer);
+#else
+	Referencers.RemoveSingleSwap(InReferencer, false);
+#endif
+}
+
 void FDependsNode::ClearDependencies()
 {
-	IterateOverDependencyArrays([](TArray<FDependsNode*>& InArray, EAssetRegistryDependencyType::Type)
+	IterateOverDependencyLists([](FDependsNodeList& InSet, EAssetRegistryDependencyType::Type)
 	{
-		InArray.Empty();
+		InSet.Empty();
 	});
 }
 
 void FDependsNode::RemoveManageReferencesToNode()
 {
 	EAssetRegistryDependencyType::Type InDependencyType = EAssetRegistryDependencyType::Manage;
+#if USE_DEPENDS_NODE_LIST_SETS
+	// Iterate referencers set, possibly removing 
+	for (auto RefsIt = Referencers.CreateIterator(); RefsIt; ++RefsIt)
+	{
+		bool bStillExists = false;
+		(*RefsIt)->IterateOverDependencyLists([InDependencyType, this, &bStillExists](FDependsNodeList& InList, EAssetRegistryDependencyType::Type CurrentType)
+		{
+			if (InList.Contains(this))
+			{
+				if (CurrentType & InDependencyType)
+				{
+					InList.Remove(this);
+				}
+				else
+				{
+					// Reference of another type still exists
+					bStillExists = true;
+				}
+			}
+		}, EAssetRegistryDependencyType::All);
+
+		if (!bStillExists)
+		{
+			RefsIt.RemoveCurrent();
+		}
+	}
+#else
 	// Iterate referencers array, possibly removing 
 	for (int32 i = Referencers.Num() - 1; i >= 0; i--)
 	{
 		bool bStillExists = false;
-		Referencers[i]->IterateOverDependencyArrays([InDependencyType, this, &bStillExists](TArray<FDependsNode*>& InArray, EAssetRegistryDependencyType::Type CurrentType)
+		Referencers[i]->IterateOverDependencyLists([InDependencyType, this, &bStillExists](FDependsNodeList& InList, EAssetRegistryDependencyType::Type CurrentType)
 		{
-			int32 FoundIndex = InArray.Find(this);
+			int32 FoundIndex = InList.Find(this);
 
 			if (FoundIndex != INDEX_NONE)
 			{
 				if (CurrentType & InDependencyType)
 				{
-					InArray.RemoveAt(FoundIndex);
+					InList.RemoveAt(FoundIndex);
 				}
 				else
 				{
@@ -151,6 +211,7 @@ void FDependsNode::RemoveManageReferencesToNode()
 			Referencers.RemoveAt(i);
 		}
 	}
+#endif // USE_DEPENDS_NODE_LIST_SETS
 }
 
 void FDependsNode::PrintDependenciesRecursive(const FString& Indent, TSet<const FDependsNode*>& VisitedNodes) const
diff --git a/Engine/Source/Runtime/AssetRegistry/Private/DependsNode.h b/Engine/Source/Runtime/AssetRegistry/Private/DependsNode.h
--- a/Engine/Source/Runtime/AssetRegistry/Private/DependsNode.h
+++ b/Engine/Source/Runtime/AssetRegistry/Private/DependsNode.h
@@ -6,6 +6,14 @@
 #include "AssetData.h"
 #include "Misc/AssetRegistryInterface.h"
 
+#define USE_DEPENDS_NODE_LIST_SETS (WITH_EDITOR || IS_PROGRAM)
+
+#if USE_DEPENDS_NODE_LIST_SETS
+typedef TSet<class FDependsNode*> FDependsNodeList;
+#else
+typedef TArray<class FDependsNode*> FDependsNodeList;
+#endif
+
 /** Implementation of IDependsNode */
 class FDependsNode
 {
@@ -36,11 +44,11 @@ public:
 	/** Add a dependency to this node */
 	void AddDependency(FDependsNode* InDependency, EAssetRegistryDependencyType::Type InDependencyType = EAssetRegistryDependencyType::Hard, bool bGuaranteedUnique = false);
 	/** Add a referencer to this node */
-	void AddReferencer(FDependsNode* InReferencer, bool bGuaranteedUnique = false) { bGuaranteedUnique ? Referencers.Add(InReferencer) : Referencers.AddUnique(InReferencer); }
+	void AddReferencer(FDependsNode* InReferencer, bool bGuaranteedUnique = false);
 	/** Remove a dependency from this node */
 	void RemoveDependency(FDependsNode* InDependency);
 	/** Remove a referencer from this node */
-	void RemoveReferencer(FDependsNode* InReferencer) { Referencers.RemoveSingleSwap(InReferencer, false); }
+	void RemoveReferencer(FDependsNode* InReferencer);
 	/** Clear all dependency records from this node */
 	void ClearDependencies();
 	/** Removes Manage dependencies on this node and clean up referencers array. Manage references are the only ones safe to remove at runtime */
@@ -57,9 +65,9 @@ public:
 	template <class T>
 	void IterateOverDependencies(T InCallback, EAssetRegistryDependencyType::Type InDependencyType = EAssetRegistryDependencyType::All) const
 	{
-		IterateOverDependencyArrays([&InCallback](const TArray<FDependsNode*>& InArray, EAssetRegistryDependencyType::Type CurrentType)
+		IterateOverDependencyLists([&InCallback](const FDependsNodeList& InList, EAssetRegistryDependencyType::Type CurrentType)
 		{
-			for (FDependsNode* Dependency : InArray)
+			for (FDependsNode* Dependency : InList)
 			{
 				InCallback(Dependency, CurrentType);
 			}
@@ -94,32 +102,32 @@ private:
 
 	/** Iterate over all the separate dependency arrays. Const cast to avoid duplication */
 	template <class T>
-	FORCEINLINE void IterateOverDependencyArrays(T InCallback, EAssetRegistryDependencyType::Type InDependencyType = EAssetRegistryDependencyType::All) const
+	FORCEINLINE void IterateOverDependencyLists(T InCallback, EAssetRegistryDependencyType::Type InDependencyType = EAssetRegistryDependencyType::All) const
 	{
 		// This iteration is specific so it gets the "most important" references first in the array
 		if (InDependencyType & EAssetRegistryDependencyType::Hard)
 		{
-			InCallback(const_cast<TArray<FDependsNode*>&>(HardDependencies), EAssetRegistryDependencyType::Hard);
+			InCallback(const_cast<FDependsNodeList&>(HardDependencies), EAssetRegistryDependencyType::Hard);
 		}
 
 		if (InDependencyType & EAssetRegistryDependencyType::Soft)
 		{
-			InCallback(const_cast<TArray<FDependsNode*>&>(SoftDependencies), EAssetRegistryDependencyType::Soft);
+			InCallback(const_cast<FDependsNodeList&>(SoftDependencies), EAssetRegistryDependencyType::Soft);
 		}
 
 		if (InDependencyType & EAssetRegistryDependencyType::HardManage)
 		{
-			InCallback(const_cast<TArray<FDependsNode*>&>(HardManageDependencies), EAssetRegistryDependencyType::HardManage);
+			InCallback(const_cast<FDependsNodeList&>(HardManageDependencies), EAssetRegistryDependencyType::HardManage);
 		}
 
 		if (InDependencyType & EAssetRegistryDependencyType::SoftManage)
 		{
-			InCallback(const_cast<TArray<FDependsNode*>&>(SoftManageDependencies), EAssetRegistryDependencyType::SoftManage);
+			InCallback(const_cast<FDependsNodeList&>(SoftManageDependencies), EAssetRegistryDependencyType::SoftManage);
 		}
 
 		if (InDependencyType & EAssetRegistryDependencyType::SearchableName)
 		{
-			InCallback(const_cast<TArray<FDependsNode*>&>(NameDependencies), EAssetRegistryDependencyType::SearchableName);
+			InCallback(const_cast<FDependsNodeList&>(NameDependencies), EAssetRegistryDependencyType::SearchableName);
 		}
 	}
 
@@ -131,15 +139,15 @@ private:
 	/** The name of the package/object this node represents */
 	FAssetIdentifier Identifier;
 	/** The list of hard dependencies for this node */
-	TArray<FDependsNode*> HardDependencies;
+	FDependsNodeList HardDependencies;
 	/** The list of soft dependencies for this node */
-	TArray<FDependsNode*> SoftDependencies;
+	FDependsNodeList SoftDependencies;
 	/** The list of searchable name dependencies for this node */
-	TArray<FDependsNode*> NameDependencies;
+	FDependsNodeList NameDependencies;
 	/** The list of hard manage dependencies for this node */
-	TArray<FDependsNode*> SoftManageDependencies;
+	FDependsNodeList SoftManageDependencies;
 	/** The list of soft manage dependencies for this node */
-	TArray<FDependsNode*> HardManageDependencies;
+	FDependsNodeList HardManageDependencies;
 	/** The list of referencers to this node */
-	TArray<FDependsNode*> Referencers;
+	FDependsNodeList Referencers;
 };
