Fixed incorrect signature of RunAssetsThroughFilter and UseFilterToExcludeAssets in IAssetRegistry

The AssetDataList needs to be an in/out parameter (UPARAM(ref)), rather than an out parameter.

diff --git a/Engine/Source/Runtime/AssetRegistry/Public/IAssetRegistry.h b/Engine/Source/Runtime/AssetRegistry/Public/IAssetRegistry.h
--- a/Engine/Source/Runtime/AssetRegistry/Public/IAssetRegistry.h
+++ b/Engine/Source/Runtime/AssetRegistry/Public/IAssetRegistry.h
@@ -287,11 +287,11 @@ public:
 
 	/** Trims items out of the asset data list that do not pass the supplied filter */
 	UFUNCTION(BlueprintCallable, BlueprintPure=false, Category = "AssetRegistry")
-	virtual void RunAssetsThroughFilter(TArray<FAssetData>& AssetDataList, const FARFilter& Filter) const = 0;
+	virtual void RunAssetsThroughFilter(UPARAM(ref) TArray<FAssetData>& AssetDataList, const FARFilter& Filter) const = 0;
 
 	/** Trims items out of the asset data list that do not pass the supplied filter */
 	UFUNCTION(BlueprintCallable, BlueprintPure=false, Category = "AssetRegistry")
-	virtual void UseFilterToExcludeAssets(TArray<FAssetData>& AssetDataList, const FARFilter& Filter) const = 0;
+	virtual void UseFilterToExcludeAssets(UPARAM(ref) TArray<FAssetData>& AssetDataList, const FARFilter& Filter) const = 0;
 
 	/** Modifies passed in filter to make it safe for use on FAssetRegistryState. This expands recursive paths and classes */
 	virtual void ExpandRecursiveFilter(const FARFilter& InFilter, FARFilter& ExpandedFilter) const = 0;
