Guarad against crash when destroying asset view items

diff --git a/Engine/Source/Runtime/AssetRegistry/Public/AssetData.h b/Engine/Source/Runtime/AssetRegistry/Public/AssetData.h
--- a/Engine/Source/Runtime/AssetRegistry/Public/AssetData.h
+++ b/Engine/Source/Runtime/AssetRegistry/Public/AssetData.h
@@ -319,7 +319,7 @@ public:
 	/** Returns true if the asset is loaded */
 	bool IsAssetLoaded() const
 	{
-		return IsValid() && FindObject<UObject>(NULL, *ObjectPath.ToString()) != NULL;
+		return IsValid() && FindObjectSafe<UObject>(NULL, *ObjectPath.ToString()) != NULL;
 	}
 
 	/** Prints the details of the asset to the log */
