Prevent infinite recursion through identically named parent classes

diff --git a/Engine/Source/Runtime/AssetRegistry/Private/AssetRegistry.cpp b/Engine/Source/Runtime/AssetRegistry/Private/AssetRegistry.cpp
--- a/Engine/Source/Runtime/AssetRegistry/Private/AssetRegistry.cpp
+++ b/Engine/Source/Runtime/AssetRegistry/Private/AssetRegistry.cpp
@@ -2514,26 +2514,32 @@ void UAssetRegistryImpl::GetSubClasses(const TArray<FName>& InClassNames, const
 	for (FName ClassName : InClassNames)
 	{
 		// Now find all subclass names
-		GetSubClasses_Recursive(ClassName, SubClassNames, ReverseInheritanceMap, ExcludedClassNames);
+		TSet<FName> ProcessedClassNames;
+		GetSubClasses_Recursive(ClassName, SubClassNames, ProcessedClassNames, ReverseInheritanceMap, ExcludedClassNames);
 	}
 }
 
-void UAssetRegistryImpl::GetSubClasses_Recursive(FName InClassName, TSet<FName>& SubClassNames, const TMap<FName, TSet<FName>>& ReverseInheritanceMap, const TSet<FName>& ExcludedClassNames) const
+void UAssetRegistryImpl::GetSubClasses_Recursive(FName InClassName, TSet<FName>& SubClassNames, TSet<FName>& ProcessedClassNames, const TMap<FName, TSet<FName>>& ReverseInheritanceMap, const TSet<FName>& ExcludedClassNames) const
 {
 	if (ExcludedClassNames.Contains(InClassName))
 	{
 		// This class is in the exclusion list. Exclude it.
 	}
+	else if (ProcessedClassNames.Contains(InClassName))
+	{
+		// This class has already been processed. Ignore it.
+	}
 	else
 	{
 		SubClassNames.Add(InClassName);
+		ProcessedClassNames.Add(InClassName);
 
 		const TSet<FName>* FoundSubClassNames = ReverseInheritanceMap.Find(InClassName);
 		if (FoundSubClassNames)
 		{
 			for (FName ClassName : (*FoundSubClassNames))
 			{
-				GetSubClasses_Recursive(ClassName, SubClassNames, ReverseInheritanceMap, ExcludedClassNames);
+				GetSubClasses_Recursive(ClassName, SubClassNames, ProcessedClassNames, ReverseInheritanceMap, ExcludedClassNames);
 			}
 		}
 	}
diff --git a/Engine/Source/Runtime/AssetRegistry/Private/AssetRegistry.h b/Engine/Source/Runtime/AssetRegistry/Private/AssetRegistry.h
--- a/Engine/Source/Runtime/AssetRegistry/Private/AssetRegistry.h
+++ b/Engine/Source/Runtime/AssetRegistry/Private/AssetRegistry.h
@@ -217,7 +217,7 @@ private:
 
 	/** Returns the names of all subclasses of the class whose name is ClassName */
 	void GetSubClasses(const TArray<FName>& InClassNames, const TSet<FName>& ExcludedClassNames, TSet<FName>& SubClassNames) const;
-	void GetSubClasses_Recursive(FName InClassName, TSet<FName>& SubClassNames, const TMap<FName, TSet<FName>>& ReverseInheritanceMap, const TSet<FName>& ExcludedClassNames) const;
+	void GetSubClasses_Recursive(FName InClassName, TSet<FName>& SubClassNames, TSet<FName>& ProcessedClassNames, const TMap<FName, TSet<FName>>& ReverseInheritanceMap, const TSet<FName>& ExcludedClassNames) const;
 
 	/** Finds all class names of classes capable of generating new UClasses */
 	void CollectCodeGeneratorClasses();
