Removed buggy begin() end() iterators from BackgroundGatherResults and replaced with array access operator

diff --git a/Engine/Source/Runtime/AssetRegistry/Private/AssetRegistry.cpp b/Engine/Source/Runtime/AssetRegistry/Private/AssetRegistry.cpp
--- a/Engine/Source/Runtime/AssetRegistry/Private/AssetRegistry.cpp
+++ b/Engine/Source/Runtime/AssetRegistry/Private/AssetRegistry.cpp
@@ -1722,18 +1722,18 @@ void UAssetRegistryImpl::ScanPathsAndFilesSynchronous(const TArray<FString>& InP
 		if (OutFoundAssets)
 		{
 			OutFoundAssets->Reserve(OutFoundAssets->Num() + AssetResults.Num());
-			for (const FAssetData* AssetData : AssetResults)
+			for (int32 i = 0; i < AssetResults.Num(); ++i)
 			{
-				OutFoundAssets->Add(AssetData->ObjectPath);
+				OutFoundAssets->Add(AssetResults[i]->ObjectPath);
 			}
 		}
 
 		if (OutFoundPaths)
 		{
 			OutFoundPaths->Reserve(OutFoundPaths->Num() + PathResults.Num());
-			for (const FString& Path : PathResults)
+			for (int32 i = 0; i < PathResults.Num(); ++i)
 			{
-				OutFoundPaths->Add(*Path);
+				OutFoundPaths->Add(*PathResults[i]);
 			}
 		}
 
diff --git a/Engine/Source/Runtime/AssetRegistry/Private/BackgroundGatherResults.h b/Engine/Source/Runtime/AssetRegistry/Private/BackgroundGatherResults.h
--- a/Engine/Source/Runtime/AssetRegistry/Private/BackgroundGatherResults.h
+++ b/Engine/Source/Runtime/AssetRegistry/Private/BackgroundGatherResults.h
@@ -86,6 +86,13 @@ public:
 		return InternalQueue.GetAllocatedSize();
 	}
 
+	/** Array access operator */
+	FORCEINLINE T& operator[](int32 Index)
+	{
+		check(Index >= 0 && Index < Num());
+		return InternalQueue[Index + PoppedCount];
+	}
+
 	/** Prioritize any items that pass the given predicate to be at the front of the queue. */
 	void Prioritize(TFunctionRef<bool(const T&)> Pred)
 	{
@@ -103,24 +110,6 @@ public:
 	
 private:
 	typedef TArray<T, Allocator> InternalQueueType;
-	typedef typename InternalQueueType::RangedForIteratorType RangedForIteratorType;
-	typedef typename InternalQueueType::RangedForConstIteratorType RangedForConstIteratorType;
-
-	/**
-	 * DO NOT USE DIRECTLY
-	 * STL-like iterators to enable range-based for loop support.
-	 */
-	#if TARRAY_RANGED_FOR_CHECKS
-		FORCEINLINE friend RangedForIteratorType      begin(      TBackgroundGatherResults& Results) { return RangedForIteratorType     (Results.InternalQueue.Num(), Results.InternalQueue.GetData() + Results.PoppedCount); }
-		FORCEINLINE friend RangedForConstIteratorType begin(const TBackgroundGatherResults& Results) { return RangedForConstIteratorType(Results.InternalQueue.Num(), Results.InternalQueue.GetData() + Results.PoppedCount); }
-		FORCEINLINE friend RangedForIteratorType      end  (      TBackgroundGatherResults& Results) { return RangedForIteratorType     (Results.InternalQueue.Num(), Results.InternalQueue.GetData() + Results.InternalQueue.Num()); }
-		FORCEINLINE friend RangedForConstIteratorType end  (const TBackgroundGatherResults& Results) { return RangedForConstIteratorType(Results.InternalQueue.Num(), Results.InternalQueue.GetData() + Results.InternalQueue.Num()); }
-	#else
-		FORCEINLINE friend RangedForIteratorType      begin(      TBackgroundGatherResults& Results) { return Results.InternalQueue.GetData() + Results.PoppedCount; }
-		FORCEINLINE friend RangedForConstIteratorType begin(const TBackgroundGatherResults& Results) { return Results.InternalQueue.GetData() + Results.PoppedCount; }
-		FORCEINLINE friend RangedForIteratorType      end  (      TBackgroundGatherResults& Results) { return Results.InternalQueue.GetData() + Results.InternalQueue.Num(); }
-		FORCEINLINE friend RangedForConstIteratorType end  (const TBackgroundGatherResults& Results) { return Results.InternalQueue.GetData() + Results.InternalQueue.Num(); }
-	#endif
 
 private:
 	/** Number of items that have been popped off the queue without it being trimmed. Items before this count should not be mutated, and we should trim to this number. */
