Guarded against reading a corrupt package

diff --git a/Engine/Source/Editor/UnrealEd/Private/ObjectTools.cpp b/Engine/Source/Editor/UnrealEd/Private/ObjectTools.cpp
--- a/Engine/Source/Editor/UnrealEd/Private/ObjectTools.cpp
+++ b/Engine/Source/Editor/UnrealEd/Private/ObjectTools.cpp
@@ -4111,7 +4111,7 @@ namespace ThumbnailTools
 
 
 		// Make sure this is indeed a package
-		if( FileSummary.Tag != PACKAGE_FILE_TAG )
+		if( FileSummary.Tag != PACKAGE_FILE_TAG || FileReader->IsError() )
 		{
 			// Unrecognized or malformed package file
 			return false;
diff --git a/Engine/Source/Runtime/AssetRegistry/Private/PackageReader.cpp b/Engine/Source/Runtime/AssetRegistry/Private/PackageReader.cpp
--- a/Engine/Source/Runtime/AssetRegistry/Private/PackageReader.cpp
+++ b/Engine/Source/Runtime/AssetRegistry/Private/PackageReader.cpp
@@ -57,12 +57,12 @@ bool FPackageReader::OpenPackageFile(EOpenPackageResult* OutErrorCode)
 	}
 
 	// Read package file summary from the file
-	*this << PackageFileSummary;
+	*Loader << PackageFileSummary;
 
 	// Validate the summary.
 
 	// Make sure this is indeed a package
-	if( PackageFileSummary.Tag != PACKAGE_FILE_TAG )
+	if( PackageFileSummary.Tag != PACKAGE_FILE_TAG || Loader->IsError() )
 	{
 		// Unrecognized or malformed package file
 		SetPackageErrorCode(EOpenPackageResult::MalformedTag);
@@ -124,7 +124,7 @@ bool FPackageReader::ReadAssetRegistryData (TArray<FAssetData*>& AssetDataList)
 	check(Loader);
 
 	// Does the package contain asset registry tags
-	if( PackageFileSummary.AssetRegistryDataOffset == 0 )
+	if( PackageFileSummary.AssetRegistryDataOffset <= 0 || PackageFileSummary.AssetRegistryDataOffset > PackageFileSize )
 	{
 		// No Tag Table!
 		return false;
@@ -218,7 +218,7 @@ bool FPackageReader::ReadAssetDataFromThumbnailCache(TArray<FAssetData*>& AssetD
 	check(Loader);
 
 	// Does the package contain a thumbnail table?
-	if( PackageFileSummary.ThumbnailTableOffset == 0 )
+	if( PackageFileSummary.ThumbnailTableOffset <= 0 || PackageFileSummary.ThumbnailTableOffset > PackageFileSize )
 	{
 		return false;
 	}
diff --git a/Engine/Source/Runtime/CoreUObject/Private/UObject/PackageFileSummary.cpp b/Engine/Source/Runtime/CoreUObject/Private/UObject/PackageFileSummary.cpp
--- a/Engine/Source/Runtime/CoreUObject/Private/UObject/PackageFileSummary.cpp
+++ b/Engine/Source/Runtime/CoreUObject/Private/UObject/PackageFileSummary.cpp
@@ -215,14 +215,18 @@ void operator<<(FStructuredArchive::FSlot Slot, FPackageFileSummary& Sum)
 		Record << NAMED_ITEM("Guid", Sum.Guid) << NAMED_ITEM("GenerationCount", GenerationCount);
 		if (BaseArchive.IsLoading() && GenerationCount > 0)
 		{
-			Sum.Generations.Empty(1);
-			Sum.Generations.AddUninitialized(GenerationCount);
+			Sum.Generations.Reset(GenerationCount);
+			Sum.Generations.AddZeroed(GenerationCount);
 		}
 
 		FStructuredArchive::FStream GenerationsStream = Record.EnterStream(FIELD_NAME_TEXT("Generations"));
 		for (int32 i = 0; i<GenerationCount; i++)
 		{
 			Sum.Generations[i].Serialize(GenerationsStream.EnterElement(), Sum);
+			if (BaseArchive.IsLoading() && BaseArchive.IsError())
+			{
+				return;
+			}
 		}
 
 		if (Sum.GetFileVersionUE4() >= VER_UE4_ENGINE_VERSION_OBJECT)
