FARFilter will now exclude classes when no classes are listed

diff --git a/Engine/Source/Runtime/AssetRegistry/Private/AssetRegistry.cpp b/Engine/Source/Runtime/AssetRegistry/Private/AssetRegistry.cpp
--- a/Engine/Source/Runtime/AssetRegistry/Private/AssetRegistry.cpp
+++ b/Engine/Source/Runtime/AssetRegistry/Private/AssetRegistry.cpp
@@ -1119,8 +1119,20 @@ void UAssetRegistryImpl::ExpandRecursiveFilter(const FARFilter& InFilter, FARFil
 
 	if (InFilter.bRecursiveClasses)
 	{
-		// GetSubClasses includes the base classes
-		GetSubClasses(InFilter.ClassNames, InFilter.RecursiveClassesExclusionSet, FilterClassNames);
+		if (InFilter.RecursiveClassesExclusionSet.Num() > 0 && InFilter.ClassNames.Num() == 0)
+		{
+			// Build list of all classes then remove excluded classes
+			TArray<FName> ClassNamesObject;
+			ClassNamesObject.Add(UObject::StaticClass()->GetFName());
+
+			// GetSubClasses includes the base classes
+			GetSubClasses(ClassNamesObject, InFilter.RecursiveClassesExclusionSet, FilterClassNames);
+		}
+		else
+		{
+			// GetSubClasses includes the base classes
+			GetSubClasses(InFilter.ClassNames, InFilter.RecursiveClassesExclusionSet, FilterClassNames);
+		}
 	}
 	else
 	{
