Avoid the async background asset registry tasks running until the engine is running

There can sometimes be a race condition where one of the global lists (such as the package names or paths lists) are updated due to content being mounted, while one of these background threads is using that list.

#rnx

diff --git a/Engine/Source/Runtime/AssetRegistry/Private/AssetDataGatherer.cpp b/Engine/Source/Runtime/AssetRegistry/Private/AssetDataGatherer.cpp
--- a/Engine/Source/Runtime/AssetRegistry/Private/AssetDataGatherer.cpp
+++ b/Engine/Source/Runtime/AssetRegistry/Private/AssetDataGatherer.cpp
@@ -102,6 +102,16 @@ bool FAssetDataDiscovery::Init()
 
 uint32 FAssetDataDiscovery::Run()
 {
+	if (!bIsSynchronous)
+	{
+		// If we're running asynchronous, don't allow these tasks to start until the engine is "running"
+		// as we may still be manipulating global lists during start-up configuration
+		while (!GIsRunning && StopTaskCounter.GetValue() == 0)
+		{
+			FPlatformProcess::SleepNoStats(0.1f);
+		}
+	}
+
 	double DiscoverStartTime = FPlatformTime::Seconds();
 	int32 NumDiscoveredFiles = 0;
 
