Fix 'Collections' folder from appearing in content browser

diff --git a/Engine/Source/Runtime/AssetRegistry/Private/AssetRegistry.cpp b/Engine/Source/Runtime/AssetRegistry/Private/AssetRegistry.cpp
--- a/Engine/Source/Runtime/AssetRegistry/Private/AssetRegistry.cpp
+++ b/Engine/Source/Runtime/AssetRegistry/Private/AssetRegistry.cpp
@@ -2449,7 +2449,7 @@ void UAssetRegistryImpl::OnDirectoryChanged(const TArray<FFileChangeData>& FileC
 			{
 			case FFileChangeData::FCA_Added:
 			{
-				if (FPaths::DirectoryExists(File))
+				if (FPaths::DirectoryExists(File) && LongPackageName != TEXT("/Game/Collections"))
 				{
 					AddPath(LongPackageName);
 					UE_LOG(LogAssetRegistry, Verbose, TEXT("Directory was added to content directory: %s"), *File);
