Backing out CL# 4982820 as it caused an indefinite wait when using cook-on-the-side

#rnx

diff --git a/Engine/Source/Runtime/AssetRegistry/Private/AssetDataGatherer.cpp b/Engine/Source/Runtime/AssetRegistry/Private/AssetDataGatherer.cpp
--- a/Engine/Source/Runtime/AssetRegistry/Private/AssetDataGatherer.cpp
+++ b/Engine/Source/Runtime/AssetRegistry/Private/AssetDataGatherer.cpp
@@ -102,15 +102,16 @@ bool FAssetDataDiscovery::Init()
 
 uint32 FAssetDataDiscovery::Run()
 {
-	if (!bIsSynchronous)
-	{
-		// If we're running asynchronous, don't allow these tasks to start until the engine is "running"
-		// as we may still be manipulating global lists during start-up configuration
-		while (!GIsRunning && StopTaskCounter.GetValue() == 0)
-		{
-			FPlatformProcess::SleepNoStats(0.1f);
-		}
-	}
+	// Commenting out the code below as it causes cook-on-the-side to wait indefinitely in FAssetRegistryGenerator::Initialize
+	//if (!bIsSynchronous)
+	//{
+	//	// If we're running asynchronous, don't allow these tasks to start until the engine is "running"
+	//	// as we may still be manipulating global lists during start-up configuration
+	//	while (!GIsRunning && StopTaskCounter.GetValue() == 0)
+	//	{
+	//		FPlatformProcess::SleepNoStats(0.1f);
+	//	}
+	//}
 
 	double DiscoverStartTime = FPlatformTime::Seconds();
 	int32 NumDiscoveredFiles = 0;
