Optimize asset gathering by calling RemoveSingleSwap on DependsNode.

While pruning depdendencies during asset registry building call RemoveSingleSwap to reduce cost of removing assets from node dependency lists. Also avoid resizing the arrays since is the last thing that happens in the cook we might as well keep it. On a full cook of a large dataset, walltime of FAssetRegistryState::RemoveAssetData went from 22.9s to 9.2s.

diff --git a/Engine/Source/Runtime/AssetRegistry/Private/DependsNode.cpp b/Engine/Source/Runtime/AssetRegistry/Private/DependsNode.cpp
--- a/Engine/Source/Runtime/AssetRegistry/Private/DependsNode.cpp
+++ b/Engine/Source/Runtime/AssetRegistry/Private/DependsNode.cpp
@@ -109,7 +109,7 @@ void FDependsNode::RemoveDependency(FDependsNode* InDependency)
 {
 	IterateOverDependencyArrays([InDependency](TArray<FDependsNode*>& InArray, EAssetRegistryDependencyType::Type)
 	{
-		InArray.Remove(InDependency);
+		InArray.RemoveSingleSwap(InDependency, false);
 	});
 }
 
diff --git a/Engine/Source/Runtime/AssetRegistry/Private/DependsNode.h b/Engine/Source/Runtime/AssetRegistry/Private/DependsNode.h
--- a/Engine/Source/Runtime/AssetRegistry/Private/DependsNode.h
+++ b/Engine/Source/Runtime/AssetRegistry/Private/DependsNode.h
@@ -40,7 +40,7 @@ public:
 	/** Remove a dependency from this node */
 	void RemoveDependency(FDependsNode* InDependency);
 	/** Remove a referencer from this node */
-	void RemoveReferencer(FDependsNode* InReferencer) { Referencers.Remove(InReferencer); }
+	void RemoveReferencer(FDependsNode* InReferencer) { Referencers.RemoveSingleSwap(InReferencer, false); }
 	/** Clear all dependency records from this node */
 	void ClearDependencies();
 	/** Removes Manage dependencies on this node and clean up referencers array. Manage references are the only ones safe to remove at runtime */
