Several optimizations to AssetRegistry in cook scenarios

* Added InitializeFromExistingAndPrune method for AssetRegistryState, which copies and filters in one go. Alot faster than calling InitializeFromExisting and then PruneAssetData.
* Writes chunks info files slightly different by looping trough assets once writing to corresponding files, instead of looping chunks then looping all assets. Sorting in AllAssetChunks.txt will be different but still consistent between runs.
* Added a specialized reserve function in DependsNode in order to streamline memory allocations when copying a node.
* Added console commands to dump out asset registry information.

Tested on a large dataset, InitializeFromExisting + Prune went from taking 72 seconds to 3.3 seconds, GenerateAssetChunkInformationCSV went from taking 11.4 seconds to 1.1 seconds.

diff --git a/Engine/Plugins/Editor/AssetManagerEditor/Source/AssetManagerEditor/Private/AssetManagerEditorModule.cpp b/Engine/Plugins/Editor/AssetManagerEditor/Source/AssetManagerEditor/Private/AssetManagerEditorModule.cpp
--- a/Engine/Plugins/Editor/AssetManagerEditor/Source/AssetManagerEditor/Private/AssetManagerEditorModule.cpp
+++ b/Engine/Plugins/Editor/AssetManagerEditor/Source/AssetManagerEditor/Private/AssetManagerEditorModule.cpp
@@ -296,6 +296,7 @@ public:
 	void PerformAuditConsoleCommand(const TArray<FString>& Args);
 	void PerformDependencyChainConsoleCommand(const TArray<FString>& Args);
 	void PerformDependencyClassConsoleCommand(const TArray<FString>& Args);
+	void DumpAssetRegistry(const TArray<FString>& Args);
 	void DumpAssetDependencies(const TArray<FString>& Args);
 
 	virtual void OpenAssetAuditUI(TArray<FAssetData> SelectedAssets) override;
@@ -435,6 +436,15 @@ void FAssetManagerEditorModule::StartupModule()
 			ECVF_Default
 			));
 
+	#if ASSET_REGISTRY_STATE_DUMPING_ENABLED
+		AuditCmds.Add(IConsoleManager::Get().RegisterConsoleCommand(
+			TEXT("AssetManager.DumpAssetRegistry"),
+			TEXT("Prints entries in the asset registry. Arguments are required: ObjectPath, PackageName, Path, Class, Tag, Dependencies, PackageData."),
+			FConsoleCommandWithArgsDelegate::CreateRaw(this, &FAssetManagerEditorModule::DumpAssetRegistry),
+			ECVF_Default
+		));
+	#endif
+
 		AuditCmds.Add(IConsoleManager::Get().RegisterConsoleCommand(
 			TEXT("AssetManager.DumpAssetDependencies"),
 			TEXT("Shows a list of all primary assets and the secondary assets that they depend on. Also writes out a .graphviz file"),
@@ -1997,6 +2007,21 @@ void FAssetManagerEditorModule::LogAssetsWithMultipleLabels()
 	}
 }
 
+
+void FAssetManagerEditorModule::DumpAssetRegistry(const TArray<FString>& Args)
+{
+#if ASSET_REGISTRY_STATE_DUMPING_ENABLED
+	UAssetManager& Manager = UAssetManager::Get();
+	TArray<FString> ReportLines;
+
+	const FAssetRegistryState* State = Manager.GetAssetRegistry().GetAssetRegistryState();
+	State->Dump(Args, ReportLines);
+
+	Manager.WriteCustomReport(FString::Printf(TEXT("AssetRegistryState-%s.txt"), *FDateTime::Now().ToString()), ReportLines);
+#endif
+}
+
+
 void FAssetManagerEditorModule::DumpAssetDependencies(const TArray<FString>& Args)
 {
 	if (!UAssetManager::IsValid())
diff --git a/Engine/Source/Editor/UnrealEd/Private/Commandlets/AssetRegistryGenerator.cpp b/Engine/Source/Editor/UnrealEd/Private/Commandlets/AssetRegistryGenerator.cpp
--- a/Engine/Source/Editor/UnrealEd/Private/Commandlets/AssetRegistryGenerator.cpp
+++ b/Engine/Source/Editor/UnrealEd/Private/Commandlets/AssetRegistryGenerator.cpp
@@ -1011,8 +1011,8 @@ bool FAssetRegistryGenerator::SaveAssetRegistry(const FString& SandboxPath, bool
 		{
 			// Prune out the development only packages, and any assets that belong in a different chunk asset registry
 			FAssetRegistryState NewState;
-			NewState.InitializeFromExisting(State, SaveOptions);
-			NewState.PruneAssetData(CookedPackages, TSet<FName>(), ChunkBucketElement.Value, SaveOptions);
+			NewState.InitializeFromExistingAndPrune(State, CookedPackages, TSet<FName>(), ChunkBucketElement.Value, SaveOptions);
+
 			InjectEncryptionData(NewState);
 
 			// Create runtime registry data
@@ -1215,12 +1215,8 @@ bool FAssetRegistryGenerator::GatherAllPackageDependencies(FName PackageName, TA
 
 bool FAssetRegistryGenerator::GenerateAssetChunkInformationCSV(const FString& OutputPath, bool bWriteIndividualFiles)
 {
-	FString TmpString;
-	FString CSVString;
-	FString HeaderText(TEXT("ChunkID, Package Name, Class Type, Hard or Soft Chunk, File Size, Other Chunks\n"));
-	FString EndLine(TEXT("\n"));
-	FString NoneText(TEXT("None\n"));
-	CSVString = HeaderText;
+	FString TmpString, TmpStringChunks;
+	ANSICHAR HeaderText[] = "ChunkID, Package Name, Class Type, Hard or Soft Chunk, File Size, Other Chunks\n";
 
 	const TMap<FName, const FAssetData*>& ObjectToDataMap = State.GetObjectPathToAssetDataMap();
 	TArray<const FAssetData*> AssetDataList;
@@ -1235,69 +1231,93 @@ bool FAssetRegistryGenerator::GenerateAssetChunkInformationCSV(const FString& Ou
 		return A.ObjectPath < B.ObjectPath;
 	});
 
-	for (int32 ChunkID = 0, ChunkNum = FinalChunkManifests.Num(); ChunkID < ChunkNum; ++ChunkID)
+	// Create file for all chunks
+	TUniquePtr<FArchive> AllChunksFile(IFileManager::Get().CreateFileWriter(*FPaths::Combine(*OutputPath, TEXT("AllChunksInfo.csv"))));
+	if (!AllChunksFile.IsValid())
+	{
+		return false;
+	}
+
+	AllChunksFile->Serialize(HeaderText, sizeof(HeaderText)-1);
+
+	// Create file for each chunk if needed
+	TArray<TUniquePtr<FArchive>> ChunkFiles;
+	if (bWriteIndividualFiles)
 	{
-		FString PerChunkManifestCSV = HeaderText;
-		for (const FAssetData* AssetDataPtr : AssetDataList)
+		for (int32 ChunkID = 0, ChunkNum = FinalChunkManifests.Num(); ChunkID < ChunkNum; ++ChunkID)
 		{
-			const FAssetData& AssetData = *AssetDataPtr;
-			// Add only assets that have actually been cooked and belong to any chunk
-			if (AssetData.ChunkIDs.Num() > 0)
+			FArchive* ChunkFile = IFileManager::Get().CreateFileWriter(*FPaths::Combine(*OutputPath, *FString::Printf(TEXT("Chunks%dInfo.csv"), ChunkID)));
+			if (ChunkFile != nullptr)
 			{
-				const FAssetPackageData* PackageData = State.GetAssetPackageData(AssetData.PackageName);
-				if (AssetData.ChunkIDs.Contains(ChunkID) && PackageData && PackageData->DiskSize >= 0)
+				return false;
+			}
+			ChunkFile->Serialize(HeaderText, sizeof(HeaderText)-1);
+			ChunkFiles.Add(TUniquePtr<FArchive>(ChunkFile));
+		}
+	}
+
+	for (const FAssetData* AssetDataPtr : AssetDataList)
+	{
+		const FAssetData& AssetData = *AssetDataPtr;
+		const FAssetPackageData* PackageData = State.GetAssetPackageData(AssetData.PackageName);
+
+		// Add only assets that have actually been cooked and belong to any chunk and that have a file size
+		if (AssetData.ChunkIDs.Num() > 0 && PackageData->DiskSize > 0)
+		{
+			for (int32 ChunkID : AssetData.ChunkIDs)
+			{
+				const int64 FileSize = PackageData->DiskSize;
+				FString SoftChain;
+				bool bHardChunk = false;
+				if (ChunkID < ChunkManifests.Num())
 				{
-					int64 FileSize = PackageData->DiskSize;
-					FString SoftChain;
-					bool bHardChunk = false;
-					if (ChunkID < ChunkManifests.Num())
-					{
-						bHardChunk = ChunkManifests[ChunkID] && ChunkManifests[ChunkID]->Contains(AssetData.PackageName);
-					
-						if (!bHardChunk)
-						{
-							//
-							SoftChain = GetShortestReferenceChain(AssetData.PackageName, ChunkID);
-						}
-					}
-					if (SoftChain.IsEmpty())
+					bHardChunk = ChunkManifests[ChunkID] && ChunkManifests[ChunkID]->Contains(AssetData.PackageName);
+
+					if (!bHardChunk)
 					{
-						SoftChain = TEXT("Soft: Possibly Unassigned Asset");
+						SoftChain = GetShortestReferenceChain(AssetData.PackageName, ChunkID);
 					}
+				}
+				if (SoftChain.IsEmpty())
+				{
+					SoftChain = TEXT("Soft: Possibly Unassigned Asset");
+				}
 
-					TmpString = FString::Printf(TEXT("%d,%s,%s,%s,%lld,"), ChunkID, *AssetData.PackageName.ToString(), *AssetData.AssetClass.ToString(), bHardChunk ? TEXT("Hard") : *SoftChain, FileSize);
-					CSVString += TmpString;
-					PerChunkManifestCSV += TmpString;
-					if (AssetData.ChunkIDs.Num() == 1)
+				// Build "other chunks" string or None if not part of
+				TmpStringChunks.Empty(64);
+				for (const auto& OtherChunk : AssetData.ChunkIDs)
+				{
+					if (OtherChunk != ChunkID)
 					{
-						CSVString += NoneText;
-						PerChunkManifestCSV += NoneText;
+						TmpString = FString::Printf(TEXT("%d "), OtherChunk);
 					}
-					else
+				}
+
+				// Build csv line
+				TmpString = FString::Printf(TEXT("%d,%s,%s,%s,%lld,%s\n"),
+					ChunkID,
+					*AssetData.PackageName.ToString(),
+					*AssetData.AssetClass.ToString(),
+					bHardChunk ? TEXT("Hard") : *SoftChain,
+					FileSize,
+					AssetData.ChunkIDs.Num() == 1 ? TEXT("None") : *TmpStringChunks
+				);
+
+				// Write line to all chunks file and individual chunks files if requested
+				{
+					auto Src = StringCast<ANSICHAR>(*TmpString, TmpString.Len());
+					AllChunksFile->Serialize((ANSICHAR*)Src.Get(), Src.Length() * sizeof(ANSICHAR));
+					
+					if (bWriteIndividualFiles)
 					{
-						for (const auto& OtherChunk : AssetData.ChunkIDs)
-						{
-							if (OtherChunk != ChunkID)
-							{
-								TmpString = FString::Printf(TEXT("%d "), OtherChunk);
-								CSVString += TmpString;
-								PerChunkManifestCSV += TmpString;
-							}
-						}
-						CSVString += EndLine;
-						PerChunkManifestCSV += EndLine;
+						ChunkFiles[ChunkID]->Serialize((ANSICHAR*)Src.Get(), Src.Length() * sizeof(ANSICHAR));
 					}
 				}
 			}
 		}
-
-		if (bWriteIndividualFiles)
-		{
-			FFileHelper::SaveStringToFile(PerChunkManifestCSV, *FPaths::Combine(*OutputPath, *FString::Printf(TEXT("Chunks%dInfo.csv"), ChunkID)));
-		}
 	}
 
-	return FFileHelper::SaveStringToFile(CSVString, *FPaths::Combine(*OutputPath, TEXT("AllChunksInfo.csv")));
+	return true;
 }
 
 void FAssetRegistryGenerator::AddPackageToManifest(const FString& PackageSandboxPath, FName PackageName, int32 ChunkId)
diff --git a/Engine/Source/Runtime/AssetRegistry/Private/AssetRegistryState.cpp b/Engine/Source/Runtime/AssetRegistry/Private/AssetRegistryState.cpp
--- a/Engine/Source/Runtime/AssetRegistry/Private/AssetRegistryState.cpp
+++ b/Engine/Source/Runtime/AssetRegistry/Private/AssetRegistryState.cpp
@@ -215,6 +215,125 @@ void FAssetRegistryState::Reset()
 	CachedPackageData.Empty();
 }
 
+void FAssetRegistryState::FilterTags(const FAssetDataTagMapSharedView& InTagsAndValues, FAssetDataTagMap& OutTagsAndValues, const TSet<FName>* ClassSpecificFilterlist, const FAssetRegistrySerializationOptions& Options)
+{
+	static FName WildcardName(TEXT("*"));
+	const TSet<FName>* AllClassesFilterlist = Options.CookFilterlistTagsByClass.Find(WildcardName);
+
+	// Exclude blacklisted tags or include only white listed tags, based on how we were configured in ini
+	for (const auto& TagPair : InTagsAndValues)
+	{
+		const bool bInAllClasseslist = AllClassesFilterlist && (AllClassesFilterlist->Contains(TagPair.Key) || AllClassesFilterlist->Contains(WildcardName));
+		const bool bInClassSpecificlist = ClassSpecificFilterlist && (ClassSpecificFilterlist->Contains(TagPair.Key) || ClassSpecificFilterlist->Contains(WildcardName));
+		if (Options.bUseAssetRegistryTagsWhitelistInsteadOfBlacklist)
+		{
+			// It's a white list, only include it if it is in the all classes list or in the class specific list
+			if (bInAllClasseslist || bInClassSpecificlist)
+			{
+				// It is in the white list. Keep it.
+				OutTagsAndValues.Add(TagPair.Key, TagPair.Value);
+			}
+		}
+		else
+		{
+			// It's a blacklist, include it unless it is in the all classes list or in the class specific list
+			if (!bInAllClasseslist && !bInClassSpecificlist)
+			{
+				// It isn't in the blacklist. Keep it.
+				OutTagsAndValues.Add(TagPair.Key, TagPair.Value);
+			}
+		}
+	}
+}
+
+void FAssetRegistryState::InitializeFromExistingAndPrune(const FAssetRegistryState & ExistingState, const TSet<FName>& RequiredPackages, const TSet<FName>& RemovePackages, const TSet<int32> ChunksToKeep, const FAssetRegistrySerializationOptions & Options)
+{
+	const bool bIsFilteredByChunkId = ChunksToKeep.Num() != 0;
+	const bool bIsFilteredByRequiredPackage = RequiredPackages.Num() != 0;
+	const bool bIsFilteredByRemovedPackages = RemovePackages.Num() != 0;
+	const bool bIsFilteredByNoTags = Options.bFilterAssetDataWithNoTags;
+
+	// Duplicate asset data entries
+	for (const TPair<FName, FAssetData*>& AssetPair : ExistingState.CachedAssetsByObjectPath)
+	{
+		const FAssetData* AssetData = AssetPair.Value;
+		if (AssetData)
+		{
+			bool bDiscard = false;
+			bDiscard |= bIsFilteredByChunkId && !AssetData->ChunkIDs.ContainsByPredicate([&](int32 ChunkId) { return ChunksToKeep.Contains(ChunkId); });
+			bDiscard |= bIsFilteredByRequiredPackage && !RequiredPackages.Contains(AssetData->PackageName);
+			bDiscard |= bIsFilteredByRemovedPackages && RemovePackages.Contains(AssetData->PackageName);
+			bDiscard |= bIsFilteredByNoTags && AssetData->TagsAndValues.Num() == 0 && !FPackageName::IsLocalizedPackage(AssetData->PackageName.ToString());
+
+			if (bDiscard)
+			{
+				continue;
+			}
+
+			FAssetDataTagMap NewTagsAndValues;
+			FAssetRegistryState::FilterTags(AssetData->TagsAndValues, NewTagsAndValues, Options.CookFilterlistTagsByClass.Find(AssetData->AssetClass), Options);
+
+			FAssetData* NewAssetData = new FAssetData(AssetData->PackageName, AssetData->PackagePath, AssetData->AssetName,
+				AssetData->AssetClass, NewTagsAndValues, AssetData->ChunkIDs, AssetData->PackageFlags);
+			// Add asset to new state
+			AddAssetData(NewAssetData);
+			// Also add it's package if it doesn't exist
+			CreateOrGetAssetPackageData(NewAssetData->PackageName);
+		}
+	}
+
+	TSet<FAssetIdentifier> ScriptPackages;
+
+	// Build list of script packages for depend nodes copying
+	for (const FName& Name : RequiredPackages)
+	{
+		bool bIsScriptPackage = FPackageName::IsScriptPackage(Name.ToString());
+		const FAssetPackageData* OldData = ExistingState.GetAssetPackageData(Name);
+		if (bIsScriptPackage) 
+		{
+			if (OldData)
+			{
+				ScriptPackages.Add(Name);
+
+				FAssetPackageData* NewData = CreateOrGetAssetPackageData(Name);
+				*NewData = *OldData;
+			}
+		}
+		else if (CachedAssetsByPackageName.Find(Name))
+		{
+			FAssetPackageData* NewData = CreateOrGetAssetPackageData(Name);
+			*NewData = *OldData;
+		}
+	}
+
+
+	// Duplicate dependency nodes
+	for (const TPair<FAssetIdentifier, FDependsNode*>& Pair : ExistingState.CachedDependsNodes)
+	{
+		FDependsNode* OldNode = Pair.Value;
+
+		// Only add if this node is already part of the filtered asset set
+		if (!CachedAssetsByPackageName.Contains(OldNode->GetIdentifier().PackageName))
+		{
+			continue;
+		}
+
+		FDependsNode* NewNode = CreateOrFindDependsNode(Pair.Key);
+		NewNode->Reserve(OldNode);
+	
+		OldNode->IterateOverDependencies([&, OldNode, NewNode](FDependsNode* InDependency, EAssetRegistryDependencyType::Type InDependencyType) {
+			const FAssetIdentifier& Identifier = InDependency->GetIdentifier();
+			if (CachedAssetsByPackageName.Find(Identifier.PackageName) || ScriptPackages.Contains(Identifier))
+			{
+				// Only add link if it's part of the filtered asset set
+				FDependsNode* NewDependency = CreateOrFindDependsNode(Identifier);
+				NewNode->AddDependency(NewDependency, InDependencyType, true);
+				NewDependency->AddReferencer(NewNode);
+			}
+		});
+	}
+}
+
 void FAssetRegistryState::InitializeFromExisting(const TMap<FName, FAssetData*>& AssetDataMap, const TMap<FAssetIdentifier, FDependsNode*>& DependsNodeMap, const TMap<FName, FAssetPackageData*>& AssetPackageDataMap, const FAssetRegistrySerializationOptions& Options, EInitializationMode InInitializationMode)
 {
 	if (InInitializationMode == EInitializationMode::Rebuild)
@@ -240,35 +359,8 @@ void FAssetRegistryState::InitializeFromExisting(const TMap<FName, FAssetData*>&
 			// Filter asset registry tags now
 			const FAssetData& AssetData = *Pair.Value;
 
-			static FName WildcardName(TEXT("*"));
-			const TSet<FName>* AllClassesFilterlist = Options.CookFilterlistTagsByClass.Find(WildcardName);
-			const TSet<FName>* ClassSpecificFilterlist = Options.CookFilterlistTagsByClass.Find(AssetData.AssetClass);
-
-			// Exclude blacklisted tags or include only whitelisted tags, based on how we were configured in ini
 			FAssetDataTagMap LocalTagsAndValues;
-			for (const auto& TagPair : AssetData.TagsAndValues)
-			{
-				const bool bInAllClasseslist = AllClassesFilterlist && (AllClassesFilterlist->Contains(TagPair.Key) || AllClassesFilterlist->Contains(WildcardName));
-				const bool bInClassSpecificlist = ClassSpecificFilterlist && (ClassSpecificFilterlist->Contains(TagPair.Key) || ClassSpecificFilterlist->Contains(WildcardName));
-				if (Options.bUseAssetRegistryTagsWhitelistInsteadOfBlacklist)
-				{
-					// It's a whitelist, only include it if it is in the all classes list or in the class specific list
-					if (bInAllClasseslist || bInClassSpecificlist)
-					{
-						// It is in the whitelist. Keep it.
-						LocalTagsAndValues.Add(TagPair.Key, TagPair.Value);
-					}
-				}
-				else
-				{
-					// It's a blacklist, include it unless it is in the all classes list or in the class specific list
-					if (!bInAllClasseslist && !bInClassSpecificlist)
-					{
-						// It isn't in the blacklist. Keep it.
-						LocalTagsAndValues.Add(TagPair.Key, TagPair.Value);
-					}
-				}
-			}
+			FAssetRegistryState::FilterTags(AssetData.TagsAndValues, LocalTagsAndValues, Options.CookFilterlistTagsByClass.Find(AssetData.AssetClass), Options);
 
 			if (InInitializationMode == EInitializationMode::OnlyUpdateExisting)
 			{
@@ -1510,3 +1602,130 @@ bool FAssetRegistryState::IsFilterValid(const FARFilter& Filter, bool bAllowRecu
 
 	return true;
 }
+
+#if ASSET_REGISTRY_STATE_DUMPING_ENABLED
+
+void FAssetRegistryState::Dump(const TArray<FString>& Arguments, TArray<FString>& OutLines) const
+{
+	struct FAssetRefSorter
+	{
+		FORCEINLINE bool operator()(const FAssetData& A, const FAssetData& B) const
+		{
+			return (A.AssetName < B.AssetName);
+		}
+	};
+
+	if (Arguments.Contains(TEXT("ObjectPath")))
+	{
+		OutLines.Add(FString::Printf(TEXT("CachedAssetsByObjectPath (%d): "), CachedAssetsByObjectPath.Num()));
+		for (const TPair<FName, FAssetData*>& Pair : CachedAssetsByObjectPath)
+		{
+			OutLines.Add(FString::Printf(TEXT("	%s: %s"), *Pair.Key.ToString(), *Pair.Value->PackageName.ToString()));
+		}
+	}
+
+	if (Arguments.Contains(TEXT("PackageName")))
+	{
+		OutLines.Add(FString::Printf(TEXT("CachedAssetsByPackageName (%d): "), CachedAssetsByPackageName.Num()));
+		for (const TPair<FName, TArray<FAssetData*>>& Pair : CachedAssetsByPackageName)
+		{
+			if (Pair.Value.Num() == 0)
+			{
+				continue;
+			}
+
+			TArray<FAssetData*> Datas(Pair.Value);
+			Datas.HeapSort(FAssetRefSorter());
+			FString Tmp;
+			for (const FAssetData* data : Datas)
+			{
+				Tmp.Append(FString::Printf(TEXT("%s, "), *data->PackageName.ToString()));
+			}
+			OutLines.Add(FString::Printf(TEXT("	%s: %s"), *Pair.Key.ToString(), *Tmp));
+		}
+	}
+
+	if (Arguments.Contains(TEXT("Path")))
+	{
+		OutLines.Add(FString::Printf(TEXT("CachedAssetsByPath (%d): "), CachedAssetsByPath.Num()));
+		for (const TPair<FName, TArray<FAssetData*>>& Pair : CachedAssetsByPath)
+		{
+			if (Pair.Value.Num() == 0)
+			{
+				continue;
+			}
+
+			TArray<FAssetData*> Datas(Pair.Value);
+			Datas.HeapSort(FAssetRefSorter());
+			FString Tmp;
+			for (const FAssetData* data : Datas)
+			{
+				Tmp.Append(FString::Printf(TEXT("%s, "), *data->PackageName.ToString()));
+			}
+			OutLines.Add(FString::Printf(TEXT("	%s: %s"), *Pair.Key.ToString(), *Tmp));
+		}
+	}
+
+	if (Arguments.Contains(TEXT("Class")))
+	{
+		OutLines.Add(FString::Printf(TEXT("CachedAssetsByClass (%d): "), CachedAssetsByClass.Num()));
+		for (const TPair<FName, TArray<FAssetData*>>& Pair : CachedAssetsByClass)
+		{
+			if (Pair.Value.Num() == 0)
+			{
+				continue;
+			}
+
+			TArray<FAssetData*> Datas(Pair.Value);
+			Datas.HeapSort(FAssetRefSorter());
+			FString Tmp;
+			for (const FAssetData* data : Datas)
+			{
+				Tmp.Append(FString::Printf(TEXT("%s, "), *data->PackageName.ToString()));
+			}
+			OutLines.Add(FString::Printf(TEXT("	%s: %s"), *Pair.Key.ToString(), *Tmp));
+		}
+	}
+
+	if (Arguments.Contains(TEXT("Tag")))
+	{
+		OutLines.Add(FString::Printf(TEXT("CachedAssetsByTag (%d): "), CachedAssetsByTag.Num()));
+		for (const TPair<FName, TArray<FAssetData*>>& Pair : CachedAssetsByTag)
+		{
+			if (Pair.Value.Num() == 0)
+			{
+				continue;
+			}
+			TArray<FAssetData*> Datas(Pair.Value);
+			Datas.HeapSort(FAssetRefSorter());
+			FString Tmp;
+			for (const FAssetData* data : Datas)
+			{
+				Tmp.Append(FString::Printf(TEXT("%s, "), *data->PackageName.ToString()));
+			}
+			OutLines.Add(FString::Printf(TEXT("	%s: %s"), *Pair.Key.ToString(), *Tmp));
+		}
+	}
+
+	if (Arguments.Contains(TEXT("Dependencies")))
+	{
+		OutLines.Add(FString::Printf(TEXT("CachedDependsNodes (%d): "), CachedDependsNodes.Num()));
+		for (const TPair<FAssetIdentifier, FDependsNode*>& Pair : CachedDependsNodes)
+		{
+			const FDependsNode* Node = Pair.Value;
+			OutLines.Add(FString::Printf(TEXT("	%s: %s (%d connection)"), *Pair.Key.PackageName.ToString(), *Node->GetIdentifier().PackageName.ToString(), Node->GetConnectionCount()));
+		}
+	}
+
+	if (Arguments.Contains(TEXT("PackageData")))
+	{
+		OutLines.Add(FString::Printf(TEXT("CachedPackageData (%d): "), CachedPackageData.Num()));
+		for (const TPair<FName, FAssetPackageData*>& Pair : CachedPackageData)
+		{
+			const FAssetPackageData* PackageData = Pair.Value;
+			OutLines.Add(FString::Printf(TEXT("	%s: %s %d bytes"), *Pair.Key.ToString(), *PackageData->PackageGuid.ToString(), PackageData->DiskSize));
+		}
+	}
+}
+
+#endif // ASSET_REGISTRY_STATE_DUMPING_ENABLED
diff --git a/Engine/Source/Runtime/AssetRegistry/Private/DependsNode.h b/Engine/Source/Runtime/AssetRegistry/Private/DependsNode.h
--- a/Engine/Source/Runtime/AssetRegistry/Private/DependsNode.h
+++ b/Engine/Source/Runtime/AssetRegistry/Private/DependsNode.h
@@ -86,6 +86,10 @@ public:
 		Referencers.Reserve(InNumReferencers);
 	}
 
+	void Reserve(const FDependsNode* Other) {
+		Reserve(Other->HardDependencies.Num(), Other->SoftDependencies.Num(), Other->NameDependencies.Num(), Other->SoftManageDependencies.Num(), Other->HardManageDependencies.Num(), Other->Referencers.Num());
+	}
+
 private:
 
 	/** Iterate over all the separate dependency arrays. Const cast to avoid duplication */
diff --git a/Engine/Source/Runtime/AssetRegistry/Public/AssetRegistryState.h b/Engine/Source/Runtime/AssetRegistry/Public/AssetRegistryState.h
--- a/Engine/Source/Runtime/AssetRegistry/Public/AssetRegistryState.h
+++ b/Engine/Source/Runtime/AssetRegistry/Public/AssetRegistryState.h
@@ -9,6 +9,10 @@
 class FDependsNode;
 struct FARFilter;
 
+#ifndef ASSET_REGISTRY_STATE_DUMPING_ENABLED
+	#define ASSET_REGISTRY_STATE_DUMPING_ENABLED !(UE_BUILD_SHIPPING || UE_BUILD_TEST)
+#endif
+
 /** Load/Save options used to modify how the cache is serialized. These are read out of the AssetRegistry section of Engine.ini and can be changed per platform. */
 struct FAssetRegistrySerializationOptions
 {
@@ -262,6 +266,18 @@ public:
 	void PruneAssetData(const TSet<FName>& RequiredPackages, const TSet<FName>& RemovePackages, const TSet<int32> ChunksToKeep, const FAssetRegistrySerializationOptions& Options);
 	void PruneAssetData(const TSet<FName>& RequiredPackages, const TSet<FName>& RemovePackages, const FAssetRegistrySerializationOptions& Options);
 
+	
+	/**
+	 * Initializes a cache from an existing using a set of filters. This is more efficient than calling InitalizeFromExisting and then PruneAssetData.
+	 * @param ExistingState State to use initialize from
+	 * @param RequiredPackages If set, only these packages will be maintained. If empty it will keep all unless filtered by other parameters
+	 * @param RemovePackages These packages will be removed from the current set
+	 * @param ChunksToKeep The list of chunks that are allowed to remain. Any assets in other chunks are pruned. If empty, all assets are kept regardless of chunk
+	 * @param Options Serialization options to read filter info from
+	 */
+	void InitializeFromExistingAndPrune(const FAssetRegistryState& ExistingState, const TSet<FName>& RequiredPackages, const TSet<FName>& RemovePackages, const TSet<int32> ChunksToKeep, const FAssetRegistrySerializationOptions& Options);
+
+
 	/** Serialize the registry to/from a file, skipping editor only data */
 	bool Serialize(FArchive& Ar, const FAssetRegistrySerializationOptions& Options);
 
@@ -274,6 +290,15 @@ public:
 	/** Returns the number of assets in this state */
 	int32 GetNumAssets() const { return NumAssets; }
 
+#if ASSET_REGISTRY_STATE_DUMPING_ENABLED
+	/**
+	 * Writes out the state in textual form. Use arguments to control which segments to emit.
+	 * @param Arguments List of segments to emit. Possible values: 'ObjectPath', 'PackageName', 'Path', 'Class', 'Tag', 'Dependencies' and 'PackageData'
+	 * @param OutLines Textual representation will be written to this array.
+	 */
+	void Dump(const TArray<FString>& Arguments, TArray<FString>& OutLines) const;
+#endif
+
 private:
 	/** Find the first non-redirector dependency node starting from InDependency. */
 	FDependsNode* ResolveRedirector(FDependsNode* InDependency, TMap<FName, FAssetData*>& InAllowedAssets, TMap<FDependsNode*, FDependsNode*>& InCache);
@@ -290,6 +315,9 @@ private:
 	/** Shrink all contained data structures. */
 	void Shrink();
 
+	/** Filter a set of tags and output a copy of the filtered set. */
+	static void FilterTags(const FAssetDataTagMapSharedView& InTagsAndValues, FAssetDataTagMap& OutTagsAndValues, const TSet<FName>* ClassSpecificFilterlist, const FAssetRegistrySerializationOptions & Options);
+
 	/** Set up the data structures for USE_COMPACT_ASSET_REGISTRY, doesn't really belong here */
 	static void IngestIniSettingsForCompact(TArray<FString>& AsFName, TArray<FString>& AsPathName, TArray<FString>& AsLocText);
 
