A few drive-by optimizations that speeds up asset discovery during editor startup by >2x

* Makes string conversion, FPlatformString::Convert significantly faster by using a constexpr ascii bit
* Introduces FAsciiSet that can be used to optimize and simplify code that loops over strings and checks for null-terminator and a number different special characters.

This pattern is all over the place and consumes 1-2% CPU in many different scenarios, especially in UHT. By using a constexpr FAsciiSet all characters and the null-terminator can be detected with a single branch.

diff --git a/Engine/Source/Runtime/AssetRegistry/Private/AssetDataGatherer.cpp b/Engine/Source/Runtime/AssetRegistry/Private/AssetDataGatherer.cpp
--- a/Engine/Source/Runtime/AssetRegistry/Private/AssetDataGatherer.cpp
+++ b/Engine/Source/Runtime/AssetRegistry/Private/AssetDataGatherer.cpp
@@ -2,9 +2,10 @@
 
 #include "AssetDataGatherer.h"
 #include "HAL/PlatformProcess.h"
+#include "HAL/RunnableThread.h"
+#include "Misc/AsciiSet.h"
 #include "Misc/CommandLine.h"
 #include "Misc/Paths.h"
-#include "HAL/RunnableThread.h"
 #include "Misc/ScopeLock.h"
 #include "AssetRegistryPrivate.h"
 #include "NameTableArchive.h"
@@ -36,27 +37,13 @@ namespace
 		}
 	};
 
-	bool IsValidPackageFileToRead(const FString& Filename)
-	{
-		FString LongPackageName;
-		if (FPackageName::TryConvertFilenameToLongPackageName(Filename, LongPackageName))
-		{
-			// Make sure the path does not contain invalid characters. These packages will not be successfully loaded or read later.
-			for (const TCHAR PackageChar : LongPackageName)
-			{
-				for (const TCHAR InvalidChar : INVALID_LONGPACKAGE_CHARACTERS)
-				{
-					if (PackageChar == InvalidChar)
-					{
-						return false;
-					}
-				}
-			}
-
-			return true;
-		}
+	static constexpr FAsciiSet InvalidLongPackageCharacters(INVALID_LONGPACKAGE_CHARACTERS);
 
-		return false;
+	static bool ConvertToValidLongPackageName(const FString& Filename, /* out */ FString& LongPackageName)
+	{
+		// Make sure the path does not contain invalid characters. These packages will not be successfully loaded or read later.
+		return FPackageName::TryConvertFilenameToLongPackageName(Filename, LongPackageName) &&
+			FAsciiSet::HasNone(*LongPackageName, InvalidLongPackageCharacters);
 	}
 }
 
@@ -137,11 +124,19 @@ uint32 FAssetDataDiscovery::Run()
 			{
 				FScopeLock CritSectionLock(&WorkerThreadCriticalSection);
 
-				// Place all the discovered files into the files to search list
-				DiscoveredPaths.Append(MoveTemp(LocalDiscoveredPathsArray));
+				// Work around for TArray::Append(const TArray&) not growing expontentially.
+				// This causes O(N^2) allocations when continuously appending small arrays,
+				// at least on platforms with allocators that use QuantizeSize().
+				const auto ReserveAndAppend = [](auto& To, auto&& From)
+				{
+					To.Reserve(FMath::RoundUpToPowerOfTwo(To.Num() + From.Num()));
+					To.Append(MoveTemp(From));
+				};
 
-				PriorityDiscoveredFiles.Append(MoveTemp(LocalPriorityFilesToSearch));
-				NonPriorityDiscoveredFiles.Append(MoveTemp(LocalNonPriorityFilesToSearch));
+				// Place all the discovered files into the files to search list
+				ReserveAndAppend(DiscoveredPaths, MoveTemp(LocalDiscoveredPathsArray));
+				ReserveAndAppend(PriorityDiscoveredFiles, MoveTemp(LocalPriorityFilesToSearch));
+				ReserveAndAppend(NonPriorityDiscoveredFiles, MoveTemp(LocalNonPriorityFilesToSearch));
 			}
 		}
 
@@ -178,10 +173,9 @@ uint32 FAssetDataDiscovery::Run()
 		}
 		else if (FPackageName::IsPackageFilename(PackageFilenameStr))
 		{
-			if (IsValidPackageFileToRead(PackageFilenameStr))
+			FString LongPackageNameStr;
+			if (ConvertToValidLongPackageName(PackageFilenameStr, LongPackageNameStr))
 			{
-				const FString LongPackageNameStr = FPackageName::FilenameToLongPackageName(PackageFilenameStr);
-
 				if (IsPriorityFile(PackageFilenameStr))
 				{
 					LocalPriorityFilesToSearch.Add(FDiscoveredPackageFile(PackageFilenameStr, InPackageStatData.ModificationTime));
@@ -791,9 +785,12 @@ void FAssetDataGatherer::AddPathToSearch(const FString& Path)
 void FAssetDataGatherer::AddFilesToSearch(const TArray<FString>& Files)
 {
 	TArray<FString> FilesToAdd;
+	FilesToAdd.Reserve(Files.Num());
+	FString LongPackageName;
 	for (const FString& Filename : Files)
 	{
-		if ( IsValidPackageFileToRead(Filename) )
+		LongPackageName.Reset();
+		if (ConvertToValidLongPackageName(Filename, LongPackageName))
 		{
 			// Add the path to this asset into the list of discovered paths
 			FilesToAdd.Add(Filename);
diff --git a/Engine/Source/Runtime/Core/Private/Tests/Misc/AsciiSetTest.cpp b/Engine/Source/Runtime/Core/Private/Tests/Misc/AsciiSetTest.cpp
new file mode 100644
--- /dev/null
+++ b/Engine/Source/Runtime/Core/Private/Tests/Misc/AsciiSetTest.cpp
@@ -0,0 +1,57 @@
+// Copyright 1998-2019 Epic Games, Inc. All Rights Reserved.
+
+#if WITH_DEV_AUTOMATION_TESTS 
+
+#include "Misc/AsciiSet.h"
+#include "Misc/AutomationTest.h"
+
+IMPLEMENT_SIMPLE_AUTOMATION_TEST(TAsciiSetTest, "System.Core.Misc.AsciiSet", EAutomationTestFlags::ApplicationContextMask | EAutomationTestFlags::SmokeFilter)
+
+bool TAsciiSetTest::RunTest(const FString& Parameters)
+{
+	constexpr FAsciiSet Whitespaces(" \v\f\t\r\n");
+	TestTrue(TEXT("Contains"), Whitespaces.Test(' '));
+	TestTrue(TEXT("Contains"), Whitespaces.Test('\n'));
+	TestFalse(TEXT("Contains"), Whitespaces.Test('a'));
+
+	constexpr FAsciiSet NonWhitespaces = ~Whitespaces;
+	uint32 WhitespaceNum = 0;
+	for (uint8 Char = 0; Char < 128; ++Char)
+	{
+		WhitespaceNum += !!Whitespaces.Test(Char);
+		TestEqual(TEXT("Inverse"), !!Whitespaces.Test(Char), !NonWhitespaces.Test(Char));
+	}
+	TestEqual(TEXT("Num"), WhitespaceNum, 6);
+
+	TestEqual(TEXT("Skip"), FAsciiSet::Skip(TEXT("  \t\tHello world!"), Whitespaces), TEXT("Hello world!"));
+	TestEqual(TEXT("Find"), *FAsciiSet::Find("NonWhitespace\t NonWhitespace", Whitespaces), '\t');
+
+	constexpr FAsciiSet XmlEscapeChars("&<>\"'");
+	TestTrue(TEXT("None"), FAsciiSet::HasNone("No escape chars", XmlEscapeChars));
+	TestFalse(TEXT("Any"), FAsciiSet::HasAny("No escape chars", XmlEscapeChars));
+	TestFalse(TEXT("Only"), FAsciiSet::HasOnly("No escape chars", XmlEscapeChars));
+
+	TestTrue(TEXT("None"), FAsciiSet::HasNone("", XmlEscapeChars));
+	TestFalse(TEXT("Any"), FAsciiSet::HasAny("", XmlEscapeChars));
+	TestTrue(TEXT("Only"), FAsciiSet::HasOnly("", XmlEscapeChars));
+
+	TestFalse(TEXT("None"), FAsciiSet::HasNone("&<>\"'", XmlEscapeChars));
+	TestTrue(TEXT("Any"), FAsciiSet::HasAny("&<>\"'", XmlEscapeChars));
+	TestTrue(TEXT("Only"), FAsciiSet::HasOnly("&<>\"'", XmlEscapeChars));
+
+	TestFalse(TEXT("None"), FAsciiSet::HasNone("&<>\"' and more", XmlEscapeChars));
+	TestTrue(TEXT("Any"), FAsciiSet::HasAny("&<>\"' and more", XmlEscapeChars));
+	TestFalse(TEXT("Only"), FAsciiSet::HasOnly("&<>\"' and more", XmlEscapeChars));
+
+	constexpr FAsciiSet Abc("abc");
+	constexpr FAsciiSet Abcd = Abc + 'd';
+	TestTrue(TEXT("Add"), Abcd.Test('a'));
+	TestTrue(TEXT("Add"), Abcd.Test('b'));
+	TestTrue(TEXT("Add"), Abcd.Test('c'));
+	TestTrue(TEXT("Add"), Abcd.Test('d'));
+	TestFalse(TEXT("Add"), Abcd.Test('e'));
+
+	return true;
+}
+
+#endif //WITH_DEV_AUTOMATION_TESTS
\ No newline at end of file
diff --git a/Engine/Source/Runtime/Core/Public/GenericPlatform/GenericPlatformString.h b/Engine/Source/Runtime/Core/Public/GenericPlatform/GenericPlatformString.h
--- a/Engine/Source/Runtime/Core/Public/GenericPlatform/GenericPlatformString.h
+++ b/Engine/Source/Runtime/Core/Public/GenericPlatform/GenericPlatformString.h
@@ -145,34 +145,29 @@ struct FGenericPlatformString : public FGenericPlatformStricmp
 		DestEncoding*
 	>::Type Convert(DestEncoding* Dest, int32 DestSize, const SourceEncoding* Src, int32 SrcSize, DestEncoding BogusChar = (DestEncoding)'?')
 	{
-		const SourceEncoding* InSrc         = Src;
-		int32                 InSrcSize     = SrcSize;
-		bool                  bInvalidChars = false;
-		while (SrcSize)
+		const int32 Size = DestSize <= SrcSize ? DestSize : SrcSize;
+		bool bInvalidChars = false;
+		for (int I = 0; I < Size; ++I)
 		{
-			if (!DestSize)
-				return nullptr;
-
-			SourceEncoding SrcCh = *Src++;
-			if (CanConvertChar<DestEncoding>(SrcCh))
-			{
-				*Dest++ = (DestEncoding)SrcCh;
-			}
-			else
-			{
-				*Dest++       = BogusChar;
-				bInvalidChars = true;
-			}
-			--SrcSize;
-			--DestSize;
+			SourceEncoding SrcCh = Src[I];
+			Dest[I] = (DestEncoding)SrcCh;
+			bInvalidChars |= !CanConvertChar<DestEncoding>(SrcCh);
 		}
 
 		if (bInvalidChars)
 		{
-			LogBogusChars<DestEncoding>(InSrc, InSrcSize);
+			for (int I = 0; I < Size; ++I)
+			{
+				if (!CanConvertChar<DestEncoding>(Src[I]))
+				{
+					Dest[I] = BogusChar;
+				}
+			}
+
+			LogBogusChars<DestEncoding>(Src, Size);
 		}
 
-		return Dest;
+		return DestSize < SrcSize ? nullptr : Dest + Size;
 	}
 
 	/**
diff --git a/Engine/Source/Runtime/Core/Public/Misc/AsciiSet.h b/Engine/Source/Runtime/Core/Public/Misc/AsciiSet.h
new file mode 100644
--- /dev/null
+++ b/Engine/Source/Runtime/Core/Public/Misc/AsciiSet.h
@@ -0,0 +1,135 @@
+// Copyright 1998-2019 Epic Games, Inc. All Rights Reserved.
+
+#pragma once
+
+#include "Misc/Char.h"
+
+/**
+ * ASCII character bitset useful for fast and readable parsing
+ *
+ * Entirely constexpr. Works with both wide and narrow strings.
+ *
+ * Example use cases:
+ *
+ *   constexpr FAsciiSet WhitespaceCharacters(" \v\f\t\r\n");
+ *   bool bIsWhitespace = WhitespaceCharacters.Contains(MyChar);
+ *   const char* HelloWorld = FAsciiSet::Skip("  \t\tHello world!", WhitespaceCharacters;
+ *
+ *   constexpr FAsciiSet XmlEscapeChars("&<>\"'");
+ *   check(FAsciiSet::HasNone(EscapedXmlString, XmlEscapeChars));
+ */
+class FAsciiSet
+{
+public:
+	template<typename CharType, int N>
+	constexpr FAsciiSet(const CharType(&Chars)[N])
+		: FAsciiSet(StringToBitset(Chars))
+	{}
+
+	/** Returns non-zero if a character is part of the set */
+	template<typename CharType>
+	constexpr FORCEINLINE uint64 Test(CharType Char) const
+	{
+		return TestImpl(TChar<CharType>::ToUnsigned(Char));
+	}
+
+	/** Create new set with specified character in it */
+	constexpr FORCEINLINE FAsciiSet operator+(char Char) const
+	{
+		InitData Bitset = { LoMask, HiMask };
+		Set(Bitset, Char);
+		return FAsciiSet(Bitset);
+	}
+
+	/** Create new set containg inverse set of characters - likely including null-terminator */
+	constexpr FORCEINLINE FAsciiSet operator~() const
+	{
+		return FAsciiSet(~LoMask, ~HiMask);
+	}
+
+	// C string util functions
+
+	/** Scan string for first character in a set or null-terminator */
+	template<class CharType>
+	static constexpr const CharType* Find(const CharType* Str, FAsciiSet Set)
+	{
+		for (FAsciiSet SetOrNil(Set.LoMask | NilMask, Set.HiMask); !SetOrNil.Test(*Str); ++Str);
+
+		return Str;
+	}
+
+	/** Scan string for first character not in a set or null-terminator */
+	template<typename CharType>
+	static constexpr const CharType* Skip(const CharType* Str, FAsciiSet Set)
+	{
+		return Find(Str, ~Set);
+	}
+
+	/** Test if string contains any character in set */
+	template<typename CharType>
+	static constexpr bool HasAny(const CharType* Str, FAsciiSet Set)
+	{
+		return *Find(Str, Set) != '\0';
+	}
+
+	/** Test if string contains no character in set */
+	template<typename CharType>
+	static constexpr bool HasNone(const CharType* Str, FAsciiSet Set)
+	{
+		return *Find(Str, Set) == '\0';
+	}
+
+	/** Test if string contains any character outside of set */
+	template<typename CharType>
+	static constexpr bool HasOnly(const CharType* Str, FAsciiSet Set)
+	{
+		return *Skip(Str, Set) == '\0';
+	}
+
+private:
+	// Work-around for constexpr limitations
+	struct InitData { uint64 Lo, Hi; };
+	static constexpr uint64 NilMask = uint64(1) << '\0';
+
+	template<typename CharType>
+	static constexpr FORCEINLINE void Set(InitData& Bitset, CharType Char)
+	{
+		uint64 IsLo = uint64(0) - (Char >> 6 == 0);
+		uint64 IsHi = uint64(0) - (Char >> 6 == 1);
+		uint64 Bit = uint64(1) << uint8(Char & 0x3f);
+
+		Bitset.Lo |= Bit & IsLo;
+		Bitset.Hi |= Bit & IsHi;
+	}
+
+	constexpr FORCEINLINE uint64 TestImpl(uint32 Char) const
+	{
+		uint64 IsLo = uint64(0) - (Char >> 6 == 0);
+		uint64 IsHi = uint64(0) - (Char >> 6 == 1);
+		uint64 Bit = uint64(1) << (Char & 0x3f);
+
+		return (Bit & IsLo & LoMask) | (Bit & IsHi & HiMask);
+	}
+
+	template<typename CharType, int N>
+	static constexpr InitData StringToBitset(const CharType(&Chars)[N])
+	{
+		InitData Bitset = { 0, 0 };
+		for (int I = 0; I < N - 1; ++I)
+		{
+			Set(Bitset, Chars[I]);
+		}
+
+		return Bitset;
+	}
+
+	constexpr FAsciiSet(InitData Bitset) 
+		: LoMask(Bitset.Lo), HiMask(Bitset.Hi)
+	{}
+
+	constexpr FAsciiSet(uint64 Lo, uint64 Hi)
+		: LoMask(Lo), HiMask(Hi)
+	{}
+
+	uint64 LoMask, HiMask;
+};
