Adding a simpler getter for the asset registry.

diff --git a/Engine/Source/Runtime/AssetRegistry/Public/AssetRegistryModule.h b/Engine/Source/Runtime/AssetRegistry/Public/AssetRegistryModule.h
--- a/Engine/Source/Runtime/AssetRegistry/Public/AssetRegistryModule.h
+++ b/Engine/Source/Runtime/AssetRegistry/Public/AssetRegistryModule.h
@@ -30,6 +30,13 @@ public:
 	/** Gets the asset registry singleton */
 	virtual IAssetRegistry& Get() const;
 
+	/** Gets the asset registry singleton */
+	static IAssetRegistry& GetRegistry()
+	{
+		FAssetRegistryModule& AssetRegistryModule = FModuleManager::LoadModuleChecked<FAssetRegistryModule>(AssetRegistryConstants::ModuleName);
+		return AssetRegistryModule.Get();
+	}
+
 	/** Tick the asset registry with the supplied timestep */
 	static void TickAssetRegistry(float DeltaTime)
 	{
