Using a local static for the inline text, otherwise live editing the code was having trouble reloadng due to not being able to find the empty string module.

diff --git a/Engine/Source/Runtime/AssetRegistry/Public/AssetRegistryModule.h b/Engine/Source/Runtime/AssetRegistry/Public/AssetRegistryModule.h
--- a/Engine/Source/Runtime/AssetRegistry/Public/AssetRegistryModule.h
+++ b/Engine/Source/Runtime/AssetRegistry/Public/AssetRegistryModule.h
@@ -33,7 +33,8 @@ public:
 	/** Gets the asset registry singleton */
 	static IAssetRegistry& GetRegistry()
 	{
-		FAssetRegistryModule& AssetRegistryModule = FModuleManager::LoadModuleChecked<FAssetRegistryModule>(AssetRegistryConstants::ModuleName);
+		static FName ModuleName("AssetRegistry");
+		FAssetRegistryModule& AssetRegistryModule = FModuleManager::LoadModuleChecked<FAssetRegistryModule>(ModuleName);
 		return AssetRegistryModule.Get();
 	}
 
