AssetRegistry optimizations
#rnx

diff --git a/Engine/Source/Runtime/AssetRegistry/Private/AssetRegistry.cpp b/Engine/Source/Runtime/AssetRegistry/Private/AssetRegistry.cpp
--- a/Engine/Source/Runtime/AssetRegistry/Private/AssetRegistry.cpp
+++ b/Engine/Source/Runtime/AssetRegistry/Private/AssetRegistry.cpp
@@ -629,6 +629,9 @@ bool UAssetRegistryImpl::GetAssets(const FARFilter& InFilter, TArray<FAssetData>
 		const int32 NumFilterClasses = FilterClassNames.Num();
 		const int32 NumFilterObjectPaths = FilterObjectPaths.Num();
 
+		// Reusable structures to avoid memory allocations
+		TArray<UObject::FAssetRegistryTag> ObjectTags;
+
 		auto FilterInMemoryObjectLambda = [&](const UObject* Obj)
 		{
 			if (Obj->IsAsset())
@@ -654,7 +657,7 @@ bool UAssetRegistryImpl::GetAssets(const FARFilter& InFilter, TArray<FAssetData>
 				// Object Path
 				if (NumFilterObjectPaths)
 				{
-					const FName ObjectPath = FName(*Obj->GetPathName());
+					const FName ObjectPath = FName(*Obj->GetPathName(), FNAME_Find);
 					if (!FilterObjectPaths.Contains(ObjectPath))
 					{
 						return;
@@ -669,7 +672,7 @@ bool UAssetRegistryImpl::GetAssets(const FARFilter& InFilter, TArray<FAssetData>
 				}
 
 				// Tags and values
-				TArray<UObject::FAssetRegistryTag> ObjectTags;
+				check(ObjectTags.Num() == 0);
 				Obj->GetAssetRegistryTags(ObjectTags);
 				if (Filter.TagsAndValues.Num())
 				{
@@ -713,9 +716,10 @@ bool UAssetRegistryImpl::GetAssets(const FARFilter& InFilter, TArray<FAssetData>
 						TagMap.Add(AssetRegistryTag.Name, AssetRegistryTag.Value);
 					}
 				}
+				ObjectTags.Reset();
 
 				// This asset is in memory and passes all filters
-				FAssetData* AssetData = new (OutAssetData)FAssetData(PackageName, PackagePath, Obj->GetFName(), Obj->GetClass()->GetFName(), TagMap, InMemoryPackage->GetChunkIDs(), InMemoryPackage->GetPackageFlags());
+				OutAssetData.Emplace(FAssetData(PackageName, PackagePath, Obj->GetFName(), Obj->GetClass()->GetFName(), MoveTemp(TagMap), InMemoryPackage->GetChunkIDs(), InMemoryPackage->GetPackageFlags()));
 			}
 		};
 
@@ -1295,6 +1299,7 @@ void UAssetRegistryImpl::ExpandRecursiveFilter(const FARFilter& InFilter, FARFil
 
 	ExpandedFilter = InFilter;
 
+	FilterPackagePaths.Reserve(NumFilterPackagePaths);
 	for (int32 PathIdx = 0; PathIdx < NumFilterPackagePaths; ++PathIdx)
 	{
 		FilterPackagePaths.Add(InFilter.PackagePaths[PathIdx]);
@@ -1331,6 +1336,7 @@ void UAssetRegistryImpl::ExpandRecursiveFilter(const FARFilter& InFilter, FARFil
 	}
 	else
 	{
+		FilterClassNames.Reserve(NumFilterClasses);
 		for (int32 ClassIdx = 0; ClassIdx < NumFilterClasses; ++ClassIdx)
 		{
 			FilterClassNames.Add(InFilter.ClassNames[ClassIdx]);
@@ -2807,11 +2813,13 @@ void UAssetRegistryImpl::GetSubClasses(const TArray<FName>& InClassNames, const
 {
 	UpdateTemporaryCaches();
 
+	TSet<FName> ProcessedClassNames;
+
 	for (FName ClassName : InClassNames)
 	{
 		// Now find all subclass names
-		TSet<FName> ProcessedClassNames;
 		GetSubClasses_Recursive(ClassName, SubClassNames, ProcessedClassNames, TempReverseInheritanceMap, ExcludedClassNames);
+		ProcessedClassNames.Reset();
 	}
 
 	ClearTemporaryCaches();
diff --git a/Engine/Source/Runtime/Engine/Classes/Engine/AssetManagerTypes.h b/Engine/Source/Runtime/Engine/Classes/Engine/AssetManagerTypes.h
--- a/Engine/Source/Runtime/Engine/Classes/Engine/AssetManagerTypes.h
+++ b/Engine/Source/Runtime/Engine/Classes/Engine/AssetManagerTypes.h
@@ -39,20 +39,20 @@ struct FPrimaryAssetRules
 	UPROPERTY(EditAnywhere, Category = Rules)
 	int32 Priority;
 
-	/** If true, this rule will apply to all referenced Secondary Assets recursively, as long as they are not managed by a higher-priority Primary Asset. */
-	UPROPERTY(EditAnywhere, Category = Rules)
-	bool bApplyRecursively;
-
 	/** Assets will be put into this Chunk ID specifically, if set to something other than -1. The default Chunk is Chunk 0. */
 	UPROPERTY(EditAnywhere, Category = Rules, meta = (DisplayName = "Chunk ID"))
 	int32 ChunkId;
 
+	/** If true, this rule will apply to all referenced Secondary Assets recursively, as long as they are not managed by a higher-priority Primary Asset. */
+	UPROPERTY(EditAnywhere, Category = Rules)
+	bool bApplyRecursively;
+
 	/** Rule describing when this asset should be cooked. */
 	UPROPERTY(EditAnywhere, Category = Rules)
 	EPrimaryAssetCookRule CookRule;
 
 	FPrimaryAssetRules() 
-		: Priority(-1), bApplyRecursively(true), ChunkId(-1), CookRule(EPrimaryAssetCookRule::Unknown)
+		: Priority(-1), ChunkId(-1), bApplyRecursively(true), CookRule(EPrimaryAssetCookRule::Unknown)
 	{
 	}
 
diff --git a/Engine/Source/Runtime/Engine/Private/AssetManager.cpp b/Engine/Source/Runtime/Engine/Private/AssetManager.cpp
--- a/Engine/Source/Runtime/Engine/Private/AssetManager.cpp
+++ b/Engine/Source/Runtime/Engine/Private/AssetManager.cpp
@@ -407,9 +407,9 @@ int32 UAssetManager::ScanPathsForPrimaryAssets(FPrimaryAssetType PrimaryAssetTyp
 		int32 DotIndex = INDEX_NONE;
 		if (Path.FindChar('.', DotIndex))
 		{
-			FString PackageName = FPackageName::ObjectPathToPackageName(Path);
+			FString PackageName = Path.Mid(0, DotIndex); //avoid re-searching for index inside FPackageName::ObjectPathToPackageName
 
-			PackageNames.AddUnique(PackageName);
+			PackageNames.AddUnique(MoveTemp(PackageName));
 		}
 		else
 		{
@@ -524,14 +524,15 @@ int32 UAssetManager::ScanPathsForPrimaryAssets(FPrimaryAssetType PrimaryAssetTyp
 			else
 			{
 				// Warn that 'Foo' conflicts with 'Bar', but only once per conflict
-				static TSet<FString> IssuedWarnings;
-
-				FString ConflictMsg = FString::Printf(TEXT("Ignoring PrimaryAssetType %s - Conflicts with %s"), *PrimaryAssetType.ToString(), *PrimaryAssetId.PrimaryAssetType.ToString());
+				static TSet<TPair<FPrimaryAssetType, FPrimaryAssetType>> IssuedWarnings;
 
-				if (!IssuedWarnings.Contains(ConflictMsg))
+				TTuple<FPrimaryAssetType, FPrimaryAssetType> ConflictPair(PrimaryAssetType, PrimaryAssetId.PrimaryAssetType);
+				if (!IssuedWarnings.Contains(ConflictPair))
 				{
+					FString ConflictMsg = FString::Printf(TEXT("Ignoring PrimaryAssetType %s - Conflicts with %s"), *PrimaryAssetType.ToString(), *PrimaryAssetId.PrimaryAssetType.ToString());
+
 					UE_LOG(LogAssetManager, Display, TEXT("%s"), *ConflictMsg);
-					IssuedWarnings.Add(ConflictMsg);
+					IssuedWarnings.Add(ConflictPair);
 				}
 			}
 			continue;
@@ -2798,7 +2799,7 @@ bool UAssetManager::IsPathExcludedFromScan(const FString& Path) const
 {
 	const UAssetManagerSettings& Settings = GetSettings();
 
-	for (const FDirectoryPath& ExcludedPath : GetSettings().DirectoriesToExclude)
+	for (const FDirectoryPath& ExcludedPath : Settings.DirectoriesToExclude)
 	{
 		if (Path.Contains(ExcludedPath.Path))
 		{
