Fix ObjectTags not getting reset when early out occurs.
#rnx

diff --git a/Engine/Source/Runtime/AssetRegistry/Private/AssetRegistry.cpp b/Engine/Source/Runtime/AssetRegistry/Private/AssetRegistry.cpp
--- a/Engine/Source/Runtime/AssetRegistry/Private/AssetRegistry.cpp
+++ b/Engine/Source/Runtime/AssetRegistry/Private/AssetRegistry.cpp
@@ -703,6 +703,7 @@ bool UAssetRegistryImpl::GetAssets(const FARFilter& InFilter, TArray<FAssetData>
 
 					if (!bMatch)
 					{
+						ObjectTags.Reset();
 						return;
 					}
 				}
