After a new asset registry is merged into an existing one, fire all the new assets through the AssetAdded delegate so that any registered systems can pick them up and use them (specifically the localization system in this case)

diff --git a/Engine/Source/Runtime/AssetRegistry/Private/AssetRegistry.cpp b/Engine/Source/Runtime/AssetRegistry/Private/AssetRegistry.cpp
--- a/Engine/Source/Runtime/AssetRegistry/Private/AssetRegistry.cpp
+++ b/Engine/Source/Runtime/AssetRegistry/Private/AssetRegistry.cpp
@@ -1730,6 +1730,15 @@ void UAssetRegistryImpl::AppendState(const FAssetRegistryState& InState)
 {
 	State.InitializeFromExisting(InState, SerializationOptions, FAssetRegistryState::EInitializationMode::Append);
 	CachePathsFromState(InState);
+
+	TArray<FAssetData> Assets;
+	InState.GetAllAssets(TSet<FName>(), Assets);
+	for (const FAssetData& Asset : Assets)
+	{
+		// Let subscribers know that the new asset was added to the registry
+		AssetAddedEvent.Broadcast(Asset);
+	}
+
 }
 
 void UAssetRegistryImpl::CachePathsFromState(const FAssetRegistryState& InState)
