#UE4 When packages are deleted, either in memory or due to a change on the disk, while there are still references to it, a depends node will be created for it that will show the reference to the deleted asset. This is what happens when you load the dependency graph from scratch, so this change just keeps the graph consistent with that if you deleted something but did not yet restart

diff --git a/Engine/Source/Runtime/AssetRegistry/Private/AssetRegistry.cpp b/Engine/Source/Runtime/AssetRegistry/Private/AssetRegistry.cpp
--- a/Engine/Source/Runtime/AssetRegistry/Private/AssetRegistry.cpp
+++ b/Engine/Source/Runtime/AssetRegistry/Private/AssetRegistry.cpp
@@ -2367,12 +2367,43 @@ void UAssetRegistryImpl::RemovePackageData(const FName PackageName)
 	TArray<FAssetData*>* PackageAssetsPtr = State.CachedAssetsByPackageName.Find(PackageName);
 	if (PackageAssetsPtr && PackageAssetsPtr->Num() > 0)
 	{
+		// If there were any referencers, re-add them to a new empty dependency node, as it would be when the referencers are loaded from disk
+		TArray<FName> SoftReferencers;
+		TArray<FName> HardReferencers;
+		GetReferencers(PackageName, SoftReferencers, EAssetRegistryDependencyType::Soft);
+		GetReferencers(PackageName, HardReferencers, EAssetRegistryDependencyType::Hard);
+
 		// Copy the array since RemoveAssetData may re-allocate it!
 		TArray<FAssetData*> PackageAssets = *PackageAssetsPtr;
 		for (FAssetData* PackageAsset : PackageAssets)
 		{
 			RemoveAssetData(PackageAsset);
 		}
+
+		// See if we have to readd the dependency now
+		if (SoftReferencers.Num() > 0 || HardReferencers.Num() > 0)
+		{
+			FDependsNode* NewNode = State.CreateOrFindDependsNode(PackageName);
+
+			auto ReAddDependency = [this, NewNode](const FName& Referencer, EAssetRegistryDependencyType::Type RefType) {
+				FDependsNode* ReferencerNode = State.CreateOrFindDependsNode(Referencer);
+				if (ReferencerNode != nullptr)
+				{
+					ReferencerNode->AddDependency(NewNode, RefType);
+					NewNode->AddReferencer(ReferencerNode);
+				}
+			};
+
+			for (const FName& SoftRef : SoftReferencers)
+			{
+				ReAddDependency(SoftRef, EAssetRegistryDependencyType::Soft);
+			}
+
+			for (const FName& HardRef : HardReferencers)
+			{
+				ReAddDependency(HardRef, EAssetRegistryDependencyType::Hard);
+			}
+		}
 	}
 }
 
