Back out changelist 6952307: Optimize FAssetData::IsAssetLoaded() by removing temporary string allocation and using FindObjectFast instead of FindObjectSafe
The optimization is only valid under the assumption that only top level objects in a package can have FAssetData, which obviously is not true.

diff --git a/Engine/Source/Runtime/AssetRegistry/Public/AssetData.h b/Engine/Source/Runtime/AssetRegistry/Public/AssetData.h
--- a/Engine/Source/Runtime/AssetRegistry/Public/AssetData.h
+++ b/Engine/Source/Runtime/AssetRegistry/Public/AssetData.h
@@ -363,12 +363,7 @@ public:
 	/** Returns true if the asset is loaded */
 	bool IsAssetLoaded() const
 	{
-		if (IsValid())
-		{
-			UPackage* Package = FindObjectFast<UPackage>(nullptr, PackageName);
-			return (Package != NULL) && (FindObjectFast<UObject>(Package, AssetName) != NULL);
-		}
-		return false;
+		return IsValid() && FindObjectSafe<UObject>(NULL, *ObjectPath.ToString()) != NULL;
 	}
 
 	/** Prints the details of the asset to the log */
