Exposed ARFilter UPropertys for editing in defaults.

diff --git a/Engine/Source/Runtime/AssetRegistry/Public/ARFilter.h b/Engine/Source/Runtime/AssetRegistry/Public/ARFilter.h
--- a/Engine/Source/Runtime/AssetRegistry/Public/ARFilter.h
+++ b/Engine/Source/Runtime/AssetRegistry/Public/ARFilter.h
@@ -16,38 +16,38 @@ struct FARFilter
 	GENERATED_BODY()
 
 	/** The filter component for package names */
-	UPROPERTY(transient, BlueprintReadWrite, Category=AssetRegistry)
+	UPROPERTY(EditDefaultsOnly, BlueprintReadWrite, Category=AssetRegistry)
 	TArray<FName> PackageNames;
 
 	/** The filter component for package paths */
-	UPROPERTY(transient, BlueprintReadWrite, Category=AssetRegistry)
+	UPROPERTY(EditDefaultsOnly, BlueprintReadWrite, Category=AssetRegistry)
 	TArray<FName> PackagePaths;
 
 	/** The filter component containing specific object paths */
-	UPROPERTY(transient, BlueprintReadWrite, Category=AssetRegistry)
+	UPROPERTY(EditDefaultsOnly, BlueprintReadWrite, Category=AssetRegistry)
 	TArray<FName> ObjectPaths;
 
 	/** The filter component for class names. Instances of the specified classes, but not subclasses (by default), will be included. Derived classes will be included only if bRecursiveClasses is true. */
-	UPROPERTY(transient, BlueprintReadWrite, Category=AssetRegistry)
+	UPROPERTY(EditDefaultsOnly, BlueprintReadWrite, Category=AssetRegistry)
 	TArray<FName> ClassNames;
 
 	/** The filter component for properties marked with the AssetRegistrySearchable flag */
 	TMultiMap<FName, TOptional<FString>> TagsAndValues;
 
 	/** Only if bRecursiveClasses is true, the results will exclude classes (and subclasses) in this list */
-	UPROPERTY(transient, BlueprintReadWrite, Category=AssetRegistry)
+	UPROPERTY(EditDefaultsOnly, BlueprintReadWrite, Category=AssetRegistry)
 	TSet<FName> RecursiveClassesExclusionSet;
 
 	/** If true, PackagePath components will be recursive */
-	UPROPERTY(transient, BlueprintReadWrite, Category=AssetRegistry)
+	UPROPERTY(EditDefaultsOnly, BlueprintReadWrite, Category=AssetRegistry)
 	bool bRecursivePaths;
 
 	/** If true, subclasses of ClassNames will also be included and RecursiveClassesExclusionSet will be excluded. */
-	UPROPERTY(transient, BlueprintReadWrite, Category=AssetRegistry)
+	UPROPERTY(EditDefaultsOnly, BlueprintReadWrite, Category=AssetRegistry)
 	bool bRecursiveClasses;
 
 	/** If true, only on-disk assets will be returned. Be warned that this is rarely what you want and should only be used for performance reasons */
-	UPROPERTY(transient, BlueprintReadWrite, Category=AssetRegistry)
+	UPROPERTY(EditDefaultsOnly, BlueprintReadWrite, Category=AssetRegistry)
 	bool bIncludeOnlyOnDiskAssets;
 
 	FARFilter()
