Adding clarification to some asset registry comments

#rnx

diff --git a/Engine/Source/Runtime/AssetRegistry/Public/IAssetRegistry.h b/Engine/Source/Runtime/AssetRegistry/Public/IAssetRegistry.h
--- a/Engine/Source/Runtime/AssetRegistry/Public/IAssetRegistry.h
+++ b/Engine/Source/Runtime/AssetRegistry/Public/IAssetRegistry.h
@@ -113,7 +113,7 @@ public:
 	/**
 	 * Does the given path contain assets, optionally also testing sub-paths?
 	 *
-	 * @param PackagePath the path to query asset data in
+	 * @param PackagePath the path to query asset data in (eg, /Game/MyFolder)
 	 * @param bRecursive if true, the supplied path will be tested recursively
 	 */
 	UFUNCTION(BlueprintCallable, BlueprintPure=false, Category = "AssetRegistry")
@@ -122,7 +122,7 @@ public:
 	/**
 	 * Gets asset data for the assets in the package with the specified package name
 	 *
-	 * @param PackageName the package name for the requested assets
+	 * @param PackageName the package name for the requested assets (eg, /Game/MyFolder/MyAsset)
 	 * @param OutAssetData the list of assets in this path
 	 */
 	UFUNCTION(BlueprintCallable, BlueprintPure=false, Category="AssetRegistry")
@@ -131,7 +131,7 @@ public:
 	/**
 	 * Gets asset data for all assets in the supplied folder path
 	 *
-	 * @param PackagePath the path to query asset data in
+	 * @param PackagePath the path to query asset data in (eg, /Game/MyFolder)
 	 * @param OutAssetData the list of assets in this path
 	 * @param bRecursive if true, all supplied paths will be searched recursively
 	 */
@@ -206,7 +206,7 @@ public:
 	/**
 	 * Gets a list of paths to objects that are referenced by the supplied package. (On disk references ONLY)
 	 *
-	 * @param PackageName		the name of the package for which to gather dependencies
+	 * @param PackageName		the name of the package for which to gather dependencies (eg, /Game/MyFolder/MyAsset)
 	 * @param OutDependencies	a list of packages that are referenced by the package whose path is PackageName
 	 * @param InDependencyType	which kinds of dependency to include in the output list
 	 */
@@ -215,7 +215,7 @@ public:
 	/**
 	 * Gets a list of paths to objects that are referenced by the supplied package. (On disk references ONLY)
 	 *
-	 * @param PackageName		the name of the package for which to gather dependencies
+	 * @param PackageName		the name of the package for which to gather dependencies (eg, /Game/MyFolder/MyAsset)
 	 * @param DependencyOptions	which kinds of dependencies to include in the output list
 	 * @param OutDependencies	a list of packages that are referenced by the package whose path is PackageName
 	 */
@@ -228,7 +228,7 @@ public:
 	/**
 	 * Gets a list of packages and searchable names that reference the supplied package or name. (On disk references ONLY)
 	 *
-	 * @param AssetIdentifier	the name of the package/name for which to gather dependencies
+	 * @param AssetIdentifier	the name of the package/name for which to gather dependencies (eg, /Game/MyFolder/MyAsset)
 	 * @param OutReferencers	a list of things that reference AssetIdentifier
 	 * @param InReferenceType	which kinds of reference to include in the output list
 	 */
@@ -237,7 +237,7 @@ public:
 	/**
 	 * Gets a list of packages that reference the supplied package. (On disk references ONLY)
 	 *
-	 * @param PackageName		the name of the package for which to gather dependencies
+	 * @param PackageName		the name of the package for which to gather dependencies (eg, /Game/MyFolder/MyAsset)
 	 * @param OutReferencers	a list of packages that reference the package whose path is PackageName
 	 * @param InReferenceType	which kinds of reference to include in the output list
 	 */
@@ -246,7 +246,7 @@ public:
 	/**
 	 * Gets a list of packages that reference the supplied package. (On disk references ONLY)
 	 *
-	 * @param PackageName		the name of the package for which to gather dependencies
+	 * @param PackageName		the name of the package for which to gather dependencies (eg, /Game/MyFolder/MyAsset)
 	 * @param ReferenceOptions	which kinds of references to include in the output list
 	 * @param OutReferencers	a list of packages that reference the package whose path is PackageName
 	 */
