Misc serialization optimizations for tagged properties, name maps and soft object paths

diff --git a/Engine/Source/Runtime/AssetRegistry/Private/NameTableArchive.cpp b/Engine/Source/Runtime/AssetRegistry/Private/NameTableArchive.cpp
--- a/Engine/Source/Runtime/AssetRegistry/Private/NameTableArchive.cpp
+++ b/Engine/Source/Runtime/AssetRegistry/Private/NameTableArchive.cpp
@@ -4,10 +4,17 @@
 #include "HAL/FileManager.h"
 #include "AssetRegistryPrivate.h"
 
+class FNameTableErrorArchive : public FArchive
+{
+public:
+	FNameTableErrorArchive()
+	{
+		SetError();
+	}
+};
+static FNameTableErrorArchive GNameTableErrorArchive;
+
 FNameTableArchiveReader::FNameTableArchiveReader(int32 SerializationVersion, const FString& Filename)
-	: FArchive()
-	, ProxyAr(nullptr)
-	, FileAr(nullptr)
 {
 	this->SetIsLoading(true);
 
@@ -36,6 +43,7 @@ FNameTableArchiveReader::FNameTableArchiveReader(int32 SerializationVersion, con
 	}
 	
 	// If we got here it failed to load properly
+	ProxyAr = &GNameTableErrorArchive;
 	SetError();
 }
 
@@ -48,17 +56,15 @@ FNameTableArchiveReader::FNameTableArchiveReader(FArchive& WrappedArchive)
 
 	if (!SerializeNameMap())
 	{
+		ProxyAr = &GNameTableErrorArchive;
 		SetError();
 	}
 }
 
 FNameTableArchiveReader::~FNameTableArchiveReader()
 {
-	if (FileAr)
-	{
-		delete FileAr;
-		FileAr = nullptr;
-	}
+	delete FileAr;
+	FileAr = nullptr;
 }
 
 bool FNameTableArchiveReader::SerializeNameMap()
@@ -66,7 +72,7 @@ bool FNameTableArchiveReader::SerializeNameMap()
 	int64 NameOffset = 0;
 	*this << NameOffset;
 
-	if (IsError() || NameOffset > TotalSize())
+	if (NameOffset > TotalSize())
 	{
 		// The file was corrupted. Return false to fail to load the cache an thus regenerate it.
 		return false;
@@ -75,16 +81,11 @@ bool FNameTableArchiveReader::SerializeNameMap()
 	if( NameOffset > 0 )
 	{
 		int64 OriginalOffset = Tell();
-		Seek( NameOffset );
+		ProxyAr->Seek( NameOffset );
 
 		int32 NameCount = 0;
 		*this << NameCount;
-
-		if (IsError())
-		{
-			return false;
-		}
-
+		
 		NameMap.Reserve(NameCount);
 		for ( int32 NameMapIdx = 0; NameMapIdx < NameCount; ++NameMapIdx )
 		{
@@ -100,7 +101,7 @@ bool FNameTableArchiveReader::SerializeNameMap()
 			NameMap.Add(FName(NameEntry).GetDisplayIndex());
 		}
 
-		Seek( OriginalOffset );
+		ProxyAr->Seek(OriginalOffset);
 	}
 
 	return true;
@@ -108,20 +109,18 @@ bool FNameTableArchiveReader::SerializeNameMap()
 
 void FNameTableArchiveReader::Serialize(void* V, int64 Length)
 {
-	if (ProxyAr && !IsError())
-	{
-		ProxyAr->Serialize(V, Length);
+	ProxyAr->Serialize(V, Length);
 
-		if (ProxyAr->IsError())
-		{
-			SetError();
-		}
+	if (ProxyAr->IsError())
+	{
+		ProxyAr = &GNameTableErrorArchive;
+		SetError();
 	}
 }
 
 bool FNameTableArchiveReader::Precache(int64 PrecacheOffset, int64 PrecacheSize)
 {
-	if (ProxyAr && !IsError())
+	if (!IsError())
 	{
 		return ProxyAr->Precache(PrecacheOffset, PrecacheSize);
 	}
@@ -131,7 +130,7 @@ bool FNameTableArchiveReader::Precache(int64 PrecacheOffset, int64 PrecacheSize)
 
 void FNameTableArchiveReader::Seek(int64 InPos)
 {
-	if (ProxyAr && !IsError())
+	if (!IsError())
 	{
 		ProxyAr->Seek(InPos);
 	}
@@ -139,73 +138,52 @@ void FNameTableArchiveReader::Seek(int64 InPos)
 
 int64 FNameTableArchiveReader::Tell()
 {
-	if (ProxyAr)
-	{
-		return ProxyAr->Tell();
-	}
-
-	return 0;
+	return ProxyAr->Tell();
 }
 
 int64 FNameTableArchiveReader::TotalSize()
 {
-	if (ProxyAr)
-	{
-		return ProxyAr->TotalSize();
-	}
-
-	return 0;
+	return ProxyAr->TotalSize();
 }
 
 const FCustomVersionContainer& FNameTableArchiveReader::GetCustomVersions() const
 {
-	if (ProxyAr)
-	{
-		return ProxyAr->GetCustomVersions();
-	}
-	return FArchive::GetCustomVersions();
+	return ProxyAr->GetCustomVersions();
 }
 
 void FNameTableArchiveReader::SetCustomVersions(const FCustomVersionContainer& NewVersions)
 {
-	if (ProxyAr)
-	{
-		ProxyAr->SetCustomVersions(NewVersions);
-	}
+	ProxyAr->SetCustomVersions(NewVersions);
 }
 
 void FNameTableArchiveReader::ResetCustomVersions()
 {
-	if (ProxyAr)
-	{
-		ProxyAr->ResetCustomVersions();
-	}
+	ProxyAr->ResetCustomVersions();
 }
 
-FArchive& FNameTableArchiveReader::operator<<( FName& Name )
+FArchive& FNameTableArchiveReader::operator<<(FName& OutName)
 {
 	int32 NameIndex;
 	FArchive& Ar = *this;
 	Ar << NameIndex;
 
-	if( !NameMap.IsValidIndex(NameIndex) )
+	if (NameMap.IsValidIndex(NameIndex))
 	{
-		UE_LOG(LogAssetRegistry, Error, TEXT("Bad name index reading cache %i/%i"), NameIndex, NameMap.Num() );
-		SetError();
-	}
+		FNameEntryId MappedName = NameMap[NameIndex];
 
-	int32 Number;
-	Ar << Number;
+		int32 Number;
+		Ar << Number;
 
-	FNameEntryId MappedName = NameMap.IsValidIndex(NameIndex) ? NameMap[NameIndex] : FNameEntryId();
-	if (!MappedName)
-	{
-		Name = FName();
+		OutName = FName::CreateFromDisplayId(MappedName, MappedName ? Number : 0);
 	}
 	else
 	{
-		// simply create the name from the NameMap's name and the serialized instance number
-		Name = FName::CreateFromDisplayId(MappedName, Number);
+		UE_LOG(LogAssetRegistry, Error, TEXT("Bad name index reading cache %i/%i"), NameIndex, NameMap.Num());
+
+		ProxyAr = &GNameTableErrorArchive;
+		SetError();
+
+		OutName = FName();
 	}
 
 	return *this;
diff --git a/Engine/Source/Runtime/CoreUObject/Private/UObject/LinkerLoad.cpp b/Engine/Source/Runtime/CoreUObject/Private/UObject/LinkerLoad.cpp
--- a/Engine/Source/Runtime/CoreUObject/Private/UObject/LinkerLoad.cpp
+++ b/Engine/Source/Runtime/CoreUObject/Private/UObject/LinkerLoad.cpp
@@ -822,7 +822,6 @@ FLinkerLoad::FLinkerLoad(UPackage* InParent, const TCHAR* InFilename, uint32 InL
 , StructuredArchiveFormatter(nullptr)
 , Loader(nullptr)
 , AsyncRoot(nullptr)
-, NameMapIndex(0)
 , GatherableTextDataMapIndex(0)
 , ImportMapIndex(0)
 , ExportMapIndex(0)
@@ -1407,15 +1406,16 @@ FLinkerLoad::ELinkerStatus FLinkerLoad::SerializeNameMap()
 	// before any serialization. @todo async, @todo seamless: this could be spread out across name,
 	// import and export maps if the package file summary contained more detailed information on
 	// serialized size of individual entries.
-	bool bFinishedPrecaching = true;
-
-	if( NameMapIndex == 0 && Summary.NameCount > 0 )
+	const int32 NameCount = Summary.NameCount;
+	if (NameMap.Num() == 0 && NameCount > 0)
 	{
 		Seek(Summary.NameOffset);
 
 		// Make sure there is something to precache first.
-		if( Summary.TotalHeaderSize > 0 )
+		if (Summary.TotalHeaderSize > 0)
 		{
+			bool bFinishedPrecaching = true;
+
 			// Precache name, import and export map.
 			if (bLoaderIsFArchiveAsync2)
 			{
@@ -1426,33 +1426,32 @@ FLinkerLoad::ELinkerStatus FLinkerLoad::SerializeNameMap()
 			{
 				bFinishedPrecaching = Loader->Precache(Summary.NameOffset, Summary.TotalHeaderSize - Summary.NameOffset);
 			}
-		}
-		// Backward compat code for VER_MOVED_EXPORTIMPORTMAPS_ADDED_TOTALHEADERSIZE.
-		else
-		{
-			bFinishedPrecaching = true;
+
+			if (!bFinishedPrecaching)
+			{
+				return LINKER_TimedOut;
+			}
 		}
 	}
 
-	NameMap.Reserve(Summary.NameCount);
+	SCOPED_LOADTIMER(LinkerLoad_SerializeNameMap_ProcessingEntries);
 
-	while (bFinishedPrecaching && NameMapIndex < Summary.NameCount && !IsTimeLimitExceeded(TEXT("serializing name map"), 100))
+	NameMap.Reserve(NameCount);
+	FNameEntrySerialized NameEntry(ENAME_LinkerConstructor);
+	for (int32 Idx = NameMap.Num(); Idx < NameCount; ++Idx)
 	{
-		//SCOPED_LOADTIMER(LinkerLoad_SerializeNameMap_ProcessingEntries);
-
-		FNameEntrySerialized NameEntry(ENAME_LinkerConstructor);
-
-		// Read the name from the underlying Archive when the format is binary to avoid the overhead from ArchiveProxy
 		*this << NameEntry;
-
-		// Add it to the name table with no splitting and no hash calculations
 		NameMap.Emplace(FName(NameEntry).GetDisplayIndex());
 
-		NameMapIndex++;
+		constexpr int32 TimeSliceGranularity = 128;
+		if (Idx % TimeSliceGranularity == TimeSliceGranularity - 1 && 
+			NameMap.Num() != NameCount && IsTimeLimitExceeded(TEXT("serializing name map")))
+		{
+			return LINKER_TimedOut;
+		}
 	}
-
-	// Return whether we finished this step and it's safe to start with the next.
-	return ((NameMapIndex == Summary.NameCount) && !IsTimeLimitExceeded( TEXT("serializing name map") )) ? LINKER_Loaded : LINKER_TimedOut;
+	
+	return LINKER_Loaded;
 }
 
 /**
diff --git a/Engine/Source/Runtime/CoreUObject/Private/UObject/PropertyTag.cpp b/Engine/Source/Runtime/CoreUObject/Private/UObject/PropertyTag.cpp
--- a/Engine/Source/Runtime/CoreUObject/Private/UObject/PropertyTag.cpp
+++ b/Engine/Source/Runtime/CoreUObject/Private/UObject/PropertyTag.cpp
@@ -112,7 +112,7 @@ void operator<<(FStructuredArchive::FSlot Slot, FPropertyTag& Tag)
 
 	// Name.
 	Record << NAMED_ITEM("Name", Tag.Name);
-	if ((Tag.Name == NAME_None) || !Tag.Name.IsValid())
+	if (Tag.Name.IsNone())
 	{
 		return;
 	}
@@ -128,52 +128,57 @@ void operator<<(FStructuredArchive::FSlot Slot, FPropertyTag& Tag)
 		FArchive::FScopeSetDebugSerializationFlags S(UnderlyingArchive, DSF_IgnoreDiff);
 		Record << NAMED_ITEM("Size", Tag.Size) << NAMED_ITEM("ArrayIndex", Tag.ArrayIndex);
 	}
-	// only need to serialize this for structs
-	if (Tag.Type == NAME_StructProperty)
+
+	if (Tag.Type.GetNumber() == 0)
 	{
-		Record << NAMED_ITEM("StructName", Tag.StructName);
-		if (Version >= VER_UE4_STRUCT_GUID_IN_PROPERTY_TAG)
+		FNameEntryId TagType = Tag.Type.GetComparisonIndex();
+
+		// only need to serialize this for structs
+		if (TagType == NAME_StructProperty)
 		{
-			Record << NAMED_ITEM("StructGuid", Tag.StructGuid);
+			Record << NAMED_ITEM("StructName", Tag.StructName);
+			if (Version >= VER_UE4_STRUCT_GUID_IN_PROPERTY_TAG)
+			{
+				Record << NAMED_ITEM("StructGuid", Tag.StructGuid);
+			}
 		}
-	}
-	// only need to serialize this for bools
-	else if (Tag.Type == NAME_BoolProperty && !UnderlyingArchive.IsTextFormat())
-	{
-		if (UnderlyingArchive.IsSaving())
+		// only need to serialize this for bools
+		else if (TagType == NAME_BoolProperty && !UnderlyingArchive.IsTextFormat())
 		{
-			FSerializedPropertyScope SerializedProperty(UnderlyingArchive, Tag.Prop);
-			Record << NAMED_ITEM("BoolVal", Tag.BoolVal);
+			if (UnderlyingArchive.IsSaving())
+			{
+				FSerializedPropertyScope SerializedProperty(UnderlyingArchive, Tag.Prop);
+				Record << NAMED_ITEM("BoolVal", Tag.BoolVal);
+			}
+			else
+			{
+				Record << NAMED_ITEM("BoolVal", Tag.BoolVal);
+			}
 		}
-		else
+		// only need to serialize this for bytes/enums
+		else if (TagType == NAME_ByteProperty || TagType == NAME_EnumProperty)
 		{
-			Record << NAMED_ITEM("BoolVal", Tag.BoolVal);
+			Record << NAMED_ITEM("EnumName", Tag.EnumName);
 		}
-	}
-	// only need to serialize this for bytes/enums
-	else if (Tag.Type == NAME_ByteProperty || Tag.Type == NAME_EnumProperty)
-	{
-		Record << NAMED_ITEM("EnumName", Tag.EnumName);
-	}
-	// only need to serialize this for arrays
-	else if (Tag.Type == NAME_ArrayProperty)
-	{
-		if (Version >= VAR_UE4_ARRAY_PROPERTY_INNER_TAGS)
-		{
-			Record << NAMED_ITEM("InnerType", Tag.InnerType);
-		}
-	}
-
-	if (Version >= VER_UE4_PROPERTY_TAG_SET_MAP_SUPPORT)
-	{
-		if (Tag.Type == NAME_SetProperty)
+		// only need to serialize this for arrays
+		else if (TagType == NAME_ArrayProperty)
 		{
-			Record << NAMED_ITEM("InnerType", Tag.InnerType);
+			if (Version >= VAR_UE4_ARRAY_PROPERTY_INNER_TAGS)
+			{
+				Record << NAMED_ITEM("InnerType", Tag.InnerType);
+			}
 		}
-		else if (Tag.Type == NAME_MapProperty)
+		else if (Version >= VER_UE4_PROPERTY_TAG_SET_MAP_SUPPORT)
 		{
-			Record << NAMED_ITEM("InnerType", Tag.InnerType);
-			Record << NAMED_ITEM("ValueType", Tag.ValueType);
+			if (TagType == NAME_SetProperty)
+			{
+				Record << NAMED_ITEM("InnerType", Tag.InnerType);
+			}
+			else if (TagType == NAME_MapProperty)
+			{
+				Record << NAMED_ITEM("InnerType", Tag.InnerType);
+				Record << NAMED_ITEM("ValueType", Tag.ValueType);
+			}
 		}
 	}
 
diff --git a/Engine/Source/Runtime/CoreUObject/Private/UObject/SoftObjectPath.cpp b/Engine/Source/Runtime/CoreUObject/Private/UObject/SoftObjectPath.cpp
--- a/Engine/Source/Runtime/CoreUObject/Private/UObject/SoftObjectPath.cpp
+++ b/Engine/Source/Runtime/CoreUObject/Private/UObject/SoftObjectPath.cpp
@@ -382,7 +382,6 @@ UObject* FSoftObjectPath::ResolveObject() const
 		return nullptr;
 	}
 
-	FString PathString = ToString();
 #if WITH_EDITOR
 	if (GPlayInEditorID != INDEX_NONE)
 	{
@@ -390,12 +389,31 @@ UObject* FSoftObjectPath::ResolveObject() const
 		FSoftObjectPath FixupObjectPath = *this;
 		if (FixupObjectPath.FixupForPIE())
 		{
-			PathString = FixupObjectPath.ToString();
+			return FixupObjectPath.ResolveObjectInternal();
 		}
 	}
 #endif
 
-	UObject* FoundObject = FindObject<UObject>(nullptr, *PathString);
+	return ResolveObjectInternal();
+}
+
+UObject* FSoftObjectPath::ResolveObjectInternal() const
+{
+	if (SubPathString.IsEmpty())
+	{
+		TCHAR PathString[FName::StringBufferSize];
+		AssetPathName.ToString(PathString);
+		return ResolveObjectInternal(PathString);
+	}
+	else
+	{
+		return ResolveObjectInternal(*ToString());
+	}
+}
+
+UObject* FSoftObjectPath::ResolveObjectInternal(const TCHAR* PathString) const
+{
+	UObject* FoundObject = FindObject<UObject>(nullptr, PathString);
 
 #if WITH_EDITOR
 	// Look at core redirects if we didn't find the object
@@ -404,8 +422,7 @@ UObject* FSoftObjectPath::ResolveObject() const
 		FSoftObjectPath FixupObjectPath = *this;
 		if (FixupObjectPath.FixupCoreRedirects())
 		{
-			PathString = FixupObjectPath.ToString();
-			FoundObject = FindObject<UObject>(nullptr, *PathString);
+			FoundObject = FindObject<UObject>(nullptr, *FixupObjectPath.ToString());
 		}
 	}
 #endif
diff --git a/Engine/Source/Runtime/CoreUObject/Private/UObject/UObjectHash.cpp b/Engine/Source/Runtime/CoreUObject/Private/UObject/UObjectHash.cpp
--- a/Engine/Source/Runtime/CoreUObject/Private/UObject/UObjectHash.cpp
+++ b/Engine/Source/Runtime/CoreUObject/Private/UObject/UObjectHash.cpp
@@ -536,13 +536,16 @@ UObject* StaticFindObjectFastInternalThreadSafe(FUObjectHashTables& ThreadHash,
 	{
 		// Find an object with the specified name and (optional) class, in any package; if bAnyPackage is false, only matches top-level packages
 		FName ActualObjectName = ObjectName;
-		const FString ObjectNameString = ObjectName.ToString();
-		const int32 DotIndex = FMath::Max<int32>(ObjectNameString.Find(TEXT("."), ESearchCase::CaseSensitive, ESearchDir::FromEnd),
-			ObjectNameString.Find(TEXT(":"), ESearchCase::CaseSensitive, ESearchDir::FromEnd));
-		if (DotIndex != INDEX_NONE)
+		TCHAR ObjectNameString[NAME_SIZE];
+		ObjectName.GetPlainNameString(ObjectNameString);
+
+		const TCHAR* LastDot = FCString::Strrchr(ObjectNameString, '.');
+		const TCHAR* FirstColon = FCString::Strchr(ObjectNameString, ':');
+		if (const TCHAR* LastDelimiter = (UPTRINT)LastDot > (UPTRINT)FirstColon ? LastDot : FirstColon)
 		{
-			ActualObjectName = FName(*ObjectNameString.Mid(DotIndex + 1));
+			ActualObjectName = FName(LastDelimiter + 1, ObjectName.GetNumber());
 		}
+
 		const int32 Hash = GetObjectHash(ActualObjectName);
 		FHashTableLock HashLock(ThreadHash);
 
diff --git a/Engine/Source/Runtime/CoreUObject/Public/UObject/LinkerLoad.h b/Engine/Source/Runtime/CoreUObject/Public/UObject/LinkerLoad.h
--- a/Engine/Source/Runtime/CoreUObject/Public/UObject/LinkerLoad.h
+++ b/Engine/Source/Runtime/CoreUObject/Public/UObject/LinkerLoad.h
@@ -258,8 +258,6 @@ private:
 
 	// Variables used during async linker creation.
 
-	/** Current index into name map, used by async linker creation for spreading out serializing name entries.					*/
-	int32						NameMapIndex;
 	/** Current index into gatherable text data map, used by async linker creation for spreading out serializing text entries.	*/
 	int32						GatherableTextDataMapIndex;
 	/** Current index into import map, used by async linker creation for spreading out serializing importmap entries.			*/
@@ -809,29 +807,28 @@ private:
 	void BadNameIndexError(int32 NameIndex);
 	FORCEINLINE virtual FArchive& operator<<(FName& Name) override
 	{
-		Name = NAME_None;
 		FArchive& Ar = *this;
 		int32 NameIndex;
 		Ar << NameIndex;
-		int32 Number;
+		int32 Number = 0;
 		Ar << Number;
 
-		if (!NameMap.IsValidIndex(NameIndex))
+		if (NameMap.IsValidIndex(NameIndex))
 		{
-			BadNameIndexError(NameIndex);
-			ArIsError = true;
-			ArIsCriticalError = true;
+			// if the name wasn't loaded (because it wasn't valid in this context)
+			FNameEntryId MappedName = NameMap[NameIndex];
+
+			// simply create the name from the NameMap's name and the serialized instance number
+			Name = FName::CreateFromDisplayId(MappedName, Number);
 		}
 		else
 		{
-			// if the name wasn't loaded (because it wasn't valid in this context)
-			FNameEntryId MappedName = NameMap[NameIndex];
-			if (MappedName)
-			{
-				// simply create the name from the NameMap's name and the serialized instance number
-				Name = FName::CreateFromDisplayId(MappedName, Number);
-			}
+			Name = FName();
+			BadNameIndexError(NameIndex);
+			ArIsError = true;
+			ArIsCriticalError = true;
 		}
+
 		return *this;
 	}
 
diff --git a/Engine/Source/Runtime/CoreUObject/Public/UObject/SoftObjectPath.h b/Engine/Source/Runtime/CoreUObject/Public/UObject/SoftObjectPath.h
--- a/Engine/Source/Runtime/CoreUObject/Public/UObject/SoftObjectPath.h
+++ b/Engine/Source/Runtime/CoreUObject/Public/UObject/SoftObjectPath.h
@@ -110,32 +110,32 @@ struct COREUOBJECT_API FSoftObjectPath
 	/** Resets reference to point to null */
 	void Reset()
 	{		
-		AssetPathName = NAME_None;
+		AssetPathName = FName();
 		SubPathString.Reset();
 	}
 	
 	/** Check if this could possibly refer to a real object, or was initialized to null */
 	FORCEINLINE bool IsValid() const
 	{
-		return AssetPathName != NAME_None;
+		return !AssetPathName.IsNone();
 	}
 
 	/** Checks to see if this is initialized to null */
 	FORCEINLINE bool IsNull() const
 	{
-		return AssetPathName == NAME_None;
+		return AssetPathName.IsNone();
 	}
 
 	/** Check if this represents an asset, meaning it is not null but does not have a sub path */
 	FORCEINLINE bool IsAsset() const
 	{
-		return AssetPathName != NAME_None && SubPathString.IsEmpty();
+		return !AssetPathName.IsNone() && SubPathString.IsEmpty();
 	}
 
 	/** Check if this represents a sub object, meaning it has a sub path */
 	FORCEINLINE bool IsSubobject() const
 	{
-		return AssetPathName != NAME_None && !SubPathString.IsEmpty();
+		return !AssetPathName.IsNone() && !SubPathString.IsEmpty();
 	}
 
 	/** Struct overrides */
@@ -211,6 +211,9 @@ private:
 	/** Package names currently being duplicated, needed by FixupForPIE */
 	static TSet<FName> PIEPackageNames;
 
+	UObject* ResolveObjectInternal() const;
+	UObject* ResolveObjectInternal(const TCHAR* PathString) const;
+
 	friend struct Z_Construct_UScriptStruct_FSoftObjectPath_Statics;
 };
 
