Optimize FAssetData::IsAssetLoaded() by removing temporary string allocation and using FindObjectFast instead of FindObjectSafe

diff --git a/Engine/Source/Runtime/AssetRegistry/Public/AssetData.h b/Engine/Source/Runtime/AssetRegistry/Public/AssetData.h
--- a/Engine/Source/Runtime/AssetRegistry/Public/AssetData.h
+++ b/Engine/Source/Runtime/AssetRegistry/Public/AssetData.h
@@ -356,7 +356,12 @@ public:
 	/** Returns true if the asset is loaded */
 	bool IsAssetLoaded() const
 	{
-		return IsValid() && FindObjectSafe<UObject>(NULL, *ObjectPath.ToString()) != NULL;
+		if (IsValid())
+		{
+			UPackage* Package = FindObjectFast<UPackage>(nullptr, PackageName);
+			return (Package != NULL) && (FindObjectFast<UObject>(Package, AssetName) != NULL);
+		}
+		return false;
 	}
 
 	/** Prints the details of the asset to the log */
