Improve some comments and variable names + add some minor tweaks in asset registry functions:

FAssetDataGatherer::Run - fix memory leak if stop signal received while still gathering.
FAssetDataGatherer::SerializeCache - add safety check for corrupt archive
FNameTableArchiveReader::SerializeNameMap - add safety check for corrupt archive
UAssetRegistryImpl::DependencyDataGathered - separate searchablenames export handling from import handling for clarity and performance.

#rnx

diff --git a/Engine/Source/Editor/UnrealEd/Private/CookOnTheFlyServer.cpp b/Engine/Source/Editor/UnrealEd/Private/CookOnTheFlyServer.cpp
--- a/Engine/Source/Editor/UnrealEd/Private/CookOnTheFlyServer.cpp
+++ b/Engine/Source/Editor/UnrealEd/Private/CookOnTheFlyServer.cpp
@@ -4136,7 +4136,7 @@ void UCookOnTheFlyServer::Initialize( ECookMode::Type DesiredCookMode, ECookInit
 	MaxPrecacheShaderJobs = FPlatformMisc::NumberOfCores() - 1; // number of cores -1 is a good default allows the editor to still be responsive to other shader requests and allows cooker to take advantage of multiple processors while the editor is running
 	GConfig->GetInt(TEXT("CookSettings"), TEXT("MaxPrecacheShaderJobs"), MaxPrecacheShaderJobs, GEditorIni);
 
-	MaxConcurrentShaderJobs = FPlatformMisc::NumberOfCores() * 4; // number of cores -1 is a good default allows the editor to still be responsive to other shader requests and allows cooker to take advantage of multiple processors while the editor is running
+	MaxConcurrentShaderJobs = FPlatformMisc::NumberOfCores() * 4; // TODO: document why number of cores * 4 is a good default
 	GConfig->GetInt(TEXT("CookSettings"), TEXT("MaxConcurrentShaderJobs"), MaxConcurrentShaderJobs, GEditorIni);
 
 	PackagesPerGC = 500;
@@ -4170,7 +4170,7 @@ void UCookOnTheFlyServer::Initialize( ECookMode::Type DesiredCookMode, ECookInit
 	MinFreeMemoryInMB = FMath::Max(MinFreeMemoryInMB, 0);
 	MinFreeMemory = MinFreeMemoryInMB * 1024LL * 1024LL;
 
-	// check the amount of OS memory and use that number minus the reserved memory nubmer
+	// check the amount of OS memory and use that number minus the reserved memory number
 	int32 MinReservedMemoryInMB = 0;
 	GConfig->GetInt(TEXT("CookSettings"), TEXT("MinReservedMemory"), MinReservedMemoryInMB, GEditorIni);
 	MinReservedMemoryInMB = FMath::Max(MinReservedMemoryInMB, 0);
@@ -4227,7 +4227,7 @@ void UCookOnTheFlyServer::Initialize( ECookMode::Type DesiredCookMode, ECookInit
 	{
 		PluginsToRemap = IPluginManager::Get().GetEnabledPlugins();
 		TArray<FString> AdditionalPluginDirs = Project->GetAdditionalPluginDirectories();
-		// Remove any plugin that is not in the additional directories since they are handled normally
+		// Remove any plugin that is in the additional directories since they are handled normally and don't need remapping
 		for (int32 Index = PluginsToRemap.Num() - 1; Index >= 0; Index--)
 		{
 			bool bRemove = true;
@@ -7126,7 +7126,7 @@ void UCookOnTheFlyServer::StartCookByTheBook( const FCookByTheBookStartupOptions
 	
 	if ( IsCookingDLC() )
 	{
-		// if we are cooking dlc we must be based of a release version cook
+		// if we are cooking dlc we must be based on a release version cook
 		check( !BasedOnReleaseVersion.IsEmpty() );
 
 		for ( const FName& PlatformName : TargetPlatformNames )
diff --git a/Engine/Source/Runtime/AssetRegistry/Private/AssetDataGatherer.cpp b/Engine/Source/Runtime/AssetRegistry/Private/AssetDataGatherer.cpp
--- a/Engine/Source/Runtime/AssetRegistry/Private/AssetDataGatherer.cpp
+++ b/Engine/Source/Runtime/AssetRegistry/Private/AssetDataGatherer.cpp
@@ -510,13 +510,25 @@ uint32 FAssetDataGatherer::Run()
 		NumFilesProcessedSinceLastCacheSave = 0;
 	};
 
-	while ( StopTaskCounter.GetValue() == 0 )
+	while ( true )
 	{
 		bool LocalIsDiscoveringFiles = false;
 
 		{
 			FScopeLock CritSectionLock(&WorkerThreadCriticalSection);
 
+			AssetResults.Append(MoveTemp(LocalAssetResults));
+			DependencyResults.Append(MoveTemp(LocalDependencyResults));
+			CookedPackageNamesWithoutAssetDataResults.Append(MoveTemp(LocalCookedPackageNamesWithoutAssetDataResults));
+			LocalAssetResults.Reset();
+			LocalDependencyResults.Reset();
+			LocalCookedPackageNamesWithoutAssetDataResults.Reset();
+
+			if (StopTaskCounter.GetValue() != 0)
+			{
+				break;
+			}
+
 			// Grab any new package files from the background directory scan
 			if (BackgroundPackageFileDiscovery.IsValid())
 			{
@@ -524,9 +536,6 @@ uint32 FAssetDataGatherer::Run()
 				LocalIsDiscoveringFiles = bIsDiscoveringFiles;
 			}
 
-			AssetResults.Append(MoveTemp(LocalAssetResults));
-			DependencyResults.Append(MoveTemp(LocalDependencyResults));
-			CookedPackageNamesWithoutAssetDataResults.Append(MoveTemp(LocalCookedPackageNamesWithoutAssetDataResults));
 
 			if (FilesToSearch.Num() > 0)
 			{
@@ -546,10 +555,6 @@ uint32 FAssetDataGatherer::Run()
 			}
 		}
 
-		LocalAssetResults.Reset();
-		LocalDependencyResults.Reset();
-		LocalCookedPackageNamesWithoutAssetDataResults.Reset();
-
 		TArray<FDiscoveredPackageFile> LocalFilesToRetry;
 
 		if (LocalFilesToSearch.Num() > 0)
@@ -628,7 +633,6 @@ uint32 FAssetDataGatherer::Run()
 						bool bCachePackage = bLoadAndSaveCache && LocalCookedPackageNamesWithoutAssetDataResults.Num() == 0;
 						if (bCachePackage)
 						{
-							// Don't store info on cooked packages
 							for (const auto& AssetData : AssetDataFromFile)
 							{
 								if (!!(AssetData->PackageFlags & PKG_FilterEditorOnly))
@@ -704,6 +708,8 @@ uint32 FAssetDataGatherer::Run()
 		}
 	}
 
+	check(LocalAssetResults.Num() == 0); // All LocalAssetResults needed to be copied to AssetResults to avoid leaking memory
+
 	if ( bLoadAndSaveCache )
 	{
 		WriteAssetCacheFile();
@@ -913,29 +919,38 @@ void FAssetDataGatherer::SerializeCache(FArchive& Ar)
 	else
 	{
 		// allocate one single block for all asset data structs (to reduce tens of thousands of heap allocations)
-		DiskCachedAssetDataMap.Empty(LocalNumAssets);
-
-		for (int32 AssetIndex = 0; AssetIndex < LocalNumAssets; ++AssetIndex)
+		if (Ar.IsError() || LocalNumAssets < 0)
 		{
-			// Load the name first to add the entry to the tmap below
-			FName PackageName;
-			Ar << PackageName;
-			if (Ar.IsError())
+			Ar.SetError();
+		}
+		else
+		{
+			const int32 MinAssetEntrySize = sizeof(int32);
+			int32 MaxReservation = (Ar.TotalSize() - Ar.Tell()) / MinAssetEntrySize;
+			DiskCachedAssetDataMap.Empty(FMath::Min(LocalNumAssets, MaxReservation));
+
+			for (int32 AssetIndex = 0; AssetIndex < LocalNumAssets; ++AssetIndex)
 			{
-				// There was an error reading the cache. Bail out.
-				break;
-			}
+				// Load the name first to add the entry to the tmap below
+				FName PackageName;
+				Ar << PackageName;
+				if (Ar.IsError())
+				{
+					// There was an error reading the cache. Bail out.
+					break;
+				}
 
-			// Add to the cached map
-			FDiskCachedAssetData& CachedAssetData = DiskCachedAssetDataMap.Add(PackageName);
+				// Add to the cached map
+				FDiskCachedAssetData& CachedAssetData = DiskCachedAssetDataMap.Add(PackageName);
 
-			// Now load the data
-			CachedAssetData.SerializeForCache(Ar);
+				// Now load the data
+				CachedAssetData.SerializeForCache(Ar);
 
-			if (Ar.IsError())
-			{
-				// There was an error reading the cache. Bail out.
-				break;
+				if (Ar.IsError())
+				{
+					// There was an error reading the cache. Bail out.
+					break;
+				}
 			}
 		}
 
diff --git a/Engine/Source/Runtime/AssetRegistry/Private/AssetRegistry.cpp b/Engine/Source/Runtime/AssetRegistry/Private/AssetRegistry.cpp
--- a/Engine/Source/Runtime/AssetRegistry/Private/AssetRegistry.cpp
+++ b/Engine/Source/Runtime/AssetRegistry/Private/AssetRegistry.cpp
@@ -133,8 +133,8 @@ UAssetRegistryImpl::UAssetRegistryImpl(const FObjectInitializer& ObjectInitializ
 				const FString& RootPath = *RootPathIt;
 				const FString& ContentFolder = FPackageName::LongPackageNameToFilename(RootPath);
 
-				// This could be due to a plugin that specifies it contains content, yet has no content yet. PluginManager
-				// Mounts these folders anyway which results in then being returned from QueryRootContentPaths
+				// A missing directory here could be due to a plugin that specifies it contains content, yet has no content yet. PluginManager
+				// Mounts these folders anyway which results in them being returned from QueryRootContentPaths
 				if (IFileManager::Get().DirectoryExists(*ContentFolder))
 				{
 					FDelegateHandle NewHandle;
@@ -2133,27 +2133,29 @@ void UAssetRegistryImpl::DependencyDataGathered(const double TickStartTime, TBac
 			FName PackageName;
 
 			// Find object and package name from linker
-			for (FPackageIndex LinkerIndex = SearchableNameList.Key; !LinkerIndex.IsNull();)
+			FPackageIndex LinkerIndex = SearchableNameList.Key;
+			if (LinkerIndex.IsExport())
 			{
-				if (LinkerIndex.IsExport())
-				{
-					// Package name has to be this package, take a guess at object name
-					PackageName = Result.PackageName;
-					ObjectName = FName(*FPackageName::GetLongPackageAssetName(Result.PackageName.ToString()));
-
-					break;
-				}
-
-				FObjectResource& Resource = Result.ImpExp(LinkerIndex);
-				LinkerIndex = Resource.OuterIndex;
-				if (ObjectName.IsNone() && !LinkerIndex.IsNull())
-				{
-					ObjectName = Resource.ObjectName;
-				}
-				else if (LinkerIndex.IsNull())
+				// Package name has to be this package, take a guess at object name
+				PackageName = Result.PackageName;
+				ObjectName = FName(*FPackageName::GetLongPackageAssetName(Result.PackageName.ToString()));
+			}
+			else if (LinkerIndex.IsImport())
+			{
+				FObjectResource* Resource = &Result.ImpExp(LinkerIndex);
+				FPackageIndex OuterLinkerIndex = Resource->OuterIndex;
+				check(OuterLinkerIndex.IsNull() || OuterLinkerIndex.IsImport());
+				if (!OuterLinkerIndex.IsNull())
 				{
-					PackageName = Resource.ObjectName;
+					ObjectName = Resource->ObjectName;
+					while (!OuterLinkerIndex.IsNull())
+					{
+						Resource = &Result.ImpExp(OuterLinkerIndex);
+						OuterLinkerIndex = Resource->OuterIndex;
+						check(OuterLinkerIndex.IsNull() || OuterLinkerIndex.IsImport());
+					}
 				}
+				PackageName = Resource->ObjectName;
 			}
 
 			for (FName NameReference : SearchableNameList.Value)
@@ -2220,7 +2222,8 @@ void UAssetRegistryImpl::CookedPackageNamesWithoutAssetDataGathered(const double
 	// Add the found assets
 	while (CookedPackageNamesWithoutAssetDataResults.Num() > 0)
 	{
-		// If this data is cooked and it we couldn't find any asset in its export table then try load the entire package 
+		// If this data is cooked and it we couldn't find any asset in its export table then try to load the entire package 
+		// Loading the entire package will make all of its assets searchable through the in-memory scanning performed by GetAsseets
 		const FString& BackgroundResult = CookedPackageNamesWithoutAssetDataResults.Pop();
 		LoadPackage(nullptr, *BackgroundResult, 0);
 
diff --git a/Engine/Source/Runtime/AssetRegistry/Private/NameTableArchive.cpp b/Engine/Source/Runtime/AssetRegistry/Private/NameTableArchive.cpp
--- a/Engine/Source/Runtime/AssetRegistry/Private/NameTableArchive.cpp
+++ b/Engine/Source/Runtime/AssetRegistry/Private/NameTableArchive.cpp
@@ -35,7 +35,7 @@ FNameTableArchiveReader::FNameTableArchiveReader(int32 SerializationVersion, con
 			{
 				if (SerializeNameMap())
 				{
-					// Succesfully loaded
+					// Successfully loaded
 					return;
 				}
 			}
@@ -74,7 +74,7 @@ bool FNameTableArchiveReader::SerializeNameMap()
 
 	if (NameOffset > TotalSize())
 	{
-		// The file was corrupted. Return false to fail to load the cache an thus regenerate it.
+		// The file was corrupted. Return false to fail to load the cache and thus regenerate it.
 		return false;
 	}
 
@@ -85,8 +85,14 @@ bool FNameTableArchiveReader::SerializeNameMap()
 
 		int32 NameCount = 0;
 		*this << NameCount;
+		if (IsError() || NameCount < 0)
+		{
+			return false;
+		}
 		
-		NameMap.Reserve(NameCount);
+		const int32 MinFNameEntrySize = sizeof(int32);
+		int32 MaxReservation = (TotalSize() - Tell()) / MinFNameEntrySize;
+		NameMap.Reserve(FMath::Min(NameCount, MaxReservation));
 		for ( int32 NameMapIdx = 0; NameMapIdx < NameCount; ++NameMapIdx )
 		{
 			// Read the name entry from the file.
diff --git a/Engine/Source/Runtime/AssetRegistry/Private/NameTableArchive.h b/Engine/Source/Runtime/AssetRegistry/Private/NameTableArchive.h
--- a/Engine/Source/Runtime/AssetRegistry/Private/NameTableArchive.h
+++ b/Engine/Source/Runtime/AssetRegistry/Private/NameTableArchive.h
@@ -23,6 +23,10 @@ struct TNameTableStringHash : BaseKeyFuncs<ValueType, FString, /*bInAllowDuplica
 	}
 };
 
+/** 
+ * Reader for an FNameTableArchive.  An FNameTableArchive is like a normal archive but with a separate FName table embedded in it, and
+ * FNames in it are serialized as indices into the Name table.
+*/
 class FNameTableArchiveReader : public FArchive
 {
 public:
@@ -54,6 +58,9 @@ private:
 	TArray<FNameEntryId> NameMap;
 };
 
+/**
+ * Writer for an FNameTableArchive.  See the class comment on FNameTableArchiveReader.
+ */
 class FNameTableArchiveWriter : public FArchive
 {
 public:
diff --git a/Engine/Source/Runtime/AssetRegistry/Private/PackageReader.cpp b/Engine/Source/Runtime/AssetRegistry/Private/PackageReader.cpp
--- a/Engine/Source/Runtime/AssetRegistry/Private/PackageReader.cpp
+++ b/Engine/Source/Runtime/AssetRegistry/Private/PackageReader.cpp
@@ -127,7 +127,7 @@ bool FPackageReader::ReadAssetRegistryData (TArray<FAssetData*>& AssetDataList)
 		return false;
 	}
 
-	// Seek the the part of the file where the asset registry tags live
+	// Seek to the part of the file where the asset registry tags live
 	Seek( PackageFileSummary.AssetRegistryDataOffset );
 
 	// Determine the package name and path
@@ -193,7 +193,7 @@ bool FPackageReader::ReadAssetRegistryData (TArray<FAssetData*>& AssetDataList)
 
 		FString AssetName = ObjectPath;
 
-		// Before world were RF_Public, other non-public assets were added to the asset data table in map packages.
+		// Before worlds were RF_Public, other non-public assets were added to the asset data table in map packages.
 		// Here we simply skip over them
 		if ( bIsMapPackage && PackageFileSummary.GetFileVersionUE4() < VER_UE4_PUBLIC_WORLDS )
 		{
diff --git a/Engine/Source/Runtime/CoreUObject/Private/Misc/RedirectCollector.cpp b/Engine/Source/Runtime/CoreUObject/Private/Misc/RedirectCollector.cpp
--- a/Engine/Source/Runtime/CoreUObject/Private/Misc/RedirectCollector.cpp
+++ b/Engine/Source/Runtime/CoreUObject/Private/Misc/RedirectCollector.cpp
@@ -157,11 +157,11 @@ void FRedirectCollector::ProcessSoftObjectPathPackageList(FName FilterPackage, b
 
 	// potentially add soft object path package names to OutReferencedPackages
 	OutReferencedPackages.Reserve(SoftObjectPathProperties->Num());
-	for (const FSoftObjectPathProperty& SoftObjecPathProperty : *SoftObjectPathProperties)
+	for (const FSoftObjectPathProperty& SoftObjectPathProperty : *SoftObjectPathProperties)
 	{
-		if (!SoftObjecPathProperty.GetReferencedByEditorOnlyProperty() || bGetEditorOnly)
+		if (!SoftObjectPathProperty.GetReferencedByEditorOnlyProperty() || bGetEditorOnly)
 		{
-			const FName& ToLoadFName = SoftObjecPathProperty.GetAssetPathName();
+			const FName& ToLoadFName = SoftObjectPathProperty.GetAssetPathName();
 			FString PackageNameString = FPackageName::ObjectPathToPackageName(ToLoadFName.ToString());
 			OutReferencedPackages.Add(FName(*PackageNameString));
 		}
diff --git a/Engine/Source/Runtime/CoreUObject/Public/Misc/RedirectCollector.h b/Engine/Source/Runtime/CoreUObject/Public/Misc/RedirectCollector.h
--- a/Engine/Source/Runtime/CoreUObject/Public/Misc/RedirectCollector.h
+++ b/Engine/Source/Runtime/CoreUObject/Public/Misc/RedirectCollector.h
@@ -64,7 +64,7 @@ private:
 public:
 
 	/**
-	 * Called from FSoftObjectPath::PostLoadPath, registers this for later querying
+	 * Called from FSoftObjectPath::PostLoadPath, registers the given SoftObjectPath for later querying
 	 * @param InPath The soft object path that was loaded
 	 * @Param InArchive The archive that loaded this path
 	 */
