Added enumeration variants for GetAssets, GetAllAssets, GetAllCachedPaths, and GetSubPaths in the Asset Registry

These can avoid allocating an intermediate array if you only need to process the result. The existing GetX functions now use the EnumerateX functions internally.

diff --git a/Engine/Source/Runtime/AssetRegistry/Private/AssetRegistry.cpp b/Engine/Source/Runtime/AssetRegistry/Private/AssetRegistry.cpp
--- a/Engine/Source/Runtime/AssetRegistry/Private/AssetRegistry.cpp
+++ b/Engine/Source/Runtime/AssetRegistry/Private/AssetRegistry.cpp
@@ -553,17 +553,11 @@ bool UAssetRegistryImpl::HasAssets(const FName PackagePath, const bool bRecursiv
 
 	if (!bHasAssets && bRecursive)
 	{
-		TSet<FName> SubPaths;
-		CachedPathTree.GetSubPaths(PackagePath, SubPaths);
-
-		for (const FName& SubPath : SubPaths)
+		CachedPathTree.EnumerateSubPaths(PackagePath, [this, &bHasAssets](FName SubPath)
 		{
 			bHasAssets = State.HasAssets(SubPath);
-			if (bHasAssets)
-			{
-				break;
-			}
-		}
+			return !bHasAssets;
+		});
 	}
 
 	return bHasAssets;
@@ -616,9 +610,22 @@ bool UAssetRegistryImpl::GetAssetsByTagValues(const TMultiMap<FName, FString>& A
 
 bool UAssetRegistryImpl::GetAssets(const FARFilter& InFilter, TArray<FAssetData>& OutAssetData) const
 {
-	double GetAssetsStartTime = FPlatformTime::Seconds();
+	const double GetAssetsStartTime = FPlatformTime::Seconds();
+
+	const bool bResult = EnumerateAssets(InFilter, [&OutAssetData](const FAssetData& AssetData)
+	{
+		OutAssetData.Emplace(AssetData);
+		return true;
+	});
+
+	UE_LOG(LogAssetRegistry, Verbose, TEXT("GetAssets completed in %0.4f seconds"), FPlatformTime::Seconds() - GetAssetsStartTime);
+
+	return bResult;
+}
 
-	// Verify filter input. If all assets are needed, use GetAllAssets() instead.
+bool UAssetRegistryImpl::EnumerateAssets(const FARFilter& InFilter, TFunctionRef<bool(const FAssetData&)> Callback) const
+{
+	// Verify filter input. If all assets are needed, use EnumerateAllAssets() instead.
 	if (!FAssetRegistryState::IsFilterValid(InFilter, true) || InFilter.IsEmpty())
 	{
 		return false;
@@ -646,7 +653,7 @@ bool UAssetRegistryImpl::GetAssets(const FARFilter& InFilter, TArray<FAssetData>
 		// Reusable structures to avoid memory allocations
 		TArray<UObject::FAssetRegistryTag> ObjectTags;
 
-		auto FilterInMemoryObjectLambda = [&](const UObject* Obj)
+		auto FilterInMemoryObjectLambda = [&](const UObject* Obj, bool& OutContinue)
 		{
 			if (Obj->IsAsset())
 			{
@@ -734,7 +741,7 @@ bool UAssetRegistryImpl::GetAssets(const FARFilter& InFilter, TArray<FAssetData>
 				ObjectTags.Reset();
 
 				// This asset is in memory and passes all filters
-				OutAssetData.Emplace(FAssetData(PackageName, PackagePath, Obj->GetFName(), Obj->GetClass()->GetFName(), MoveTemp(TagMap), InMemoryPackage->GetChunkIDs(), InMemoryPackage->GetPackageFlags()));
+				OutContinue = Callback(FAssetData(PackageName, PackagePath, Obj->GetFName(), Obj->GetClass()->GetFName(), MoveTemp(TagMap), InMemoryPackage->GetChunkIDs(), InMemoryPackage->GetPackageFlags()));
 			}
 		};
 
@@ -753,21 +760,29 @@ bool UAssetRegistryImpl::GetAssets(const FARFilter& InFilter, TArray<FAssetData>
 
 			for (UObject* Object : InMemoryObjects)
 			{
-				FilterInMemoryObjectLambda(Object);
+				bool bContinue = true;
+				FilterInMemoryObjectLambda(Object, bContinue);
+				if (!bContinue)
+				{
+					return true;
+				}
 			}
 		}
 		else
 		{
 			for (FObjectIterator ObjIt; ObjIt; ++ObjIt)
 			{
-				FilterInMemoryObjectLambda(*ObjIt);
+				bool bContinue = true;
+				FilterInMemoryObjectLambda(*ObjIt, bContinue);
+				if (!bContinue)
+				{
+					return true;
+				}
 			}
 		}
 	}
 
-	State.GetAssets(Filter, PackagesToSkip, OutAssetData);
-
-	UE_LOG(LogAssetRegistry, Verbose, TEXT("GetAssets completed in %0.4f seconds"), FPlatformTime::Seconds() - GetAssetsStartTime);
+	State.EnumerateAssets(Filter, PackagesToSkip, Callback);
 
 	return true;
 }
@@ -794,9 +809,23 @@ FAssetData UAssetRegistryImpl::GetAssetByObjectPath(const FName ObjectPath, bool
 }
 
 bool UAssetRegistryImpl::GetAllAssets(TArray<FAssetData>& OutAssetData, bool bIncludeOnlyOnDiskAssets) const
+{
+	const double GetAllAssetsStartTime = FPlatformTime::Seconds();
+
+	const bool bResult = EnumerateAllAssets([&OutAssetData](const FAssetData& AssetData)
+	{
+		OutAssetData.Emplace(AssetData);
+		return true;
+	}, bIncludeOnlyOnDiskAssets);
+
+	UE_LOG(LogAssetRegistry, VeryVerbose, TEXT("GetAllAssets completed in %0.4f seconds"), FPlatformTime::Seconds() - GetAllAssetsStartTime);
+	
+	return bResult;
+}
+
+bool UAssetRegistryImpl::EnumerateAllAssets(TFunctionRef<bool(const FAssetData&)> Callback, bool bIncludeOnlyOnDiskAssets) const
 {
 	TSet<FName> PackageNamesToSkip = CachedEmptyPackages;
-	double GetAllAssetsStartTime = FPlatformTime::Seconds();
 
 	// All in memory assets
 	if (!bIncludeOnlyOnDiskAssets)
@@ -805,15 +834,17 @@ bool UAssetRegistryImpl::GetAllAssets(TArray<FAssetData>& OutAssetData, bool bIn
 		{
 			if (ObjIt->IsAsset())
 			{
-				const FAssetData& AssetData = OutAssetData[OutAssetData.Emplace(*ObjIt)];
+				const FAssetData AssetData(*ObjIt);
+				if (!Callback(AssetData))
+				{
+					return true;
+				}
 				PackageNamesToSkip.Add(AssetData.PackageName);
 			}
 		}
 	}
 
-	State.GetAllAssets(PackageNamesToSkip, OutAssetData);
-
-	UE_LOG(LogAssetRegistry, VeryVerbose, TEXT("GetAllAssets completed in %0.4f seconds"), FPlatformTime::Seconds() - GetAllAssetsStartTime);
+	State.EnumerateAllAssets(PackageNamesToSkip, Callback);
 
 	return true;
 }
@@ -964,26 +995,46 @@ void UAssetRegistryImpl::GetDerivedClassNames(const TArray<FName>& ClassNames, c
 
 void UAssetRegistryImpl::GetAllCachedPaths(TArray<FString>& OutPathList) const
 {
-	TSet<FName> PathList;
-	CachedPathTree.GetAllPaths(PathList);
+	EnumerateAllCachedPaths([&OutPathList](FName Path)
+	{
+		OutPathList.Emplace(Path.ToString());
+		return true;
+	});
+}
 
-	OutPathList.Empty(PathList.Num());
-	for (FName PathName : PathList)
+void UAssetRegistryImpl::EnumerateAllCachedPaths(TFunctionRef<bool(FString)> Callback) const
+{
+	EnumerateAllCachedPaths([&Callback](FName Path)
 	{
-		OutPathList.Add(PathName.ToString());
-	}
+		return Callback(Path.ToString());
+	});
+}
+
+void UAssetRegistryImpl::EnumerateAllCachedPaths(TFunctionRef<bool(FName)> Callback) const
+{
+	CachedPathTree.EnumerateAllPaths(Callback);
 }
 
 void UAssetRegistryImpl::GetSubPaths(const FString& InBasePath, TArray<FString>& OutPathList, bool bInRecurse) const
 {
-	TSet<FName> PathList;
-	CachedPathTree.GetSubPaths(FName(*InBasePath), PathList, bInRecurse);
+	EnumerateSubPaths(*InBasePath, [&OutPathList](FName Path)
+	{
+		OutPathList.Emplace(Path.ToString());
+		return true;
+	}, bInRecurse);
+}
 
-	OutPathList.Empty(PathList.Num());
-	for (FName PathName : PathList)
+void UAssetRegistryImpl::EnumerateSubPaths(const FString& InBasePath, TFunctionRef<bool(FString)> Callback, bool bInRecurse) const
+{
+	EnumerateSubPaths(*InBasePath, [&Callback](FName Path)
 	{
-		OutPathList.Add(PathName.ToString());
-	}
+		return Callback(Path.ToString());
+	}, bInRecurse);
+}
+
+void UAssetRegistryImpl::EnumerateSubPaths(const FName InBasePath, TFunctionRef<bool(FName)> Callback, bool bInRecurse) const
+{
+	CachedPathTree.EnumerateSubPaths(InBasePath, Callback, bInRecurse);
 }
 
 void UAssetRegistryImpl::RunAssetsThroughFilter(TArray<FAssetData>& AssetDataList, const FARFilter& Filter) const
@@ -1744,14 +1795,12 @@ void UAssetRegistryImpl::AppendState(const FAssetRegistryState& InState)
 	State.InitializeFromExisting(InState, SerializationOptions, FAssetRegistryState::EInitializationMode::Append);
 	CachePathsFromState(InState);
 
-	TArray<FAssetData> Assets;
-	InState.GetAllAssets(TSet<FName>(), Assets);
-	for (const FAssetData& Asset : Assets)
+	InState.EnumerateAllAssets(TSet<FName>(), [this](const FAssetData& AssetData)
 	{
 		// Let subscribers know that the new asset was added to the registry
-		AssetAddedEvent.Broadcast(Asset);
-	}
-
+		AssetAddedEvent.Broadcast(AssetData);
+		return true;
+	});
 }
 
 void UAssetRegistryImpl::CachePathsFromState(const FAssetRegistryState& InState)
diff --git a/Engine/Source/Runtime/AssetRegistry/Private/AssetRegistry.h b/Engine/Source/Runtime/AssetRegistry/Private/AssetRegistry.h
--- a/Engine/Source/Runtime/AssetRegistry/Private/AssetRegistry.h
+++ b/Engine/Source/Runtime/AssetRegistry/Private/AssetRegistry.h
@@ -40,8 +40,10 @@ public:
 	virtual bool GetAssetsByTags(const TArray<FName>& AssetTags, TArray<FAssetData>& OutAssetData) const override;
 	virtual bool GetAssetsByTagValues(const TMultiMap<FName, FString>& AssetTagsAndValues, TArray<FAssetData>& OutAssetData) const override;
 	virtual bool GetAssets(const FARFilter& Filter, TArray<FAssetData>& OutAssetData) const override;
+	virtual bool EnumerateAssets(const FARFilter& Filter, TFunctionRef<bool(const FAssetData&)> Callback) const override;
 	virtual FAssetData GetAssetByObjectPath( const FName ObjectPath, bool bIncludeOnlyOnDiskAssets = false ) const override;
 	virtual bool GetAllAssets(TArray<FAssetData>& OutAssetData, bool bIncludeOnlyOnDiskAssets = false) const override;
+	virtual bool EnumerateAllAssets(TFunctionRef<bool(const FAssetData&)> Callback, bool bIncludeOnlyOnDiskAssets = false) const override;
 	virtual bool GetDependencies(const FAssetIdentifier& AssetIdentifier, TArray<FAssetIdentifier>& OutDependencies, EAssetRegistryDependencyType::Type InDependencyType = EAssetRegistryDependencyType::All) const override;
 	virtual bool GetDependencies(FName PackageName, TArray<FName>& OutDependencies, EAssetRegistryDependencyType::Type InDependencyType = EAssetRegistryDependencyType::Packages) const override;
 	virtual bool GetReferencers(const FAssetIdentifier& AssetIdentifier, TArray<FAssetIdentifier>& OutReferencers, EAssetRegistryDependencyType::Type InReferenceType = EAssetRegistryDependencyType::All) const override;
@@ -52,7 +54,11 @@ public:
 	virtual bool GetAncestorClassNames(FName ClassName, TArray<FName>& OutAncestorClassNames) const override;
 	virtual void GetDerivedClassNames(const TArray<FName>& ClassNames, const TSet<FName>& ExcludedClassNames, TSet<FName>& OutDerivedClassNames) const override;
 	virtual void GetAllCachedPaths(TArray<FString>& OutPathList) const override;
+	virtual void EnumerateAllCachedPaths(TFunctionRef<bool(FString)> Callback) const override;
+	virtual void EnumerateAllCachedPaths(TFunctionRef<bool(FName)> Callback) const override;
 	virtual void GetSubPaths(const FString& InBasePath, TArray<FString>& OutPathList, bool bInRecurse) const override;
+	virtual void EnumerateSubPaths(const FString& InBasePath, TFunctionRef<bool(FString)> Callback, bool bInRecurse) const override;
+	virtual void EnumerateSubPaths(const FName InBasePath, TFunctionRef<bool(FName)> Callback, bool bInRecurse) const override;
 	virtual void RunAssetsThroughFilter (TArray<FAssetData>& AssetDataList, const FARFilter& Filter) const override;
 	virtual void UseFilterToExcludeAssets(TArray<FAssetData>& AssetDataList, const FARFilter& Filter) const override;
 	virtual void ExpandRecursiveFilter(const FARFilter& InFilter, FARFilter& ExpandedFilter) const override;
diff --git a/Engine/Source/Runtime/AssetRegistry/Private/AssetRegistryState.cpp b/Engine/Source/Runtime/AssetRegistry/Private/AssetRegistryState.cpp
--- a/Engine/Source/Runtime/AssetRegistry/Private/AssetRegistryState.cpp
+++ b/Engine/Source/Runtime/AssetRegistry/Private/AssetRegistryState.cpp
@@ -560,7 +560,16 @@ bool FAssetRegistryState::HasAssets(const FName PackagePath) const
 
 bool FAssetRegistryState::GetAssets(const FARFilter& Filter, const TSet<FName>& PackageNamesToSkip, TArray<FAssetData>& OutAssetData) const
 {
-	// Verify filter input. If all assets are needed, use GetAllAssets() instead.
+	return EnumerateAssets(Filter, PackageNamesToSkip, [&OutAssetData](const FAssetData& AssetData)
+	{
+		OutAssetData.Emplace(AssetData);
+		return true;
+	});
+}
+
+bool FAssetRegistryState::EnumerateAssets(const FARFilter& Filter, const TSet<FName>& PackageNamesToSkip, TFunctionRef<bool(const FAssetData&)> Callback) const
+{
+	// Verify filter input. If all assets are needed, use EnumerateAllAssets() instead.
 	if (!IsFilterValid(Filter, false) || Filter.IsEmpty())
 	{
 		return false;
@@ -718,7 +727,10 @@ bool FAssetRegistryState::GetAssets(const FARFilter& Filter, const TSet<FName>&
 				continue;
 			}
 
-			OutAssetData.Add(*AssetData);
+			if (!Callback(*AssetData))
+			{
+				return true;
+			}
 		}
 	}
 
@@ -726,6 +738,15 @@ bool FAssetRegistryState::GetAssets(const FARFilter& Filter, const TSet<FName>&
 }
 
 bool FAssetRegistryState::GetAllAssets(const TSet<FName>& PackageNamesToSkip, TArray<FAssetData>& OutAssetData) const
+{
+	return EnumerateAllAssets(PackageNamesToSkip, [&OutAssetData](const FAssetData& AssetData)
+	{
+		OutAssetData.Emplace(AssetData);
+		return true;
+	});
+}
+
+bool FAssetRegistryState::EnumerateAllAssets(const TSet<FName>& PackageNamesToSkip, TFunctionRef<bool(const FAssetData&)> Callback) const
 {
 	// All unloaded disk assets
 	for (const TPair<FName, FAssetData*>& AssetDataPair : CachedAssetsByObjectPath)
@@ -737,7 +758,10 @@ bool FAssetRegistryState::GetAllAssets(const TSet<FName>& PackageNamesToSkip, TA
 			// Make sure the asset's package was not loaded then the object was deleted/renamed
 			if (!PackageNamesToSkip.Contains(AssetData->PackageName))
 			{
-				OutAssetData.Emplace(*AssetData);
+				if (!Callback(*AssetData))
+				{
+					return true;
+				}
 			}
 		}
 	}
diff --git a/Engine/Source/Runtime/AssetRegistry/Private/PathTree.cpp b/Engine/Source/Runtime/AssetRegistry/Private/PathTree.cpp
--- a/Engine/Source/Runtime/AssetRegistry/Private/PathTree.cpp
+++ b/Engine/Source/Runtime/AssetRegistry/Private/PathTree.cpp
@@ -152,17 +152,36 @@ bool FPathTree::RemovePath(FName Path)
 bool FPathTree::GetAllPaths(TSet<FName>& OutPaths) const
 {
 	OutPaths.Reset();
-	OutPaths.Reserve(ParentPathToChildPaths.Num());
+	EnumerateAllPaths([&OutPaths](FName Path)
+	{
+		OutPaths.Emplace(Path);
+		return true;
+	});
+	return OutPaths.Num() > 0;
+}
 
+void FPathTree::EnumerateAllPaths(TFunctionRef<bool(FName)> Callback) const
+{
 	for (const auto& PathPair : ParentPathToChildPaths)
 	{
-		OutPaths.Add(PathPair.Key);
+		if (!Callback(PathPair.Key))
+		{
+			return;
+		}
 	}
-
-	return OutPaths.Num() > 0;
 }
 
 bool FPathTree::GetSubPaths(FName BasePath, TSet<FName>& OutPaths, bool bRecurse) const
+{
+	const int32 OutPathsOriginalNum = OutPaths.Num();
+	return EnumerateSubPaths(BasePath, [&OutPaths](FName Path)
+	{
+		OutPaths.Emplace(Path);
+		return true;
+	}, bRecurse) && OutPaths.Num() > OutPathsOriginalNum;
+}
+
+bool FPathTree::EnumerateSubPaths(FName BasePath, TFunctionRef<bool(FName)> Callback, bool bRecurse) const
 {
 	if (BasePath.IsNone())
 	{
@@ -192,19 +211,20 @@ bool FPathTree::GetSubPaths(FName BasePath, TSet<FName>& OutPaths, bool bRecurse
 		}
 	}
 
-	const int32 OutPathsOriginalNum = OutPaths.Num();
-
 	for (const FName& ChildPath : *ChildPathsPtr)
 	{
 		check(ParentPathToChildPaths.Contains(ChildPath)); // This failing is an integrity violation as this entry lists a child that we don't know about
 
-		OutPaths.Add(ChildPath);
+		if (!Callback(ChildPath))
+		{
+			return true;
+		}
 
 		if (bRecurse)
 		{
-			GetSubPaths(ChildPath, OutPaths, /*bRecurse=*/true);
+			EnumerateSubPaths(ChildPath, Callback, /*bRecurse=*/true);
 		}
 	}
 
-	return OutPaths.Num() > OutPathsOriginalNum;
+	return true;
 }
diff --git a/Engine/Source/Runtime/AssetRegistry/Private/PathTree.h b/Engine/Source/Runtime/AssetRegistry/Private/PathTree.h
--- a/Engine/Source/Runtime/AssetRegistry/Private/PathTree.h
+++ b/Engine/Source/Runtime/AssetRegistry/Private/PathTree.h
@@ -16,9 +16,15 @@ public:
 	/** Get all of the paths we know about */
 	bool GetAllPaths(TSet<FName>& OutPaths) const;
 
+	/** Enumerate all of the paths we know about */
+	void EnumerateAllPaths(TFunctionRef<bool(FName)> Callback) const;
+
 	/** Recursively gathers all child paths from the specified base path relative to this node */
 	bool GetSubPaths(FName BasePath, TSet<FName>& OutPaths, bool bRecurse = true) const;
 
+	/** Recursively enumerates all child paths from the specified base path relative to this node */
+	bool EnumerateSubPaths(FName BasePath, TFunctionRef<bool(FName)> Callback, bool bRecurse = true) const;
+
 	uint32 GetAllocatedSize(void) const
 	{
 		uint32 AllocatedSize = ParentPathToChildPaths.GetAllocatedSize() + ChildPathToParentPath.GetAllocatedSize();
diff --git a/Engine/Source/Runtime/AssetRegistry/Public/AssetRegistryState.h b/Engine/Source/Runtime/AssetRegistry/Public/AssetRegistryState.h
--- a/Engine/Source/Runtime/AssetRegistry/Public/AssetRegistryState.h
+++ b/Engine/Source/Runtime/AssetRegistry/Public/AssetRegistryState.h
@@ -108,6 +108,17 @@ public:
 	 */
 	bool GetAssets(const FARFilter& Filter, const TSet<FName>& PackageNamesToSkip, TArray<FAssetData>& OutAssetData) const;
 
+	/**
+	 * Enumerate asset data for all assets that match the filter.
+	 * Assets returned must satisfy every filter component if there is at least one element in the component's array.
+	 * Assets will satisfy a component if they match any of the elements in it.
+	 *
+	 * @param Filter filter to apply to the assets in the AssetRegistry
+	 * @param PackageNamesToSkip explicit list of packages to skip, because they were already added
+	 * @param Callback function to call for each asset data enumerated
+	 */
+	bool EnumerateAssets(const FARFilter& Filter, const TSet<FName>& PackageNamesToSkip, TFunctionRef<bool(const FAssetData&)> Callback) const;
+
 	/**
 	 * Gets asset data for all assets in the registry state.
 	 *
@@ -116,6 +127,14 @@ public:
 	 */
 	bool GetAllAssets(const TSet<FName>& PackageNamesToSkip, TArray<FAssetData>& OutAssetData) const;
 
+	/**
+	 * Enumerates asset data for all assets in the registry state.
+	 *
+	 * @param PackageNamesToSkip explicit list of packages to skip, because they were already added
+	 * @param Callback function to call for each asset data enumerated
+	 */
+	bool EnumerateAllAssets(const TSet<FName>& PackageNamesToSkip, TFunctionRef<bool(const FAssetData&)> Callback) const;
+
 	/**
 	 * Gets a list of packages and searchable names that are referenced by the supplied package or name. (On disk references ONLY)
 	 *
diff --git a/Engine/Source/Runtime/AssetRegistry/Public/IAssetRegistry.h b/Engine/Source/Runtime/AssetRegistry/Public/IAssetRegistry.h
--- a/Engine/Source/Runtime/AssetRegistry/Public/IAssetRegistry.h
+++ b/Engine/Source/Runtime/AssetRegistry/Public/IAssetRegistry.h
@@ -175,6 +175,16 @@ public:
 	UFUNCTION(BlueprintCallable, BlueprintPure=false, Category="AssetRegistry")
 	virtual bool GetAssets(const FARFilter& Filter, TArray<FAssetData>& OutAssetData) const = 0;
 
+	/**
+	 * Enumerate asset data for all assets that match the filter.
+	 * Assets returned must satisfy every filter component if there is at least one element in the component's array.
+	 * Assets will satisfy a component if they match any of the elements in it.
+	 *
+	 * @param Filter filter to apply to the assets in the AssetRegistry
+	 * @param Callback function to call for each asset data enumerated
+	 */
+	virtual bool EnumerateAssets(const FARFilter& Filter, TFunctionRef<bool(const FAssetData&)> Callback) const = 0;
+
 	/**
 	 * Gets the asset data for the specified object path
 	 *
@@ -194,6 +204,14 @@ public:
 	UFUNCTION(BlueprintCallable, BlueprintPure=false, Category = "AssetRegistry")
 	virtual bool GetAllAssets(TArray<FAssetData>& OutAssetData, bool bIncludeOnlyOnDiskAssets = false) const = 0;
 
+	/**
+	 * Enumerate asset data for all assets in the registry.
+	 * This method may be slow, use a filter if possible to avoid iterating over the entire registry.
+	 *
+	 * @param Callback function to call for each asset data enumerated
+	 */
+	virtual bool EnumerateAllAssets(TFunctionRef<bool(const FAssetData&)> Callback, bool bIncludeOnlyOnDiskAssets = false) const = 0;
+
 	/**
 	 * Gets a list of packages and searchable names that are referenced by the supplied package or name. (On disk references ONLY)
 	 *
@@ -281,10 +299,22 @@ public:
 	UFUNCTION(BlueprintCallable, BlueprintPure=false, Category = "AssetRegistry")
 	virtual void GetAllCachedPaths(TArray<FString>& OutPathList) const = 0;
 
+	/** Enumerate all the paths that are currently cached */
+	virtual void EnumerateAllCachedPaths(TFunctionRef<bool(FString)> Callback) const = 0;
+
+	/** Enumerate all the paths that are currently cached */
+	virtual void EnumerateAllCachedPaths(TFunctionRef<bool(FName)> Callback) const = 0;
+
 	/** Gets a list of all paths that are currently cached below the passed-in base path */
 	UFUNCTION(BlueprintCallable, BlueprintPure=false, Category = "AssetRegistry")
 	virtual void GetSubPaths(const FString& InBasePath, TArray<FString>& OutPathList, bool bInRecurse) const = 0;
 
+	/** Enumerate the all paths that are currently cached below the passed-in base path */
+	virtual void EnumerateSubPaths(const FString& InBasePath, TFunctionRef<bool(FString)> Callback, bool bInRecurse) const = 0;
+
+	/** Enumerate the all paths that are currently cached below the passed-in base path */
+	virtual void EnumerateSubPaths(const FName InBasePath, TFunctionRef<bool(FName)> Callback, bool bInRecurse) const = 0;
+
 	/** Trims items out of the asset data list that do not pass the supplied filter */
 	UFUNCTION(BlueprintCallable, BlueprintPure=false, Category = "AssetRegistry")
 	virtual void RunAssetsThroughFilter(UPARAM(ref) TArray<FAssetData>& AssetDataList, const FARFilter& Filter) const = 0;
