Cleaned-up logic for running assets through backend filters

There is now a common implementation for the inclusive and exclusive filtering modes, and an API that allows testing both a single item or a list of items.

diff --git a/Engine/Source/Runtime/AssetRegistry/Private/AssetRegistry.cpp b/Engine/Source/Runtime/AssetRegistry/Private/AssetRegistry.cpp
--- a/Engine/Source/Runtime/AssetRegistry/Private/AssetRegistry.cpp
+++ b/Engine/Source/Runtime/AssetRegistry/Private/AssetRegistry.cpp
@@ -1039,321 +1039,137 @@ void UAssetRegistryImpl::EnumerateSubPaths(const FName InBasePath, TFunctionRef<
 
 void UAssetRegistryImpl::RunAssetsThroughFilter(TArray<FAssetData>& AssetDataList, const FARFilter& Filter) const
 {
-	if (!Filter.IsEmpty())
-	{
-		TSet<FName> RequestedClassNames;
-		if (Filter.bRecursiveClasses && Filter.ClassNames.Num() > 0)
-		{
-			// First assemble a full list of requested classes from the ClassTree
-			// GetSubClasses includes the base classes
-			GetSubClasses(Filter.ClassNames, Filter.RecursiveClassesExclusionSet, RequestedClassNames);
-		}
-
-		for (int32 AssetDataIdx = AssetDataList.Num() - 1; AssetDataIdx >= 0; --AssetDataIdx)
-		{
-			const FAssetData& AssetData = AssetDataList[AssetDataIdx];
-
-			// Package Names
-			if (Filter.PackageNames.Num() > 0)
-			{
-				bool bPassesPackageNames = false;
-				for (int32 NameIdx = 0; NameIdx < Filter.PackageNames.Num(); ++NameIdx)
-				{
-					if (Filter.PackageNames[NameIdx] == AssetData.PackageName)
-					{
-						bPassesPackageNames = true;
-						break;
-					}
-				}
-
-				if (!bPassesPackageNames)
-				{
-					AssetDataList.RemoveAt(AssetDataIdx);
-					continue;
-				}
-			}
-
-			// Package Paths
-			if (Filter.PackagePaths.Num() > 0)
-			{
-				bool bPassesPackagePaths = false;
-				if (Filter.bRecursivePaths)
-				{
-					FString AssetPackagePath = AssetData.PackagePath.ToString();
-					for (int32 PathIdx = 0; PathIdx < Filter.PackagePaths.Num(); ++PathIdx)
-					{
-						const FString Path = Filter.PackagePaths[PathIdx].ToString();
-						if (AssetPackagePath.StartsWith(Path))
-						{
-							// Only match the exact path or a path that starts with the target path followed by a slash
-							if (Path.Len() == 1 || Path.Len() == AssetPackagePath.Len() || AssetPackagePath.Mid(Path.Len(), 1) == TEXT("/"))
-							{
-								bPassesPackagePaths = true;
-								break;
-							}
-						}
-					}
-				}
-				else
-				{
-					// Non-recursive. Just request data for each requested path.
-					for (int32 PathIdx = 0; PathIdx < Filter.PackagePaths.Num(); ++PathIdx)
-					{
-						if (Filter.PackagePaths[PathIdx] == AssetData.PackagePath)
-						{
-							bPassesPackagePaths = true;
-							break;
-						}
-					}
-				}
-
-				if (!bPassesPackagePaths)
-				{
-					AssetDataList.RemoveAt(AssetDataIdx);
-					continue;
-				}
-			}
-
-			// ObjectPaths
-			if (Filter.ObjectPaths.Num() > 0)
-			{
-				bool bPassesObjectPaths = Filter.ObjectPaths.Contains(AssetData.ObjectPath);
-
-				if (!bPassesObjectPaths)
-				{
-					AssetDataList.RemoveAt(AssetDataIdx);
-					continue;
-				}
-			}
+	RunAssetsThroughFilterImpl(AssetDataList, Filter, EARFilterMode::Inclusive);
+}
 
-			// Classes
-			if (Filter.ClassNames.Num() > 0)
-			{
-				bool bPassesClasses = false;
-				if (Filter.bRecursiveClasses)
-				{
-					// Now check against each discovered class
-					for (FName ClassName : RequestedClassNames)
-					{
-						if (ClassName == AssetData.AssetClass)
-						{
-							bPassesClasses = true;
-							break;
-						}
-					}
-				}
-				else
-				{
-					// Non-recursive. Just request data for each requested classes.
-					for (int32 ClassIdx = 0; ClassIdx < Filter.ClassNames.Num(); ++ClassIdx)
-					{
-						if (Filter.ClassNames[ClassIdx] == AssetData.AssetClass)
-						{
-							bPassesClasses = true;
-							break;
-						}
-					}
-				}
+void UAssetRegistryImpl::UseFilterToExcludeAssets(TArray<FAssetData>& AssetDataList, const FARFilter& Filter) const
+{
+	RunAssetsThroughFilterImpl(AssetDataList, Filter, EARFilterMode::Exclusive);
+}
 
-				if (!bPassesClasses)
-				{
-					AssetDataList.RemoveAt(AssetDataIdx);
-					continue;
-				}
-			}
+bool UAssetRegistryImpl::IsAssetIncludedByFilter(const FAssetData& AssetData, const FARFilter& Filter) const
+{
+	return RunAssetThroughFilterImpl(AssetData, Filter, EARFilterMode::Inclusive);
+}
 
-			// Tags and values
-			if (Filter.TagsAndValues.Num() > 0)
-			{
-				bool bPassesTags = false;
-				for (auto FilterTagIt = Filter.TagsAndValues.CreateConstIterator(); FilterTagIt; ++FilterTagIt)
-				{
-					bool bAccept;
-					if (!FilterTagIt.Value().IsSet())
-					{
-						// this probably doesn't make sense, but I am preserving the original logic
-						bAccept = AssetData.TagsAndValues.ContainsKeyValue(FilterTagIt.Key(), FString());
-					}
-					else
-					{
-						bAccept = AssetData.TagsAndValues.ContainsKeyValue(FilterTagIt.Key(), FilterTagIt.Value().GetValue());
-					}
+bool UAssetRegistryImpl::IsAssetExcludedByFilter(const FAssetData& AssetData, const FARFilter& Filter) const
+{
+	return RunAssetThroughFilterImpl(AssetData, Filter, EARFilterMode::Exclusive);
+}
 
-					if (bAccept)
-					{
-						bPassesTags = true;
-						break;
-					}
-				}
+bool UAssetRegistryImpl::RunAssetThroughFilterImpl(const FAssetData& AssetData, const FARFilter& Filter, const EARFilterMode FilterMode) const
+{
+	checkf(!Filter.IsRecursive(), TEXT("Filter should have been expanded with ExpandRecursiveFilter before calling this function!"));
 
-				if (!bPassesTags)
-				{
-					AssetDataList.RemoveAt(AssetDataIdx);
-					continue;
-				}
-			}
-		}
+	const bool bPassFilterValue = FilterMode == EARFilterMode::Inclusive;
+	if (Filter.IsEmpty())
+	{
+		return bPassFilterValue;
 	}
-}
 
+	const bool bFilterResult = RunAssetThroughFilterImpl_Unchecked(AssetData, Filter, bPassFilterValue);
+	return bFilterResult == bPassFilterValue;
+}
 
-void UAssetRegistryImpl::UseFilterToExcludeAssets(TArray<FAssetData>& AssetDataList, const FARFilter& Filter) const
+bool UAssetRegistryImpl::RunAssetThroughFilterImpl_Unchecked(const FAssetData& AssetData, const FARFilter& Filter, const bool bPassFilterValue) const
 {
-	if (!Filter.IsEmpty())
+	// Package Names
+	if (Filter.PackageNames.Num() > 0)
 	{
-		TSet<FName> RequestedClassNames;
-		if (Filter.bRecursiveClasses && Filter.ClassNames.Num() > 0)
+		const bool bPassesPackageNames = Filter.PackageNames.Contains(AssetData.PackageName);
+		if (bPassesPackageNames != bPassFilterValue)
 		{
-			// First assemble a full list of requested classes from the ClassTree
-			// GetSubClasses includes the base classes
-			GetSubClasses(Filter.ClassNames, Filter.RecursiveClassesExclusionSet, RequestedClassNames);
+			return !bPassFilterValue;
 		}
+	}
 
-		for (int32 AssetDataIdx = AssetDataList.Num() - 1; AssetDataIdx >= 0; --AssetDataIdx)
+	// Package Paths
+	if (Filter.PackagePaths.Num() > 0)
+	{
+		const bool bPassesPackagePaths = Filter.PackagePaths.Contains(AssetData.PackagePath);
+		if (bPassesPackagePaths != bPassFilterValue)
 		{
-			const FAssetData& AssetData = AssetDataList[AssetDataIdx];
-
-			// Package Names
-			if (Filter.PackageNames.Num() > 0)
-			{
-				bool bPassesPackageNames = false;
-				for (int32 NameIdx = 0; NameIdx < Filter.PackageNames.Num(); ++NameIdx)
-				{
-					if (Filter.PackageNames[NameIdx] == AssetData.PackageName)
-					{
-						bPassesPackageNames = true;
-						break;
-					}
-				}
-
-				if (bPassesPackageNames)
-				{
-					AssetDataList.RemoveAt(AssetDataIdx);
-					continue;
-				}
-			}
+			return !bPassFilterValue;
+		}
+	}
 
-			// Package Paths
-			if (Filter.PackagePaths.Num() > 0)
-			{
-				bool bPassesPackagePaths = false;
-				if (Filter.bRecursivePaths)
-				{
-					FString AssetPackagePath = AssetData.PackagePath.ToString();
-					for (int32 PathIdx = 0; PathIdx < Filter.PackagePaths.Num(); ++PathIdx)
-					{
-						const FString Path = Filter.PackagePaths[PathIdx].ToString();
-						if (AssetPackagePath.StartsWith(Path))
-						{
-							// Only match the exact path or a path that starts with the target path followed by a slash
-							if (Path.Len() == 1 || Path.Len() == AssetPackagePath.Len() || AssetPackagePath.Mid(Path.Len(), 1) == TEXT("/"))
-							{
-								bPassesPackagePaths = true;
-								break;
-							}
-						}
-					}
-				}
-				else
-				{
-					// Non-recursive. Just request data for each requested path.
-					for (int32 PathIdx = 0; PathIdx < Filter.PackagePaths.Num(); ++PathIdx)
-					{
-						if (Filter.PackagePaths[PathIdx] == AssetData.PackagePath)
-						{
-							bPassesPackagePaths = true;
-							break;
-						}
-					}
-				}
+	// ObjectPaths
+	if (Filter.ObjectPaths.Num() > 0)
+	{
+		const bool bPassesObjectPaths = Filter.ObjectPaths.Contains(AssetData.ObjectPath);
+		if (bPassesObjectPaths != bPassFilterValue)
+		{
+			return !bPassFilterValue;
+		}
+	}
 
-				if (bPassesPackagePaths)
-				{
-					AssetDataList.RemoveAt(AssetDataIdx);
-					continue;
-				}
-			}
+	// Classes
+	if (Filter.ClassNames.Num() > 0)
+	{
+		const bool bPassesClasses = Filter.ClassNames.Contains(AssetData.AssetClass);
+		if (bPassesClasses != bPassFilterValue)
+		{
+			return !bPassFilterValue;
+		}
+	}
 
-			// ObjectPaths
-			if (Filter.ObjectPaths.Num() > 0)
+	// Tags and values
+	if (Filter.TagsAndValues.Num() > 0)
+	{
+		bool bPassesTags = false;
+		for (const auto& TagsAndValuePair : Filter.TagsAndValues)
+		{
+			bPassesTags |= TagsAndValuePair.Value.IsSet()
+				? AssetData.TagsAndValues.ContainsKeyValue(TagsAndValuePair.Key, TagsAndValuePair.Value.GetValue())
+				: AssetData.TagsAndValues.Contains(TagsAndValuePair.Key);
+			if (bPassesTags)
 			{
-				bool bPassesObjectPaths = Filter.ObjectPaths.Contains(AssetData.ObjectPath);
-
-				if (bPassesObjectPaths)
-				{
-					AssetDataList.RemoveAt(AssetDataIdx);
-					continue;
-				}
+				break;
 			}
+		}
+		if (bPassesTags != bPassFilterValue)
+		{
+			return !bPassFilterValue;
+		}
+	}
 
-			// Classes
-			if (Filter.ClassNames.Num() > 0)
-			{
-				bool bPassesClasses = false;
-				if (Filter.bRecursiveClasses)
-				{
-					// Now check against each discovered class
-					for (FName ClassName : RequestedClassNames)
-					{
-						if (ClassName == AssetData.AssetClass)
-						{
-							bPassesClasses = true;
-							break;
-						}
-					}
-				}
-				else
-				{
-					// Non-recursive. Just request data for each requested classes.
-					for (int32 ClassIdx = 0; ClassIdx < Filter.ClassNames.Num(); ++ClassIdx)
-					{
-						if (Filter.ClassNames[ClassIdx] == AssetData.AssetClass)
-						{
-							bPassesClasses = true;
-							break;
-						}
-					}
-				}
+	return bPassFilterValue;
+}
 
-				if (bPassesClasses)
-				{
-					AssetDataList.RemoveAt(AssetDataIdx);
-					continue;
-				}
-			}
+void UAssetRegistryImpl::RunAssetsThroughFilterImpl(TArray<FAssetData>& AssetDataList, const FARFilter& Filter, const EARFilterMode FilterMode) const
+{
+	if (Filter.IsEmpty() || !FAssetRegistryState::IsFilterValid(Filter, true))
+	{
+		return;
+	}
 
-			// Tags and values
-			if (Filter.TagsAndValues.Num() > 0)
-			{
-				bool bPassesTags = false;
-				for (auto FilterTagIt = Filter.TagsAndValues.CreateConstIterator(); FilterTagIt; ++FilterTagIt)
-				{
-					bool bAccept;
-					if (!FilterTagIt.Value().IsSet())
-					{
-						// this probably doesn't make sense, but I am preserving the original logic
-						bAccept = AssetData.TagsAndValues.ContainsKeyValue(FilterTagIt.Key(), FString());
-					}
-					else
-					{
-						bAccept = AssetData.TagsAndValues.ContainsKeyValue(FilterTagIt.Key(), FilterTagIt.Value().GetValue());
-					}
+	const FARFilter* FilterPtr = &Filter;
+	FARFilter TempExpandedFilter;
+	if (Filter.IsRecursive())
+	{
+		ExpandRecursiveFilter(Filter, TempExpandedFilter);
+		FilterPtr = &TempExpandedFilter;
+	}
+	check(FilterPtr);
 
-					if (bAccept)
-					{
-						bPassesTags = true;
-						break;
-					}
-				}
+	const int32 OriginalArrayCount = AssetDataList.Num();
+	const bool bPassFilterValue = FilterMode == EARFilterMode::Inclusive;
 
-				if (!bPassesTags)
-				{
-					AssetDataList.RemoveAt(AssetDataIdx);
-					continue;
-				}
-			}
+	// Spin the array backwards to minimize the number of elements that are repeatedly moved down
+	for (int32 AssetDataIndex = AssetDataList.Num() - 1; AssetDataIndex >= 0; --AssetDataIndex)
+	{
+		const bool bFilterResult = RunAssetThroughFilterImpl_Unchecked(AssetDataList[AssetDataIndex], *FilterPtr, bPassFilterValue);
+		if (bFilterResult != bPassFilterValue)
+		{
+			AssetDataList.RemoveAt(AssetDataIndex, 1, /*bAllowShrinking*/false);
+			continue;
 		}
 	}
+
+	if (OriginalArrayCount > AssetDataList.Num())
+	{
+		AssetDataList.Shrink();
+	}
 }
 
 void UAssetRegistryImpl::ExpandRecursiveFilter(const FARFilter& InFilter, FARFilter& ExpandedFilter) const
diff --git a/Engine/Source/Runtime/AssetRegistry/Private/AssetRegistry.h b/Engine/Source/Runtime/AssetRegistry/Private/AssetRegistry.h
--- a/Engine/Source/Runtime/AssetRegistry/Private/AssetRegistry.h
+++ b/Engine/Source/Runtime/AssetRegistry/Private/AssetRegistry.h
@@ -61,6 +61,8 @@ public:
 	virtual void EnumerateSubPaths(const FName InBasePath, TFunctionRef<bool(FName)> Callback, bool bInRecurse) const override;
 	virtual void RunAssetsThroughFilter (TArray<FAssetData>& AssetDataList, const FARFilter& Filter) const override;
 	virtual void UseFilterToExcludeAssets(TArray<FAssetData>& AssetDataList, const FARFilter& Filter) const override;
+	virtual bool IsAssetIncludedByFilter(const FAssetData& AssetData, const FARFilter& Filter) const override;
+	virtual bool IsAssetExcludedByFilter(const FAssetData& AssetData, const FARFilter& Filter) const override;
 	virtual void ExpandRecursiveFilter(const FARFilter& InFilter, FARFilter& ExpandedFilter) const override;
 	virtual void SetTemporaryCachingMode(bool bEnable) override;
 	virtual bool GetTemporaryCachingMode() const override;
@@ -248,6 +250,30 @@ private:
 	/** Internal helper which processes a given state and adds its contents to the current registry */
 	void CachePathsFromState(const FAssetRegistryState& InState);
 
+	enum class EARFilterMode : uint8
+	{
+		/** Include things that pass the filter; include everything if the filter is empty */
+		Inclusive,
+		/** Exclude things that pass the filter; exclude nothing if the filter is empty */
+		Exclusive,
+	};
+
+	/**
+	 * Given an asset data, say whether it would pass the filter based on the inclusion/exclusion mode used.
+	 *  - If an asset data passes a filter, then in inclusive mode it will return true, and in exclusive mode it will return false.
+	 *  - If an asset data fails a filter, then in inclusive mode it will return false, and in exclusive mode it will return true.
+	 *  - If the filter is empty, then in inclusive mode it will return true, and in exclusive mode it will return false.
+	 */
+	bool RunAssetThroughFilterImpl(const FAssetData& AssetData, const FARFilter& Filter, const EARFilterMode FilterMode) const;
+	bool RunAssetThroughFilterImpl_Unchecked(const FAssetData& AssetData, const FARFilter& Filter, const bool bPassFilterValue) const;
+
+	/**
+	 * Given an array of asset data, trim the items that fail the filter based on the inclusion/exclusion mode used.
+	 *  - In inclusive mode it will remove all assets that fail the filter, and in exclusive mode it will remove all assets that pass the filter.
+	 *  - If the filter is empty, then the array will be untouched.
+	 */
+	void RunAssetsThroughFilterImpl(TArray<FAssetData>& AssetDataList, const FARFilter& Filter, const EARFilterMode FilterMode) const;
+
 private:
 	
 	/** Internal state of the cached asset registry */
diff --git a/Engine/Source/Runtime/AssetRegistry/Private/AssetRegistryState.cpp b/Engine/Source/Runtime/AssetRegistry/Private/AssetRegistryState.cpp
--- a/Engine/Source/Runtime/AssetRegistry/Private/AssetRegistryState.cpp
+++ b/Engine/Source/Runtime/AssetRegistry/Private/AssetRegistryState.cpp
@@ -1689,7 +1689,7 @@ bool FAssetRegistryState::IsFilterValid(const FARFilter& Filter, bool bAllowRecu
 		}
 	}
 
-	if (!bAllowRecursion && (Filter.bRecursiveClasses || Filter.bRecursivePaths))
+	if (!bAllowRecursion && Filter.IsRecursive())
 	{
 		return false;
 	}
diff --git a/Engine/Source/Runtime/AssetRegistry/Public/ARFilter.h b/Engine/Source/Runtime/AssetRegistry/Public/ARFilter.h
--- a/Engine/Source/Runtime/AssetRegistry/Public/ARFilter.h
+++ b/Engine/Source/Runtime/AssetRegistry/Public/ARFilter.h
@@ -40,22 +40,15 @@ struct FARFilter
 
 	/** If true, PackagePath components will be recursive */
 	UPROPERTY(EditDefaultsOnly, BlueprintReadWrite, Category=AssetRegistry)
-	bool bRecursivePaths;
+	bool bRecursivePaths = false;
 
 	/** If true, subclasses of ClassNames will also be included and RecursiveClassesExclusionSet will be excluded. */
 	UPROPERTY(EditDefaultsOnly, BlueprintReadWrite, Category=AssetRegistry)
-	bool bRecursiveClasses;
+	bool bRecursiveClasses = false;
 
 	/** If true, only on-disk assets will be returned. Be warned that this is rarely what you want and should only be used for performance reasons */
 	UPROPERTY(EditDefaultsOnly, BlueprintReadWrite, Category=AssetRegistry)
-	bool bIncludeOnlyOnDiskAssets;
-
-	FARFilter()
-	{
-		bRecursivePaths = false;
-		bRecursiveClasses = false;
-		bIncludeOnlyOnDiskAssets = false;
-	}
+	bool bIncludeOnlyOnDiskAssets = false;
 
 	/** Appends the other filter to this one */
 	void Append(const FARFilter& Other)
@@ -83,6 +76,12 @@ struct FARFilter
 		return PackageNames.Num() + PackagePaths.Num() + ObjectPaths.Num() + ClassNames.Num() + TagsAndValues.Num() == 0;
 	}
 
+	/** Returns true if this filter is recursive */
+	bool IsRecursive() const
+	{
+		return bRecursivePaths || bRecursiveClasses;
+	}
+
 	/** Clears this filter of all entries */
 	void Clear()
 	{
@@ -99,4 +98,4 @@ struct FARFilter
 
 		ensure(IsEmpty());
 	}
-};
\ No newline at end of file
+};
diff --git a/Engine/Source/Runtime/AssetRegistry/Public/IAssetRegistry.h b/Engine/Source/Runtime/AssetRegistry/Public/IAssetRegistry.h
--- a/Engine/Source/Runtime/AssetRegistry/Public/IAssetRegistry.h
+++ b/Engine/Source/Runtime/AssetRegistry/Public/IAssetRegistry.h
@@ -319,10 +319,16 @@ public:
 	UFUNCTION(BlueprintCallable, BlueprintPure=false, Category = "AssetRegistry")
 	virtual void RunAssetsThroughFilter(UPARAM(ref) TArray<FAssetData>& AssetDataList, const FARFilter& Filter) const = 0;
 
-	/** Trims items out of the asset data list that do not pass the supplied filter */
+	/** Trims items out of the asset data list that pass the supplied filter */
 	UFUNCTION(BlueprintCallable, BlueprintPure=false, Category = "AssetRegistry")
 	virtual void UseFilterToExcludeAssets(UPARAM(ref) TArray<FAssetData>& AssetDataList, const FARFilter& Filter) const = 0;
 
+	/** Tests to see whether the given asset would be included (passes) the given filter (note: the filter should have been expanded first, @see ExpandRecursiveFilter */
+	virtual bool IsAssetIncludedByFilter(const FAssetData& AssetData, const FARFilter& Filter) const = 0;
+
+	/** Tests to see whether the given asset would be excluded (fails) the given filter (note: the filter should have been expanded first, @see ExpandRecursiveFilter */
+	virtual bool IsAssetExcludedByFilter(const FAssetData& AssetData, const FARFilter& Filter) const = 0;
+
 	/** Modifies passed in filter to make it safe for use on FAssetRegistryState. This expands recursive paths and classes */
 	virtual void ExpandRecursiveFilter(const FARFilter& InFilter, FARFilter& ExpandedFilter) const = 0;
 
