Pre-compile the FARFilter before using it with FAssetRegistryState

This avoids a double TArray->TSet conversion during filtering, and allows the pre-compiled filter to be re-used repeatedly for things like IsAssetIncludedByFilter and IsAssetExcludedByFilter

diff --git a/Engine/Source/Runtime/AssetRegistry/Private/AssetRegistry.cpp b/Engine/Source/Runtime/AssetRegistry/Private/AssetRegistry.cpp
--- a/Engine/Source/Runtime/AssetRegistry/Private/AssetRegistry.cpp
+++ b/Engine/Source/Runtime/AssetRegistry/Private/AssetRegistry.cpp
@@ -626,30 +626,23 @@ bool UAssetRegistryImpl::GetAssets(const FARFilter& InFilter, TArray<FAssetData>
 bool UAssetRegistryImpl::EnumerateAssets(const FARFilter& InFilter, TFunctionRef<bool(const FAssetData&)> Callback) const
 {
 	// Verify filter input. If all assets are needed, use EnumerateAllAssets() instead.
-	if (!FAssetRegistryState::IsFilterValid(InFilter, true) || InFilter.IsEmpty())
+	if (InFilter.IsEmpty())
 	{
 		return false;
 	}
 
 	// Expand recursion on filter
-	FARFilter Filter;
-	ExpandRecursiveFilter(InFilter, Filter);
+	FARCompiledFilter CompiledFilter;
+	CompileFilter(InFilter, CompiledFilter);
+	if (!FAssetRegistryState::IsFilterValid(CompiledFilter))
+	{
+		return false;
+	}
 
 	// Start with in memory assets
 	TSet<FName> PackagesToSkip = CachedEmptyPackages;
-
-	if (!Filter.bIncludeOnlyOnDiskAssets)
+	if (!CompiledFilter.bIncludeOnlyOnDiskAssets)
 	{
-		// Prepare a set of each filter component for fast searching
-		TSet<FName> FilterPackageNames(Filter.PackageNames);
-		TSet<FName> FilterPackagePaths(Filter.PackagePaths);
-		TSet<FName> FilterClassNames(Filter.ClassNames);
-		TSet<FName> FilterObjectPaths(Filter.ObjectPaths);
-		const int32 NumFilterPackageNames = FilterPackageNames.Num();
-		const int32 NumFilterPackagePaths = FilterPackagePaths.Num();
-		const int32 NumFilterClasses = FilterClassNames.Num();
-		const int32 NumFilterObjectPaths = FilterObjectPaths.Num();
-
 		// Reusable structures to avoid memory allocations
 		TArray<UObject::FAssetRegistryTag> ObjectTags;
 
@@ -670,16 +663,16 @@ bool UAssetRegistryImpl::EnumerateAssets(const FARFilter& InFilter, TFunctionRef
 
 				PackagesToSkip.Add(PackageName);
 
-				if (NumFilterPackageNames && !FilterPackageNames.Contains(PackageName))
+				if (CompiledFilter.PackageNames.Num() && !CompiledFilter.PackageNames.Contains(PackageName))
 				{
 					return;
 				}
 
 				// Object Path
-				if (NumFilterObjectPaths)
+				if (CompiledFilter.ObjectPaths.Num() > 0)
 				{
 					const FName ObjectPath = FName(*Obj->GetPathName(), FNAME_Find);
-					if (!FilterObjectPaths.Contains(ObjectPath))
+					if (!CompiledFilter.ObjectPaths.Contains(ObjectPath))
 					{
 						return;
 					}
@@ -687,7 +680,7 @@ bool UAssetRegistryImpl::EnumerateAssets(const FARFilter& InFilter, TFunctionRef
 
 				// Package path
 				const FName PackagePath = FName(*FPackageName::GetLongPackagePath(InMemoryPackage->GetName()));
-				if (NumFilterPackagePaths && !FilterPackagePaths.Contains(PackagePath))
+				if (CompiledFilter.PackagePaths.Num() > 0 && !CompiledFilter.PackagePaths.Contains(PackagePath))
 				{
 					return;
 				}
@@ -695,10 +688,10 @@ bool UAssetRegistryImpl::EnumerateAssets(const FARFilter& InFilter, TFunctionRef
 				// Tags and values
 				check(ObjectTags.Num() == 0);
 				Obj->GetAssetRegistryTags(ObjectTags);
-				if (Filter.TagsAndValues.Num())
+				if (CompiledFilter.TagsAndValues.Num() > 0)
 				{
 					bool bMatch = false;
-					for (auto FilterTagIt = Filter.TagsAndValues.CreateConstIterator(); FilterTagIt; ++FilterTagIt)
+					for (auto FilterTagIt = CompiledFilter.TagsAndValues.CreateConstIterator(); FilterTagIt; ++FilterTagIt)
 					{
 						const FName Tag = FilterTagIt.Key();
 						const TOptional<FString>& Value = FilterTagIt.Value();
@@ -746,10 +739,10 @@ bool UAssetRegistryImpl::EnumerateAssets(const FARFilter& InFilter, TFunctionRef
 		};
 
 		// Iterate over all in-memory assets to find the ones that pass the filter components
-		if (NumFilterClasses)
+		if (CompiledFilter.ClassNames.Num() > 0)
 		{
 			TArray<UObject*> InMemoryObjects;
-			for (FName ClassName : FilterClassNames)
+			for (FName ClassName : CompiledFilter.ClassNames)
 			{
 				UClass* Class = FindObjectFast<UClass>(nullptr, ClassName, false, true, RF_NoFlags);
 				if (Class != nullptr)
@@ -782,7 +775,7 @@ bool UAssetRegistryImpl::EnumerateAssets(const FARFilter& InFilter, TFunctionRef
 		}
 	}
 
-	State.EnumerateAssets(Filter, PackagesToSkip, Callback);
+	State.EnumerateAssets(CompiledFilter, PackagesToSkip, Callback);
 
 	return true;
 }
@@ -1047,20 +1040,18 @@ void UAssetRegistryImpl::UseFilterToExcludeAssets(TArray<FAssetData>& AssetDataL
 	RunAssetsThroughFilterImpl(AssetDataList, Filter, EARFilterMode::Exclusive);
 }
 
-bool UAssetRegistryImpl::IsAssetIncludedByFilter(const FAssetData& AssetData, const FARFilter& Filter) const
+bool UAssetRegistryImpl::IsAssetIncludedByFilter(const FAssetData& AssetData, const FARCompiledFilter& Filter) const
 {
 	return RunAssetThroughFilterImpl(AssetData, Filter, EARFilterMode::Inclusive);
 }
 
-bool UAssetRegistryImpl::IsAssetExcludedByFilter(const FAssetData& AssetData, const FARFilter& Filter) const
+bool UAssetRegistryImpl::IsAssetExcludedByFilter(const FAssetData& AssetData, const FARCompiledFilter& Filter) const
 {
 	return RunAssetThroughFilterImpl(AssetData, Filter, EARFilterMode::Exclusive);
 }
 
-bool UAssetRegistryImpl::RunAssetThroughFilterImpl(const FAssetData& AssetData, const FARFilter& Filter, const EARFilterMode FilterMode) const
+bool UAssetRegistryImpl::RunAssetThroughFilterImpl(const FAssetData& AssetData, const FARCompiledFilter& Filter, const EARFilterMode FilterMode) const
 {
-	checkf(!Filter.IsRecursive(), TEXT("Filter should have been expanded with ExpandRecursiveFilter before calling this function!"));
-
 	const bool bPassFilterValue = FilterMode == EARFilterMode::Inclusive;
 	if (Filter.IsEmpty())
 	{
@@ -1071,7 +1062,7 @@ bool UAssetRegistryImpl::RunAssetThroughFilterImpl(const FAssetData& AssetData,
 	return bFilterResult == bPassFilterValue;
 }
 
-bool UAssetRegistryImpl::RunAssetThroughFilterImpl_Unchecked(const FAssetData& AssetData, const FARFilter& Filter, const bool bPassFilterValue) const
+bool UAssetRegistryImpl::RunAssetThroughFilterImpl_Unchecked(const FAssetData& AssetData, const FARCompiledFilter& Filter, const bool bPassFilterValue) const
 {
 	// Package Names
 	if (Filter.PackageNames.Num() > 0)
@@ -1138,19 +1129,17 @@ bool UAssetRegistryImpl::RunAssetThroughFilterImpl_Unchecked(const FAssetData& A
 
 void UAssetRegistryImpl::RunAssetsThroughFilterImpl(TArray<FAssetData>& AssetDataList, const FARFilter& Filter, const EARFilterMode FilterMode) const
 {
-	if (Filter.IsEmpty() || !FAssetRegistryState::IsFilterValid(Filter, true))
+	if (Filter.IsEmpty())
 	{
 		return;
 	}
 
-	const FARFilter* FilterPtr = &Filter;
-	FARFilter TempExpandedFilter;
-	if (Filter.IsRecursive())
+	FARCompiledFilter CompiledFilter;
+	CompileFilter(Filter, CompiledFilter);
+	if (!FAssetRegistryState::IsFilterValid(CompiledFilter))
 	{
-		ExpandRecursiveFilter(Filter, TempExpandedFilter);
-		FilterPtr = &TempExpandedFilter;
+		return;
 	}
-	check(FilterPtr);
 
 	const int32 OriginalArrayCount = AssetDataList.Num();
 	const bool bPassFilterValue = FilterMode == EARFilterMode::Inclusive;
@@ -1158,7 +1147,7 @@ void UAssetRegistryImpl::RunAssetsThroughFilterImpl(TArray<FAssetData>& AssetDat
 	// Spin the array backwards to minimize the number of elements that are repeatedly moved down
 	for (int32 AssetDataIndex = AssetDataList.Num() - 1; AssetDataIndex >= 0; --AssetDataIndex)
 	{
-		const bool bFilterResult = RunAssetThroughFilterImpl_Unchecked(AssetDataList[AssetDataIndex], *FilterPtr, bPassFilterValue);
+		const bool bFilterResult = RunAssetThroughFilterImpl_Unchecked(AssetDataList[AssetDataIndex], CompiledFilter, bPassFilterValue);
 		if (bFilterResult != bPassFilterValue)
 		{
 			AssetDataList.RemoveAt(AssetDataIndex, 1, /*bAllowShrinking*/false);
@@ -1174,60 +1163,52 @@ void UAssetRegistryImpl::RunAssetsThroughFilterImpl(TArray<FAssetData>& AssetDat
 
 void UAssetRegistryImpl::ExpandRecursiveFilter(const FARFilter& InFilter, FARFilter& ExpandedFilter) const
 {
-	TSet<FName> FilterPackagePaths;
-	TSet<FName> FilterClassNames;
-	const int32 NumFilterPackagePaths = InFilter.PackagePaths.Num();
-	const int32 NumFilterClasses = InFilter.ClassNames.Num();
+	FARCompiledFilter CompiledFilter;
+	CompileFilter(InFilter, CompiledFilter);
 
-	ExpandedFilter = InFilter;
+	ExpandedFilter.Clear();
+	ExpandedFilter.PackageNames = CompiledFilter.PackageNames.Array();
+	ExpandedFilter.PackagePaths = CompiledFilter.PackagePaths.Array();
+	ExpandedFilter.ObjectPaths = CompiledFilter.ObjectPaths.Array();
+	ExpandedFilter.ClassNames = CompiledFilter.ClassNames.Array();
+	ExpandedFilter.TagsAndValues = CompiledFilter.TagsAndValues;
+	ExpandedFilter.bIncludeOnlyOnDiskAssets = CompiledFilter.bIncludeOnlyOnDiskAssets;
+}
 
-	FilterPackagePaths.Reserve(NumFilterPackagePaths);
-	for (int32 PathIdx = 0; PathIdx < NumFilterPackagePaths; ++PathIdx)
-	{
-		FilterPackagePaths.Add(InFilter.PackagePaths[PathIdx]);
-	}
+void UAssetRegistryImpl::CompileFilter(const FARFilter& InFilter, FARCompiledFilter& OutCompiledFilter) const
+{
+	OutCompiledFilter.Clear();
+	OutCompiledFilter.PackageNames.Append(InFilter.PackageNames);
+	OutCompiledFilter.PackagePaths.Append(InFilter.PackagePaths);
+	OutCompiledFilter.ObjectPaths.Append(InFilter.ObjectPaths);
+	OutCompiledFilter.ClassNames.Append(InFilter.ClassNames);
+	OutCompiledFilter.TagsAndValues = InFilter.TagsAndValues;
+	OutCompiledFilter.bIncludeOnlyOnDiskAssets = InFilter.bIncludeOnlyOnDiskAssets;
 
 	if (InFilter.bRecursivePaths)
 	{
-		// Add subpaths to all the input paths to the list
-		for (int32 PathIdx = 0; PathIdx < NumFilterPackagePaths; ++PathIdx)
+		// Add the sub-paths of all the input paths to the expanded list
+		for (const FName PackagePath : InFilter.PackagePaths)
 		{
-			CachedPathTree.GetSubPaths(InFilter.PackagePaths[PathIdx], FilterPackagePaths);
+			CachedPathTree.GetSubPaths(PackagePath, OutCompiledFilter.PackagePaths);
 		}
 	}
 
-	ExpandedFilter.bRecursivePaths = false;
-	ExpandedFilter.PackagePaths = FilterPackagePaths.Array();
-
 	if (InFilter.bRecursiveClasses)
 	{
+		// Add the sub-classes of all the input classes to the expanded list, excluding any that were requested
 		if (InFilter.RecursiveClassesExclusionSet.Num() > 0 && InFilter.ClassNames.Num() == 0)
 		{
-			// Build list of all classes then remove excluded classes
 			TArray<FName> ClassNamesObject;
 			ClassNamesObject.Add(UObject::StaticClass()->GetFName());
 
-			// GetSubClasses includes the base classes
-			GetSubClasses(ClassNamesObject, InFilter.RecursiveClassesExclusionSet, FilterClassNames);
+			GetSubClasses(ClassNamesObject, InFilter.RecursiveClassesExclusionSet, OutCompiledFilter.ClassNames);
 		}
 		else
 		{
-			// GetSubClasses includes the base classes
-			GetSubClasses(InFilter.ClassNames, InFilter.RecursiveClassesExclusionSet, FilterClassNames);
-		}
-	}
-	else
-	{
-		FilterClassNames.Reserve(NumFilterClasses);
-		for (int32 ClassIdx = 0; ClassIdx < NumFilterClasses; ++ClassIdx)
-		{
-			FilterClassNames.Add(InFilter.ClassNames[ClassIdx]);
+			GetSubClasses(InFilter.ClassNames, InFilter.RecursiveClassesExclusionSet, OutCompiledFilter.ClassNames);
 		}
 	}
-
-	ExpandedFilter.ClassNames = FilterClassNames.Array();
-	ExpandedFilter.bRecursiveClasses = false;
-	ExpandedFilter.RecursiveClassesExclusionSet.Empty();
 }
 
 EAssetAvailability::Type UAssetRegistryImpl::GetAssetAvailability(const FAssetData& AssetData) const
@@ -2741,12 +2722,10 @@ void UAssetRegistryImpl::GetSubClasses(const TArray<FName>& InClassNames, const
 	UpdateTemporaryCaches();
 
 	TSet<FName> ProcessedClassNames;
-
 	for (FName ClassName : InClassNames)
 	{
 		// Now find all subclass names
 		GetSubClasses_Recursive(ClassName, SubClassNames, ProcessedClassNames, TempReverseInheritanceMap, ExcludedClassNames);
-		ProcessedClassNames.Reset();
 	}
 
 	ClearTemporaryCaches();
diff --git a/Engine/Source/Runtime/AssetRegistry/Private/AssetRegistry.h b/Engine/Source/Runtime/AssetRegistry/Private/AssetRegistry.h
--- a/Engine/Source/Runtime/AssetRegistry/Private/AssetRegistry.h
+++ b/Engine/Source/Runtime/AssetRegistry/Private/AssetRegistry.h
@@ -61,9 +61,10 @@ public:
 	virtual void EnumerateSubPaths(const FName InBasePath, TFunctionRef<bool(FName)> Callback, bool bInRecurse) const override;
 	virtual void RunAssetsThroughFilter (TArray<FAssetData>& AssetDataList, const FARFilter& Filter) const override;
 	virtual void UseFilterToExcludeAssets(TArray<FAssetData>& AssetDataList, const FARFilter& Filter) const override;
-	virtual bool IsAssetIncludedByFilter(const FAssetData& AssetData, const FARFilter& Filter) const override;
-	virtual bool IsAssetExcludedByFilter(const FAssetData& AssetData, const FARFilter& Filter) const override;
+	virtual bool IsAssetIncludedByFilter(const FAssetData& AssetData, const FARCompiledFilter& Filter) const override;
+	virtual bool IsAssetExcludedByFilter(const FAssetData& AssetData, const FARCompiledFilter& Filter) const override;
 	virtual void ExpandRecursiveFilter(const FARFilter& InFilter, FARFilter& ExpandedFilter) const override;
+	virtual void CompileFilter(const FARFilter& InFilter, FARCompiledFilter& OutCompiledFilter) const override;
 	virtual void SetTemporaryCachingMode(bool bEnable) override;
 	virtual bool GetTemporaryCachingMode() const override;
 	virtual EAssetAvailability::Type GetAssetAvailability(const FAssetData& AssetData) const override;	
@@ -264,8 +265,8 @@ private:
 	 *  - If an asset data fails a filter, then in inclusive mode it will return false, and in exclusive mode it will return true.
 	 *  - If the filter is empty, then in inclusive mode it will return true, and in exclusive mode it will return false.
 	 */
-	bool RunAssetThroughFilterImpl(const FAssetData& AssetData, const FARFilter& Filter, const EARFilterMode FilterMode) const;
-	bool RunAssetThroughFilterImpl_Unchecked(const FAssetData& AssetData, const FARFilter& Filter, const bool bPassFilterValue) const;
+	bool RunAssetThroughFilterImpl(const FAssetData& AssetData, const FARCompiledFilter& Filter, const EARFilterMode FilterMode) const;
+	bool RunAssetThroughFilterImpl_Unchecked(const FAssetData& AssetData, const FARCompiledFilter& Filter, const bool bPassFilterValue) const;
 
 	/**
 	 * Given an array of asset data, trim the items that fail the filter based on the inclusion/exclusion mode used.
diff --git a/Engine/Source/Runtime/AssetRegistry/Private/AssetRegistryState.cpp b/Engine/Source/Runtime/AssetRegistry/Private/AssetRegistryState.cpp
--- a/Engine/Source/Runtime/AssetRegistry/Private/AssetRegistryState.cpp
+++ b/Engine/Source/Runtime/AssetRegistry/Private/AssetRegistryState.cpp
@@ -558,7 +558,7 @@ bool FAssetRegistryState::HasAssets(const FName PackagePath) const
 	return FoundAssetArray && FoundAssetArray->Num() > 0;
 }
 
-bool FAssetRegistryState::GetAssets(const FARFilter& Filter, const TSet<FName>& PackageNamesToSkip, TArray<FAssetData>& OutAssetData) const
+bool FAssetRegistryState::GetAssets(const FARCompiledFilter& Filter, const TSet<FName>& PackageNamesToSkip, TArray<FAssetData>& OutAssetData) const
 {
 	return EnumerateAssets(Filter, PackageNamesToSkip, [&OutAssetData](const FAssetData& AssetData)
 	{
@@ -567,29 +567,23 @@ bool FAssetRegistryState::GetAssets(const FARFilter& Filter, const TSet<FName>&
 	});
 }
 
-bool FAssetRegistryState::EnumerateAssets(const FARFilter& Filter, const TSet<FName>& PackageNamesToSkip, TFunctionRef<bool(const FAssetData&)> Callback) const
+bool FAssetRegistryState::EnumerateAssets(const FARCompiledFilter& Filter, const TSet<FName>& PackageNamesToSkip, TFunctionRef<bool(const FAssetData&)> Callback) const
 {
 	// Verify filter input. If all assets are needed, use EnumerateAllAssets() instead.
-	if (!IsFilterValid(Filter, false) || Filter.IsEmpty())
+	if (Filter.IsEmpty() || !IsFilterValid(Filter))
 	{
 		return false;
 	}
 
-	// Prepare a set of each filter component for fast searching
-	TSet<FName> FilterPackageNames(Filter.PackageNames);
-	TSet<FName> FilterPackagePaths(Filter.PackagePaths);
-	TSet<FName> FilterClassNames(Filter.ClassNames);
-	TSet<FName> FilterObjectPaths(Filter.ObjectPaths);
-
 	// Form a set of assets matched by each filter
 	TArray<TSet<FAssetData*>> DiskFilterSets;
 
 	// On disk package names
-	if (FilterPackageNames.Num())
+	if (Filter.PackageNames.Num() > 0)
 	{
 		TSet<FAssetData*>& PackageNameFilter = DiskFilterSets[DiskFilterSets.AddDefaulted()];
 
-		for (FName PackageName : FilterPackageNames)
+		for (FName PackageName : Filter.PackageNames)
 		{
 			const TArray<FAssetData*>* PackageAssets = CachedAssetsByPackageName.Find(PackageName);
 
@@ -601,11 +595,11 @@ bool FAssetRegistryState::EnumerateAssets(const FARFilter& Filter, const TSet<FN
 	}
 
 	// On disk package paths
-	if (FilterPackagePaths.Num())
+	if (Filter.PackagePaths.Num() > 0)
 	{
 		TSet<FAssetData*>& PathFilter = DiskFilterSets[DiskFilterSets.AddDefaulted()];
 
-		for (FName PackagePath : FilterPackagePaths)
+		for (FName PackagePath : Filter.PackagePaths)
 		{
 			const TArray<FAssetData*>* PathAssets = CachedAssetsByPath.Find(PackagePath);
 
@@ -617,11 +611,11 @@ bool FAssetRegistryState::EnumerateAssets(const FARFilter& Filter, const TSet<FN
 	}
 
 	// On disk classes
-	if (FilterClassNames.Num())
+	if (Filter.ClassNames.Num() > 0)
 	{
 		TSet<FAssetData*>& ClassFilter = DiskFilterSets[DiskFilterSets.AddDefaulted()];
 
-		for (FName ClassName : FilterClassNames)
+		for (FName ClassName : Filter.ClassNames)
 		{
 			const TArray<FAssetData*>* ClassAssets = CachedAssetsByClass.Find(ClassName);
 
@@ -633,11 +627,11 @@ bool FAssetRegistryState::EnumerateAssets(const FARFilter& Filter, const TSet<FN
 	}
 
 	// On disk object paths
-	if (FilterObjectPaths.Num())
+	if (Filter.ObjectPaths.Num() > 0)
 	{
 		TSet<FAssetData*>& ObjectPathsFilter = DiskFilterSets[DiskFilterSets.AddDefaulted()];
 
-		for (FName ObjectPath : FilterObjectPaths)
+		for (FName ObjectPath : Filter.ObjectPaths)
 		{
 			FAssetData* AssetDataPtr = CachedAssetsByObjectPath.FindRef(ObjectPath);
 
@@ -649,7 +643,7 @@ bool FAssetRegistryState::EnumerateAssets(const FARFilter& Filter, const TSet<FN
 	}
 
 	// On disk tags and values
-	if (Filter.TagsAndValues.Num())
+	if (Filter.TagsAndValues.Num() > 0)
 	{
 		TSet<FAssetData*>& TagAndValuesFilter = DiskFilterSets[DiskFilterSets.AddDefaulted()];
 
@@ -1647,49 +1641,14 @@ bool FAssetRegistryState::RemovePackageData(FName PackageName)
 	return false;
 }
 
-bool FAssetRegistryState::IsFilterValid(const FARFilter& Filter, bool bAllowRecursion)
+bool FAssetRegistryState::IsFilterValid(const FARCompiledFilter& Filter)
 {
-	for (int32 NameIdx = 0; NameIdx < Filter.PackageNames.Num(); ++NameIdx)
-	{
-		if (Filter.PackageNames[NameIdx] == NAME_None)
-		{
-			return false;
-		}
-	}
-
-	for (int32 PathIdx = 0; PathIdx < Filter.PackagePaths.Num(); ++PathIdx)
-	{
-		if (Filter.PackagePaths[PathIdx] == NAME_None)
-		{
-			return false;
-		}
-	}
-
-	for (int32 ObjectPathIdx = 0; ObjectPathIdx < Filter.ObjectPaths.Num(); ++ObjectPathIdx)
-	{
-		if (Filter.ObjectPaths[ObjectPathIdx] == NAME_None)
-		{
-			return false;
-		}
-	}
-
-	for (int32 ClassIdx = 0; ClassIdx < Filter.ClassNames.Num(); ++ClassIdx)
-	{
-		if (Filter.ClassNames[ClassIdx] == NAME_None)
-		{
-			return false;
-		}
-	}
-
-	for (auto FilterTagIt = Filter.TagsAndValues.CreateConstIterator(); FilterTagIt; ++FilterTagIt)
-	{
-		if (FilterTagIt.Key() == NAME_None)
-		{
-			return false;
-		}
-	}
-
-	if (!bAllowRecursion && Filter.IsRecursive())
+	if (Filter.PackageNames.Contains(NAME_None) ||
+		Filter.PackagePaths.Contains(NAME_None) ||
+		Filter.ObjectPaths.Contains(NAME_None) || 
+		Filter.ClassNames.Contains(NAME_None) ||
+		Filter.TagsAndValues.Contains(NAME_None)
+		)
 	{
 		return false;
 	}
diff --git a/Engine/Source/Runtime/AssetRegistry/Public/ARFilter.h b/Engine/Source/Runtime/AssetRegistry/Public/ARFilter.h
--- a/Engine/Source/Runtime/AssetRegistry/Public/ARFilter.h
+++ b/Engine/Source/Runtime/AssetRegistry/Public/ARFilter.h
@@ -9,7 +9,10 @@
 
 struct FAssetData;
 
-/** A struct to serve as a filter for Asset Registry queries. Each component element is processed as an 'OR' operation while all the components are processed together as an 'AND' operation. */
+/**
+ * A struct to serve as a filter for Asset Registry queries.
+ * Each component element is processed as an 'OR' operation while all the components are processed together as an 'AND' operation.
+ */
 USTRUCT(BlueprintType)
 struct FARFilter
 {
@@ -99,3 +102,49 @@ struct FARFilter
 		ensure(IsEmpty());
 	}
 };
+
+/**
+ * A struct to serve as a filter for Asset Registry queries.
+ * Each component element is processed as an 'OR' operation while all the components are processed together as an 'AND' operation.
+ * This is a version of FARFilter optimized for querying, and can be generated from an FARFilter by calling IAssetRegistry::CompileFilter to resolve any recursion.
+ */
+struct FARCompiledFilter
+{
+	/** The filter component for package names */
+	TSet<FName> PackageNames;
+
+	/** The filter component for package paths */
+	TSet<FName> PackagePaths;
+
+	/** The filter component containing specific object paths */
+	TSet<FName> ObjectPaths;
+
+	/** The filter component for class names. Instances of the specified classes, but not subclasses (by default), will be included. Derived classes will be included only if bRecursiveClasses is true. */
+	TSet<FName> ClassNames;
+
+	/** The filter component for properties marked with the AssetRegistrySearchable flag */
+	TMultiMap<FName, TOptional<FString>> TagsAndValues;
+
+	/** If true, only on-disk assets will be returned. Be warned that this is rarely what you want and should only be used for performance reasons */
+	bool bIncludeOnlyOnDiskAssets = false;
+
+	/** Returns true if this filter has no entries */
+	bool IsEmpty() const
+	{
+		return PackageNames.Num() + PackagePaths.Num() + ObjectPaths.Num() + ClassNames.Num() + TagsAndValues.Num() == 0;
+	}
+
+	/** Clears this filter of all entries */
+	void Clear()
+	{
+		PackageNames.Empty();
+		PackagePaths.Empty();
+		ObjectPaths.Empty();
+		ClassNames.Empty();
+		TagsAndValues.Empty();
+
+		bIncludeOnlyOnDiskAssets = false;
+
+		ensure(IsEmpty());
+	}
+};
diff --git a/Engine/Source/Runtime/AssetRegistry/Public/AssetRegistryState.h b/Engine/Source/Runtime/AssetRegistry/Public/AssetRegistryState.h
--- a/Engine/Source/Runtime/AssetRegistry/Public/AssetRegistryState.h
+++ b/Engine/Source/Runtime/AssetRegistry/Public/AssetRegistryState.h
@@ -7,7 +7,7 @@
 #include "Misc/AssetRegistryInterface.h"
 
 class FDependsNode;
-struct FARFilter;
+struct FARCompiledFilter;
 
 #ifndef ASSET_REGISTRY_STATE_DUMPING_ENABLED
 	#define ASSET_REGISTRY_STATE_DUMPING_ENABLED !(UE_BUILD_SHIPPING || UE_BUILD_TEST)
@@ -106,7 +106,7 @@ public:
 	 * @param PackageNamesToSkip explicit list of packages to skip, because they were already added
 	 * @param OutAssetData the list of assets in this path
 	 */
-	bool GetAssets(const FARFilter& Filter, const TSet<FName>& PackageNamesToSkip, TArray<FAssetData>& OutAssetData) const;
+	bool GetAssets(const FARCompiledFilter& Filter, const TSet<FName>& PackageNamesToSkip, TArray<FAssetData>& OutAssetData) const;
 
 	/**
 	 * Enumerate asset data for all assets that match the filter.
@@ -117,7 +117,7 @@ public:
 	 * @param PackageNamesToSkip explicit list of packages to skip, because they were already added
 	 * @param Callback function to call for each asset data enumerated
 	 */
-	bool EnumerateAssets(const FARFilter& Filter, const TSet<FName>& PackageNamesToSkip, TFunctionRef<bool(const FAssetData&)> Callback) const;
+	bool EnumerateAssets(const FARCompiledFilter& Filter, const TSet<FName>& PackageNamesToSkip, TFunctionRef<bool(const FAssetData&)> Callback) const;
 
 	/**
 	 * Gets asset data for all assets in the registry state.
@@ -313,7 +313,7 @@ public:
 	uint32 GetAllocatedSize(bool bLogDetailed = false) const;
 
 	/** Checks a filter to make sure there are no illegal entries */
-	static bool IsFilterValid(const FARFilter& Filter, bool bAllowRecursion);
+	static bool IsFilterValid(const FARCompiledFilter& Filter);
 
 	/** Returns the number of assets in this state */
 	int32 GetNumAssets() const { return NumAssets; }
diff --git a/Engine/Source/Runtime/AssetRegistry/Public/IAssetRegistry.h b/Engine/Source/Runtime/AssetRegistry/Public/IAssetRegistry.h
--- a/Engine/Source/Runtime/AssetRegistry/Public/IAssetRegistry.h
+++ b/Engine/Source/Runtime/AssetRegistry/Public/IAssetRegistry.h
@@ -11,6 +11,7 @@
 
 
 struct FARFilter;
+struct FARCompiledFilter;
 struct FAssetRegistrySerializationOptions;
 class FAssetRegistryState;
 
@@ -323,15 +324,19 @@ public:
 	UFUNCTION(BlueprintCallable, BlueprintPure=false, Category = "AssetRegistry")
 	virtual void UseFilterToExcludeAssets(UPARAM(ref) TArray<FAssetData>& AssetDataList, const FARFilter& Filter) const = 0;
 
-	/** Tests to see whether the given asset would be included (passes) the given filter (note: the filter should have been expanded first, @see ExpandRecursiveFilter */
-	virtual bool IsAssetIncludedByFilter(const FAssetData& AssetData, const FARFilter& Filter) const = 0;
+	/** Tests to see whether the given asset would be included (passes) the given filter */
+	virtual bool IsAssetIncludedByFilter(const FAssetData& AssetData, const FARCompiledFilter& Filter) const = 0;
 
-	/** Tests to see whether the given asset would be excluded (fails) the given filter (note: the filter should have been expanded first, @see ExpandRecursiveFilter */
-	virtual bool IsAssetExcludedByFilter(const FAssetData& AssetData, const FARFilter& Filter) const = 0;
+	/** Tests to see whether the given asset would be excluded (fails) the given filter */
+	virtual bool IsAssetExcludedByFilter(const FAssetData& AssetData, const FARCompiledFilter& Filter) const = 0;
 
 	/** Modifies passed in filter to make it safe for use on FAssetRegistryState. This expands recursive paths and classes */
+	UE_DEPRECATED(4.26, "ExpandRecursiveFilter is deprecated in favor of CompileFilter")
 	virtual void ExpandRecursiveFilter(const FARFilter& InFilter, FARFilter& ExpandedFilter) const = 0;
 
+	/** Modifies passed in filter optimize it for query and expand any recursive paths and classes */
+	virtual void CompileFilter(const FARFilter& InFilter, FARCompiledFilter& OutCompiledFilter) const = 0;
+
 	/** Enables or disable temporary search caching, when this is enabled scanning/searching is faster because we assume no objects are loaded between scans. Disabling frees any caches created */
 	virtual void SetTemporaryCachingMode(bool bEnable) = 0;
 
