FAssetData optimizations

 - The constructor now creates ObjectPath with zero string allocations.
 - IsUAsset now performs its check with zero string allocations.

diff --git a/Engine/Source/Runtime/AssetRegistry/Public/AssetData.h b/Engine/Source/Runtime/AssetRegistry/Public/AssetData.h
--- a/Engine/Source/Runtime/AssetRegistry/Public/AssetData.h
+++ b/Engine/Source/Runtime/AssetRegistry/Public/AssetData.h
@@ -87,20 +87,24 @@ public:
 
 public:
 	/** Default constructor */
-	FAssetData()
-	{}
+	FAssetData() = default;
 
 	/** Constructor */
 	FAssetData(FName InPackageName, FName InPackagePath, FName InAssetName, FName InAssetClass, FAssetDataTagMap InTags = FAssetDataTagMap(), TArray<int32> InChunkIDs = TArray<int32>(), uint32 InPackageFlags = 0)
-		: ObjectPath(*FString::Format(TEXT("{0}.{1}"), { InPackageName.ToString(), InAssetName.ToString() }))
-		, PackageName(InPackageName)
+		: PackageName(InPackageName)
 		, PackagePath(InPackagePath)
 		, AssetName(InAssetName)
 		, AssetClass(InAssetClass)
 		, TagsAndValues(MoveTemp(InTags))
 		, ChunkIDs(MoveTemp(InChunkIDs))
 		, PackageFlags(InPackageFlags)
-	{}
+	{
+		TCHAR ObjectPathStr[FName::StringBufferSize];
+		int32 ObjectPathLen = PackageName.ToString(ObjectPathStr);
+		ObjectPathStr[ObjectPathLen++] = TEXT('.');
+		ObjectPathLen += AssetName.ToString(ObjectPathStr + ObjectPathLen, FName::StringBufferSize - ObjectPathLen);
+		ObjectPath = FName(ObjectPathStr);
+	}
 
 	/** Constructor taking a UObject. By default trying to create one for a blueprint class will create one for the UBlueprint instead, but this can be overridden */
 	FAssetData(const UObject* InAsset, bool bAllowBlueprintClass = false)
@@ -171,7 +175,21 @@ public:
 	/** Returns true if this is the primary asset in a package, true for maps and assets but false for secondary objects like class redirectors */
 	bool IsUAsset() const
 	{
-		return FPackageName::GetLongPackageAssetName(PackageName.ToString()) == AssetName.ToString();
+		if (!IsValid())
+		{
+			return false;
+		}
+
+		TCHAR AssetNameStr[FName::StringBufferSize];
+		const int32 AssetNameLen = AssetName.ToString(AssetNameStr);
+
+		TCHAR PackageNameStr[FName::StringBufferSize];
+		const int32 PackageNameLen = PackageName.ToString(PackageNameStr);
+
+		const TCHAR* PackageAssetNameStr = PackageNameStr + PackageNameLen;
+		for (; PackageAssetNameStr > PackageNameStr && *(PackageAssetNameStr - 1) != TEXT('/'); --PackageAssetNameStr) {}
+
+		return FCString::Stricmp(PackageAssetNameStr, AssetNameStr) == 0;
 	}
 
 	void Shrink()
@@ -218,12 +236,8 @@ public:
 	/** Returns true if the this asset is a redirector. */
 	bool IsRedirector() const
 	{
-		if ( AssetClass == UObjectRedirector::StaticClass()->GetFName() )
-		{
-			return true;
-		}
-
-		return false;
+		static const FName ObjectRedirectorClassName = UObjectRedirector::StaticClass()->GetFName();
+		return AssetClass == ObjectRedirectorClassName;
 	}
 
 	/** Returns the class UClass if it is loaded. It is not possible to load the class if it is unloaded since we only have the short name. */
