Add a config value to determine if the asset registry scanner should try and gather data on cooked packages that we didn't find an asset export for. Defaults to TRUE to match existing behavior.

diff --git a/Engine/Source/Runtime/AssetRegistry/Private/AssetRegistry.cpp b/Engine/Source/Runtime/AssetRegistry/Private/AssetRegistry.cpp
--- a/Engine/Source/Runtime/AssetRegistry/Private/AssetRegistry.cpp
+++ b/Engine/Source/Runtime/AssetRegistry/Private/AssetRegistry.cpp
@@ -2078,18 +2078,35 @@ void UAssetRegistryImpl::CookedPackageNamesWithoutAssetDataGathered(const double
 {
 	const bool bFlushFullBuffer = TickStartTime < 0;
 
-	// Add the found assets
-	while (CookedPackageNamesWithoutAssetDataResults.Num() > 0)
+	struct FConfigValue
 	{
-		// If this data is cooked and it we couldn't find any asset in its export table then try to load the entire package 
-		// Loading the entire package will make all of its assets searchable through the in-memory scanning performed by GetAsseets
-		const FString& BackgroundResult = CookedPackageNamesWithoutAssetDataResults.Pop();
-		LoadPackage(nullptr, *BackgroundResult, 0);
+		FConfigValue()
+		{
+			if (GConfig)
+			{
+				GConfig->GetBool(TEXT("AssetRegistry"), TEXT("LoadCookedPackagesWithoutAssetData"), bShouldProcess, GEngineIni);
+			}
+		}
 
-		// Check to see if we have run out of time in this tick
-		if (!bFlushFullBuffer && (FPlatformTime::Seconds() - TickStartTime) > MaxSecondsPerFrame)
+		bool bShouldProcess = true;
+	};
+	static FConfigValue ShouldProcessCookedPackages;
+
+	// Add the found assets
+	if (ShouldProcessCookedPackages.bShouldProcess)
+	{
+		while (CookedPackageNamesWithoutAssetDataResults.Num() > 0)
 		{
-			return;
+			// If this data is cooked and it we couldn't find any asset in its export table then try to load the entire package 
+			// Loading the entire package will make all of its assets searchable through the in-memory scanning performed by GetAsseets
+			const FString& BackgroundResult = CookedPackageNamesWithoutAssetDataResults.Pop();
+			LoadPackage(nullptr, *BackgroundResult, 0);
+
+			// Check to see if we have run out of time in this tick
+			if (!bFlushFullBuffer && (FPlatformTime::Seconds() - TickStartTime) > MaxSecondsPerFrame)
+			{
+				return;
+			}
 		}
 	}
 
