AssetRegistry - When the project config says we should ignore cooked files with no asset data marks in the export table, actually clear the pending list of files to process.

diff --git a/Engine/Source/Runtime/AssetRegistry/Private/AssetRegistry.cpp b/Engine/Source/Runtime/AssetRegistry/Private/AssetRegistry.cpp
--- a/Engine/Source/Runtime/AssetRegistry/Private/AssetRegistry.cpp
+++ b/Engine/Source/Runtime/AssetRegistry/Private/AssetRegistry.cpp
@@ -2109,6 +2109,13 @@ void UAssetRegistryImpl::CookedPackageNamesWithoutAssetDataGathered(const double
 			}
 		}
 	}
+	else
+	{
+		// Do nothing will these packages. For projects which could run entirely from cooked data, this
+		// process will involve opening every single package synchronously on the game thread which will
+		// kill performance. We need a better way.
+		CookedPackageNamesWithoutAssetDataResults.Empty();
+	}
 
 	// Trim the results array
 	CookedPackageNamesWithoutAssetDataResults.Trim();
