FAssetData optimizations

 - The constructor now creates ObjectPath with zero string allocations.
 - IsUAsset now performs its check with zero string allocations.

diff --git a/Engine/Source/Runtime/AssetRegistry/Public/AssetData.h b/Engine/Source/Runtime/AssetRegistry/Public/AssetData.h
--- a/Engine/Source/Runtime/AssetRegistry/Public/AssetData.h
+++ b/Engine/Source/Runtime/AssetRegistry/Public/AssetData.h
@@ -10,6 +10,8 @@
 #include "UObject/Package.h"
 #include "UObject/ObjectRedirector.h"
 #include "Misc/PackageName.h"
+#include "Misc/StringBuilder.h"
+#include "Containers/StringView.h"
 #include "UObject/LinkerLoad.h"
 #include "AssetDataTagMap.h"
 #include "UObject/PrimaryAssetId.h"
@@ -83,12 +85,11 @@ public:
 	/** The IDs of the pakchunks this asset is located in for streaming install.  Empty if not assigned to a chunk */
 	TArray<int32> ChunkIDs;
 	/** Asset package flags */
-	uint32 PackageFlags;
+	uint32 PackageFlags = 0;
 
 public:
 	/** Default constructor */
-	FAssetData()
-	{}
+	FAssetData() = default;
 
 	/** Constructor */
 	FAssetData(FName InPackageName, FName InPackagePath, FName InAssetName, FName InAssetClass, FAssetDataTagMap InTags = FAssetDataTagMap(), TArray<int32> InChunkIDs = TArray<int32>(), uint32 InPackageFlags = 0)
@@ -100,7 +101,13 @@ public:
 		, TagsAndValues(MoveTemp(InTags))
 		, ChunkIDs(MoveTemp(InChunkIDs))
 		, PackageFlags(InPackageFlags)
-	{}
+	{
+		TStringBuilder<FName::StringBufferSize> ObjectPathStr;
+		PackageName.AppendString(ObjectPathStr);
+		ObjectPathStr << TEXT('.');
+		AssetName.AppendString(ObjectPathStr);
+		ObjectPath = FName(FStringView(ObjectPathStr));
+	}
 
 	/** Constructor taking a UObject. By default trying to create one for a blueprint class will create one for the UBlueprint instead, but this can be overridden */
 	FAssetData(const UObject* InAsset, bool bAllowBlueprintClass = false)
@@ -171,7 +178,28 @@ public:
 	/** Returns true if this is the primary asset in a package, true for maps and assets but false for secondary objects like class redirectors */
 	bool IsUAsset() const
 	{
-		return FPackageName::GetLongPackageAssetName(PackageName.ToString()) == AssetName.ToString();
+		if (!IsValid())
+		{
+			return false;
+		}
+
+		TStringBuilder<FName::StringBufferSize> AssetNameStrBuilder;
+		AssetName.ToString(AssetNameStrBuilder);
+
+		TStringBuilder<FName::StringBufferSize> PackageNameStrBuilder;
+		PackageName.ToString(PackageNameStrBuilder);
+
+		FStringView PackageAssetNameStr;
+		{
+			// Get everything after the last slash
+			FStringView PackageNameStr(PackageNameStrBuilder);
+			int32 IndexOfLastSlash = INDEX_NONE;
+			PackageNameStr.FindLastChar(TEXT('/'), IndexOfLastSlash);
+			PackageAssetNameStr = PackageNameStr.Mid(IndexOfLastSlash + 1);
+		}
+
+		FStringView AssetNameStr(AssetNameStrBuilder);
+		return PackageAssetNameStr == AssetNameStr;
 	}
 
 	void Shrink()
@@ -218,12 +246,8 @@ public:
 	/** Returns true if the this asset is a redirector. */
 	bool IsRedirector() const
 	{
-		if ( AssetClass == UObjectRedirector::StaticClass()->GetFName() )
-		{
-			return true;
-		}
-
-		return false;
+		static const FName ObjectRedirectorClassName = UObjectRedirector::StaticClass()->GetFName();
+		return AssetClass == ObjectRedirectorClassName;
 	}
 
 	/** Returns the class UClass if it is loaded. It is not possible to load the class if it is unloaded since we only have the short name. */
