Added EnumerateAssets overload that takes a compiled AR filter

diff --git a/Engine/Source/Runtime/AssetRegistry/Private/AssetRegistry.cpp b/Engine/Source/Runtime/AssetRegistry/Private/AssetRegistry.cpp
--- a/Engine/Source/Runtime/AssetRegistry/Private/AssetRegistry.cpp
+++ b/Engine/Source/Runtime/AssetRegistry/Private/AssetRegistry.cpp
@@ -624,6 +624,13 @@ bool UAssetRegistryImpl::GetAssets(const FARFilter& InFilter, TArray<FAssetData>
 }
 
 bool UAssetRegistryImpl::EnumerateAssets(const FARFilter& InFilter, TFunctionRef<bool(const FAssetData&)> Callback) const
+{
+	FARCompiledFilter CompiledFilter;
+	CompileFilter(InFilter, CompiledFilter);
+	return EnumerateAssets(CompiledFilter, Callback);
+}
+
+bool UAssetRegistryImpl::EnumerateAssets(const FARCompiledFilter& InFilter, TFunctionRef<bool(const FAssetData&)> Callback) const
 {
 	// Verify filter input. If all assets are needed, use EnumerateAllAssets() instead.
 	if (InFilter.IsEmpty())
@@ -631,17 +638,14 @@ bool UAssetRegistryImpl::EnumerateAssets(const FARFilter& InFilter, TFunctionRef
 		return false;
 	}
 
-	// Expand recursion on filter
-	FARCompiledFilter CompiledFilter;
-	CompileFilter(InFilter, CompiledFilter);
-	if (!FAssetRegistryState::IsFilterValid(CompiledFilter))
+	if (!FAssetRegistryState::IsFilterValid(InFilter))
 	{
 		return false;
 	}
 
 	// Start with in memory assets
 	TSet<FName> PackagesToSkip = CachedEmptyPackages;
-	if (!CompiledFilter.bIncludeOnlyOnDiskAssets)
+	if (!InFilter.bIncludeOnlyOnDiskAssets)
 	{
 		// Reusable structures to avoid memory allocations
 		TArray<UObject::FAssetRegistryTag> ObjectTags;
@@ -663,16 +667,16 @@ bool UAssetRegistryImpl::EnumerateAssets(const FARFilter& InFilter, TFunctionRef
 
 				PackagesToSkip.Add(PackageName);
 
-				if (CompiledFilter.PackageNames.Num() && !CompiledFilter.PackageNames.Contains(PackageName))
+				if (InFilter.PackageNames.Num() && !InFilter.PackageNames.Contains(PackageName))
 				{
 					return;
 				}
 
 				// Object Path
-				if (CompiledFilter.ObjectPaths.Num() > 0)
+				if (InFilter.ObjectPaths.Num() > 0)
 				{
 					const FName ObjectPath = FName(*Obj->GetPathName(), FNAME_Find);
-					if (!CompiledFilter.ObjectPaths.Contains(ObjectPath))
+					if (!InFilter.ObjectPaths.Contains(ObjectPath))
 					{
 						return;
 					}
@@ -680,7 +684,7 @@ bool UAssetRegistryImpl::EnumerateAssets(const FARFilter& InFilter, TFunctionRef
 
 				// Package path
 				const FName PackagePath = FName(*FPackageName::GetLongPackagePath(InMemoryPackage->GetName()));
-				if (CompiledFilter.PackagePaths.Num() > 0 && !CompiledFilter.PackagePaths.Contains(PackagePath))
+				if (InFilter.PackagePaths.Num() > 0 && !InFilter.PackagePaths.Contains(PackagePath))
 				{
 					return;
 				}
@@ -688,10 +692,10 @@ bool UAssetRegistryImpl::EnumerateAssets(const FARFilter& InFilter, TFunctionRef
 				// Tags and values
 				check(ObjectTags.Num() == 0);
 				Obj->GetAssetRegistryTags(ObjectTags);
-				if (CompiledFilter.TagsAndValues.Num() > 0)
+				if (InFilter.TagsAndValues.Num() > 0)
 				{
 					bool bMatch = false;
-					for (auto FilterTagIt = CompiledFilter.TagsAndValues.CreateConstIterator(); FilterTagIt; ++FilterTagIt)
+					for (auto FilterTagIt = InFilter.TagsAndValues.CreateConstIterator(); FilterTagIt; ++FilterTagIt)
 					{
 						const FName Tag = FilterTagIt.Key();
 						const TOptional<FString>& Value = FilterTagIt.Value();
@@ -739,10 +743,10 @@ bool UAssetRegistryImpl::EnumerateAssets(const FARFilter& InFilter, TFunctionRef
 		};
 
 		// Iterate over all in-memory assets to find the ones that pass the filter components
-		if (CompiledFilter.ClassNames.Num() > 0)
+		if (InFilter.ClassNames.Num() > 0)
 		{
 			TArray<UObject*> InMemoryObjects;
-			for (FName ClassName : CompiledFilter.ClassNames)
+			for (FName ClassName : InFilter.ClassNames)
 			{
 				UClass* Class = FindObjectFast<UClass>(nullptr, ClassName, false, true, RF_NoFlags);
 				if (Class != nullptr)
@@ -775,7 +779,7 @@ bool UAssetRegistryImpl::EnumerateAssets(const FARFilter& InFilter, TFunctionRef
 		}
 	}
 
-	State.EnumerateAssets(CompiledFilter, PackagesToSkip, Callback);
+	State.EnumerateAssets(InFilter, PackagesToSkip, Callback);
 
 	return true;
 }
diff --git a/Engine/Source/Runtime/AssetRegistry/Private/AssetRegistry.h b/Engine/Source/Runtime/AssetRegistry/Private/AssetRegistry.h
--- a/Engine/Source/Runtime/AssetRegistry/Private/AssetRegistry.h
+++ b/Engine/Source/Runtime/AssetRegistry/Private/AssetRegistry.h
@@ -41,6 +41,7 @@ public:
 	virtual bool GetAssetsByTagValues(const TMultiMap<FName, FString>& AssetTagsAndValues, TArray<FAssetData>& OutAssetData) const override;
 	virtual bool GetAssets(const FARFilter& Filter, TArray<FAssetData>& OutAssetData) const override;
 	virtual bool EnumerateAssets(const FARFilter& Filter, TFunctionRef<bool(const FAssetData&)> Callback) const override;
+	virtual bool EnumerateAssets(const FARCompiledFilter& Filter, TFunctionRef<bool(const FAssetData&)> Callback) const override;
 	virtual FAssetData GetAssetByObjectPath( const FName ObjectPath, bool bIncludeOnlyOnDiskAssets = false ) const override;
 	virtual bool GetAllAssets(TArray<FAssetData>& OutAssetData, bool bIncludeOnlyOnDiskAssets = false) const override;
 	virtual bool EnumerateAllAssets(TFunctionRef<bool(const FAssetData&)> Callback, bool bIncludeOnlyOnDiskAssets = false) const override;
diff --git a/Engine/Source/Runtime/AssetRegistry/Public/IAssetRegistry.h b/Engine/Source/Runtime/AssetRegistry/Public/IAssetRegistry.h
--- a/Engine/Source/Runtime/AssetRegistry/Public/IAssetRegistry.h
+++ b/Engine/Source/Runtime/AssetRegistry/Public/IAssetRegistry.h
@@ -185,6 +185,7 @@ public:
 	 * @param Callback function to call for each asset data enumerated
 	 */
 	virtual bool EnumerateAssets(const FARFilter& Filter, TFunctionRef<bool(const FAssetData&)> Callback) const = 0;
+	virtual bool EnumerateAssets(const FARCompiledFilter& Filter, TFunctionRef<bool(const FAssetData&)> Callback) const = 0;
 
 	/**
 	 * Gets the asset data for the specified object path
