Disable use of FPreloadFile for the asset registry. It's causing gamedata not to be loaded when we need it.
Define ASSET_REGISTRY_USE_PRELOAD_FILE to 1 to re-enable

diff --git a/Engine/Source/Runtime/AssetRegistry/Private/AssetRegistry.cpp b/Engine/Source/Runtime/AssetRegistry/Private/AssetRegistry.cpp
--- a/Engine/Source/Runtime/AssetRegistry/Private/AssetRegistry.cpp
+++ b/Engine/Source/Runtime/AssetRegistry/Private/AssetRegistry.cpp
@@ -32,6 +32,8 @@
 
 DEFINE_LOG_CATEGORY(LogAssetRegistry);
 
+#define ASSET_REGISTRY_USE_PRELOAD_FILE 0
+
 /** Returns the appropriate ChunkProgressReportingType for the given Asset enum */
 EChunkProgressReportingType::Type GetChunkAvailabilityProgressType(EAssetAvailabilityProgressReportingType::Type ReportType)
 {
@@ -95,6 +97,7 @@ UAssetRegistryImpl::UAssetRegistryImpl(const FObjectInitializer& ObjectInitializ
 	else if (FPlatformProperties::RequiresCookedData())
 	{
 		// load the cooked data
+#if ASSET_REGISTRY_USE_PRELOAD_FILE
 		int64 Size;
 		void* PreloadedData = GPreLoadAssetRegistry.TakeOwnershipOfLoadedData(&Size);
 
@@ -126,6 +129,30 @@ UAssetRegistryImpl::UAssetRegistryImpl(const FObjectInitializer& ObjectInitializ
 				}
 			}
 		}
+#else // ASSET_REGISTRY_USE_PRELOAD_FILE
+		// load the cooked data
+		FArrayReader SerializedAssetData;
+
+		FString AssetRegistryFilename = (FPaths::ProjectDir() / TEXT("AssetRegistry.bin"));
+		if (SerializationOptions.bSerializeAssetRegistry && IFileManager::Get().FileExists(*AssetRegistryFilename) && FFileHelper::LoadFileToArray(SerializedAssetData, *AssetRegistryFilename))
+		{
+			// serialize the data with the memory reader (will convert FStrings to FNames, etc)
+			Serialize(SerializedAssetData);
+		}
+		TArray<TSharedRef<IPlugin>> ContentPlugins = IPluginManager::Get().GetEnabledPluginsWithContent();
+		for (TSharedRef<IPlugin> ContentPlugin : ContentPlugins)
+		{
+			if (ContentPlugin->CanContainContent())
+			{
+				FString PluginAssetRegistry = ContentPlugin->GetBaseDir() / TEXT("AssetRegistry.bin");
+				if (IFileManager::Get().FileExists(*PluginAssetRegistry) && FFileHelper::LoadFileToArray(SerializedAssetData, *PluginAssetRegistry))
+				{
+					SerializedAssetData.Seek(0);
+					Serialize(SerializedAssetData);
+				}
+			}
+		}
+#endif //ASSET_REGISTRY_USE_PRELOAD_FILE
 	}
 
 	// Report startup time. This does not include DirectoryWatcher startup time.
@@ -2932,4 +2959,4 @@ bool UAssetRegistryImpl::SetPrimaryAssetIdForObjectPath(const FName ObjectPath,
 const FAssetData* UAssetRegistryImpl::GetCachedAssetDataForObjectPath(const FName ObjectPath) const
 {
 	return State.GetAssetByObjectPath(ObjectPath);
-}
\ No newline at end of file
+}
