- Added fix for loading PreLoad files that are mounted after the preload KickOff

diff --git a/Engine/Source/Runtime/AssetRegistry/Private/AssetRegistry.cpp b/Engine/Source/Runtime/AssetRegistry/Private/AssetRegistry.cpp
--- a/Engine/Source/Runtime/AssetRegistry/Private/AssetRegistry.cpp
+++ b/Engine/Source/Runtime/AssetRegistry/Private/AssetRegistry.cpp
@@ -32,8 +32,6 @@
 
 DEFINE_LOG_CATEGORY(LogAssetRegistry);
 
-#define ASSET_REGISTRY_USE_PRELOAD_FILE 0
-
 /** Returns the appropriate ChunkProgressReportingType for the given Asset enum */
 EChunkProgressReportingType::Type GetChunkAvailabilityProgressType(EAssetAvailabilityProgressReportingType::Type ReportType)
 {
@@ -97,7 +95,6 @@ UAssetRegistryImpl::UAssetRegistryImpl(const FObjectInitializer& ObjectInitializ
 	else if (FPlatformProperties::RequiresCookedData())
 	{
 		// load the cooked data
-#if ASSET_REGISTRY_USE_PRELOAD_FILE
 		int64 Size;
 		void* PreloadedData = GPreLoadAssetRegistry.TakeOwnershipOfLoadedData(&Size);
 
@@ -129,30 +126,6 @@ UAssetRegistryImpl::UAssetRegistryImpl(const FObjectInitializer& ObjectInitializ
 				}
 			}
 		}
-#else // ASSET_REGISTRY_USE_PRELOAD_FILE
-		// load the cooked data
-		FArrayReader SerializedAssetData;
-
-		FString AssetRegistryFilename = (FPaths::ProjectDir() / TEXT("AssetRegistry.bin"));
-		if (SerializationOptions.bSerializeAssetRegistry && IFileManager::Get().FileExists(*AssetRegistryFilename) && FFileHelper::LoadFileToArray(SerializedAssetData, *AssetRegistryFilename))
-		{
-			// serialize the data with the memory reader (will convert FStrings to FNames, etc)
-			Serialize(SerializedAssetData);
-		}
-		TArray<TSharedRef<IPlugin>> ContentPlugins = IPluginManager::Get().GetEnabledPluginsWithContent();
-		for (TSharedRef<IPlugin> ContentPlugin : ContentPlugins)
-		{
-			if (ContentPlugin->CanContainContent())
-			{
-				FString PluginAssetRegistry = ContentPlugin->GetBaseDir() / TEXT("AssetRegistry.bin");
-				if (IFileManager::Get().FileExists(*PluginAssetRegistry) && FFileHelper::LoadFileToArray(SerializedAssetData, *PluginAssetRegistry))
-				{
-					SerializedAssetData.Seek(0);
-					Serialize(SerializedAssetData);
-				}
-			}
-		}
-#endif //ASSET_REGISTRY_USE_PRELOAD_FILE
 	}
 
 	// Report startup time. This does not include DirectoryWatcher startup time.
@@ -2959,4 +2932,4 @@ bool UAssetRegistryImpl::SetPrimaryAssetIdForObjectPath(const FName ObjectPath,
 const FAssetData* UAssetRegistryImpl::GetCachedAssetDataForObjectPath(const FName ObjectPath) const
 {
 	return State.GetAssetByObjectPath(ObjectPath);
-}
+}
\ No newline at end of file
diff --git a/Engine/Source/Runtime/Core/Private/Misc/PreLoadFile.cpp b/Engine/Source/Runtime/Core/Private/Misc/PreLoadFile.cpp
--- a/Engine/Source/Runtime/Core/Private/Misc/PreLoadFile.cpp
+++ b/Engine/Source/Runtime/Core/Private/Misc/PreLoadFile.cpp
@@ -30,6 +30,7 @@ struct FPreloadFileRegistry
 FPreLoadFile::FPreLoadFile(const TCHAR* InPath)
 	: FDelayedAutoRegisterHelper(STATS ? EDelayedRegisterRunPhase::StatSystemReady : EDelayedRegisterRunPhase::FileSystemReady, [this] { bSystemNoLongerTakingRequests = false; KickOffRead(); })
 	, bIsComplete(false)
+	, bFailedToOpenInKickOff(false)
 	, Data(nullptr)
 	, FileSize(0)
 	, Path(InPath)
@@ -58,7 +59,7 @@ void FPreLoadFile::KickOffRead()
 	FAsyncFileCallBack SizeCallbackFunction = [this](bool bWasCancelled, IAsyncReadRequest* SizeRequest)
 	{
 		FileSize = (int32)SizeRequest->GetSizeResults();
-//		printf("Preloading %s, size = %d\n", TCHAR_TO_ANSI(*Path), FileSize);
+
 		if (FileSize > 0)
 		{
 			FAsyncFileCallBack ReadCallbackFunction = [this](bool bWasCancelledInner, IAsyncReadRequest* ReadRequest)
@@ -90,6 +91,12 @@ void FPreLoadFile::KickOffRead()
 		Reader->Serialize(Data, FileSize);
 		delete Reader;
 	}
+	else
+	{
+		// it's possible pak files with the file weren't mounted yet, so note that we didn't find it, and try again in TakeOwnership time
+		bFailedToOpenInKickOff = true;
+
+	}
 	CompletionEvent->Trigger();
 #endif
 }	
@@ -97,17 +104,15 @@ void FPreLoadFile::KickOffRead()
 
 void* FPreLoadFile::TakeOwnershipOfLoadedData(int64* OutFileSize)
 {
-	if (!CompletionEvent)
+	// may need to attempt to read again, if (re-)requesting data after the initial boot sequence
+	if (!CompletionEvent || bFailedToOpenInKickOff)
 	{
-		// may need to read again, if (re-)requesting data after the initial boot sequence
 		KickOffRead();
 	}
 
 	check(CompletionEvent);
 	if (CompletionEvent->Wait(0) == false)
 	{
-		printf("PreLoadFile %s wasn't ready...\n", TCHAR_TO_ANSI(*Path));
-
 		// wait until we are done
 		CompletionEvent->Wait();
 	}
@@ -120,6 +125,8 @@ void* FPreLoadFile::TakeOwnershipOfLoadedData(int64* OutFileSize)
 		*OutFileSize = FileSize;
 	}
 
+	checkf(Data != nullptr, TEXT("Failed to find/open/read PreLoaded file %s"), *Path);
+
 	void* ReturnData = Data;
 	Data = nullptr;
 
diff --git a/Engine/Source/Runtime/Core/Public/Misc/PreLoadFile.h b/Engine/Source/Runtime/Core/Public/Misc/PreLoadFile.h
--- a/Engine/Source/Runtime/Core/Public/Misc/PreLoadFile.h
+++ b/Engine/Source/Runtime/Core/Public/Misc/PreLoadFile.h
@@ -16,6 +16,7 @@ protected:
 	void KickOffRead();
 
 	bool bIsComplete;
+	bool bFailedToOpenInKickOff;
 	static bool bSystemNoLongerTakingRequests;
 	void* Data;
 	int64 FileSize;
