Missing change from CL# 12168130

#rnx

diff --git a/Engine/Source/Runtime/AssetRegistry/Public/AssetData.h b/Engine/Source/Runtime/AssetRegistry/Public/AssetData.h
--- a/Engine/Source/Runtime/AssetRegistry/Public/AssetData.h
+++ b/Engine/Source/Runtime/AssetRegistry/Public/AssetData.h
@@ -93,8 +93,7 @@ public:
 
 	/** Constructor */
 	FAssetData(FName InPackageName, FName InPackagePath, FName InAssetName, FName InAssetClass, FAssetDataTagMap InTags = FAssetDataTagMap(), TArray<int32> InChunkIDs = TArray<int32>(), uint32 InPackageFlags = 0)
-		: ObjectPath(*FString::Format(TEXT("{0}.{1}"), { InPackageName.ToString(), InAssetName.ToString() }))
-		, PackageName(InPackageName)
+		: PackageName(InPackageName)
 		, PackagePath(InPackagePath)
 		, AssetName(InAssetName)
 		, AssetClass(InAssetClass)
