Optimize Asset Registry Scanning
  - Parallelize package reading in AssetGatherer
  - Implement scheduling ParallelFor as background prio
  - Saving the cache is now based on the amount of time it would take to redo the work
    so it scales better on many cores systems

  - Tested on AMD TR 3970X 128GB RAM with Crucial MP600 2TB
    - 33m43s to 3m45s (Cold Cache + No AV Exclusions)
    - 6m25s to 38.9s (Cold Cache + AV Exclusion on .uasset)
    - 1m25s to 20.8s (Warm Cache + No AV Exclusions)

    - All AssetCache .bin files deleted between runs
    - Cold cache = Cleared standby pages, mostly equivalent to a reboot
    - Warm cache = 2nd run of 2 consecutive runs

#rnx

diff --git a/Engine/Source/Runtime/AssetRegistry/Private/AssetDataGatherer.cpp b/Engine/Source/Runtime/AssetRegistry/Private/AssetDataGatherer.cpp
--- a/Engine/Source/Runtime/AssetRegistry/Private/AssetDataGatherer.cpp
+++ b/Engine/Source/Runtime/AssetRegistry/Private/AssetDataGatherer.cpp
@@ -11,13 +11,14 @@
 #include "NameTableArchive.h"
 #include "PackageReader.h"
 #include "AssetRegistry.h"
+#include "Async/ParallelFor.h"
 
 namespace AssetDataGathererConstants
 {
 	static const int32 CacheSerializationVersion = 13;
 	static const int32 MaxFilesToDiscoverBeforeFlush = 2500;
 	static const int32 MaxFilesToGatherBeforeFlush = 250;
-	static const int32 MaxFilesToProcessBeforeCacheWrite = 50000;
+	static const int32 MinSecondsToElapseBeforeCacheWrite = 60;
 }
 
 namespace
@@ -469,6 +470,7 @@ bool FAssetDataGatherer::Init()
 
 uint32 FAssetDataGatherer::Run()
 {
+	TRACE_CPUPROFILER_EVENT_SCOPE(FAssetDataGatherer::Run)
 	int32 CacheSerializationVersion = AssetDataGathererConstants::CacheSerializationVersion;
 	
 	if ( bLoadAndSaveCache )
@@ -494,20 +496,19 @@ uint32 FAssetDataGatherer::Run()
 	TArray<FString> LocalCookedPackageNamesWithoutAssetDataResults;
 
 	const double InitialScanStartTime = FPlatformTime::Seconds();
+	double LastCacheWriteTime = InitialScanStartTime;
 	int32 NumCachedFiles = 0;
 	int32 NumUncachedFiles = 0;
 
-	int32 NumFilesProcessedSinceLastCacheSave = 0;
 	auto WriteAssetCacheFile = [&]()
 	{
+		TRACE_CPUPROFILER_EVENT_SCOPE(WriteAssetCacheFile)
 		FNameTableArchiveWriter CachedAssetDataWriter(CacheSerializationVersion, CacheFilename);
 
 		FAssetRegistryVersion::Type Version = FAssetRegistryVersion::LatestVersion;
 		FAssetRegistryVersion::SerializeVersion(CachedAssetDataWriter, Version);
 
 		SerializeCache(CachedAssetDataWriter);
-
-		NumFilesProcessedSinceLastCacheSave = 0;
 	};
 
 	while ( true )
@@ -559,6 +560,24 @@ uint32 FAssetDataGatherer::Run()
 
 		if (LocalFilesToSearch.Num() > 0)
 		{
+			struct FReadContext
+			{
+				FName PackageName;
+				const FDiscoveredPackageFile& AssetFileData;
+				TArray<FAssetData*> AssetDataFromFile;
+				FPackageDependencyData DependencyData;
+				TArray<FString> CookedPackageNamesWithoutAssetData;
+				bool bCanAttemptAssetRetry = false;
+				bool bResult = false;
+
+				FReadContext(FName InPackageName, const FDiscoveredPackageFile& InAssetFileData)
+					: PackageName(InPackageName)
+					, AssetFileData(InAssetFileData)
+				{
+				}
+			};
+
+			TArray<FReadContext> ReadContexts;
 			for (const FDiscoveredPackageFile& AssetFileData : LocalFilesToSearch)
 			{
 				if (StopTaskCounter.GetValue() != 0)
@@ -614,57 +633,75 @@ uint32 FAssetDataGatherer::Run()
 
 				if (!bLoadedFromCache)
 				{
-					TArray<FAssetData*> AssetDataFromFile;
-					FPackageDependencyData DependencyData;
-					TArray<FString> CookedPackageNamesWithoutAssetData;
-					bool bCanAttemptAssetRetry = false;
-					if (ReadAssetFile(AssetFileData.PackageFilename, AssetDataFromFile, DependencyData, CookedPackageNamesWithoutAssetData, bCanAttemptAssetRetry))
-					{
-						++NumUncachedFiles;
+					ReadContexts.Emplace(PackageName, AssetFileData);
+				}
+			}
 
-						LocalAssetResults.Append(AssetDataFromFile);
-						if (bGatherDependsData)
-						{
-							LocalDependencyResults.Add(DependencyData);
-						}
-						LocalCookedPackageNamesWithoutAssetDataResults.Append(CookedPackageNamesWithoutAssetData);
+			ParallelFor(ReadContexts.Num(),
+				[this, &ReadContexts](int32 Index)
+				{
+					FReadContext& ReadContext = ReadContexts[Index];
+					ReadContext.bResult = ReadAssetFile(ReadContext.AssetFileData.PackageFilename, ReadContext.AssetDataFromFile, ReadContext.DependencyData, ReadContext.CookedPackageNamesWithoutAssetData, ReadContext.bCanAttemptAssetRetry);
+				},
+				EParallelForFlags::Unbalanced | EParallelForFlags::BackgroundPriority
+			);
+
+			for (FReadContext& ReadContext : ReadContexts)
+			{
+				if (ReadContext.bResult)
+				{
+					++NumUncachedFiles;
 
-						// Don't store info on cooked packages
-						bool bCachePackage = bLoadAndSaveCache && LocalCookedPackageNamesWithoutAssetDataResults.Num() == 0;
-						if (bCachePackage)
+					LocalCookedPackageNamesWithoutAssetDataResults.Append(MoveTemp(ReadContext.CookedPackageNamesWithoutAssetData));
+
+					// Don't store info on cooked packages
+					bool bCachePackage = bLoadAndSaveCache && LocalCookedPackageNamesWithoutAssetDataResults.Num() == 0;
+					if (bCachePackage)
+					{
+						for (const FAssetData* AssetData : ReadContext.AssetDataFromFile)
 						{
-							for (const auto& AssetData : AssetDataFromFile)
+							if (!!(AssetData->PackageFlags & PKG_FilterEditorOnly))
 							{
-								if (!!(AssetData->PackageFlags & PKG_FilterEditorOnly))
-								{
-									bCachePackage = false;
-									break;
-								}
+								bCachePackage = false;
+								break;
 							}
 						}
+					}
 
-						if (bCachePackage)
+					if (bCachePackage)
+					{
+						// Update the cache
+						FDiskCachedAssetData* NewData = new FDiskCachedAssetData(ReadContext.AssetFileData.PackageTimestamp);
+						NewData->AssetDataList.Reserve(ReadContext.AssetDataFromFile.Num());
+						for (const FAssetData* BackgroundAssetData : ReadContext.AssetDataFromFile)
 						{
-							++NumFilesProcessedSinceLastCacheSave;
-
-							// Update the cache
-							FDiskCachedAssetData* NewData = new FDiskCachedAssetData(AssetFileData.PackageTimestamp);
-							NewData->AssetDataList.Reserve(AssetDataFromFile.Num());
-							for (const FAssetData* BackgroundAssetData : AssetDataFromFile)
-							{
-								NewData->AssetDataList.Add(*BackgroundAssetData);
-							}
-							NewData->DependencyData = DependencyData;
+							NewData->AssetDataList.Add(*BackgroundAssetData);
+						}
 
-							NewCachedAssetData.Add(NewData);
-							NewCachedAssetDataMap.Add(PackageName, NewData);
+						// MoveTemp only used if we don't need DependencyData anymore
+						if (bGatherDependsData)
+						{
+							NewData->DependencyData = ReadContext.DependencyData;
+						}
+						else
+						{
+							NewData->DependencyData = MoveTemp(ReadContext.DependencyData);
 						}
+
+						NewCachedAssetData.Add(NewData);
+						NewCachedAssetDataMap.Add(ReadContext.PackageName, NewData);
 					}
-					else if (bCanAttemptAssetRetry)
+
+					LocalAssetResults.Append(MoveTemp(ReadContext.AssetDataFromFile));
+					if (bGatherDependsData)
 					{
-						LocalFilesToRetry.Add(AssetFileData);
+						LocalDependencyResults.Add(MoveTemp(ReadContext.DependencyData));
 					}
 				}
+				else if (ReadContext.bCanAttemptAssetRetry)
+				{
+					LocalFilesToRetry.Add(ReadContext.AssetFileData);
+				}
 			}
 
 			LocalFilesToSearch.Reset();
@@ -673,10 +710,11 @@ uint32 FAssetDataGatherer::Run()
 
 			if (bLoadAndSaveCache)
 			{
-				// Save off the cache files if we're processed enough data since the last save
-				if (NumFilesProcessedSinceLastCacheSave >= AssetDataGathererConstants::MaxFilesToProcessBeforeCacheWrite)
+				// Only write intermediate state cache file if we have spent a good amount of time working on it
+				if (FPlatformTime::Seconds() - LastCacheWriteTime >= AssetDataGathererConstants::MinSecondsToElapseBeforeCacheWrite)
 				{
 					WriteAssetCacheFile();
+					LastCacheWriteTime = FPlatformTime::Seconds();
 				}
 			}
 		}
@@ -856,6 +894,7 @@ void FAssetDataGatherer::SortPathsByPriority(const int32 MaxNumToSort)
 
 bool FAssetDataGatherer::ReadAssetFile(const FString& AssetFilename, TArray<FAssetData*>& AssetDataList, FPackageDependencyData& DependencyData, TArray<FString>& CookedPackageNamesWithoutAssetData, bool& OutCanRetry) const
 {
+	TRACE_CPUPROFILER_EVENT_SCOPE(FAssetDataGatherer::ReadAssetFile)
 	OutCanRetry = false;
 
 	FPackageReader PackageReader;
diff --git a/Engine/Source/Runtime/Core/Public/Async/ParallelFor.h b/Engine/Source/Runtime/Core/Public/Async/ParallelFor.h
--- a/Engine/Source/Runtime/Core/Public/Async/ParallelFor.h
+++ b/Engine/Source/Runtime/Core/Public/Async/ParallelFor.h
@@ -32,6 +32,9 @@ enum class EParallelForFlags
 
 	// if running on the rendering thread, make sure the ProcessThread is called when idle
 	PumpRenderingThread = 4,
+
+	// tasks should run on background priority threads
+	BackgroundPriority = 8,
 };
 
 ENUM_CLASS_FLAGS(EParallelForFlags)
@@ -95,17 +98,19 @@ namespace ParallelForImpl
 			check(bExited);
 			FPlatformProcess::ReturnSynchEventToPool(Event);
 		}
-		bool Process(int32 TasksToSpawn, TSharedRef<TParallelForData, ESPMode::ThreadSafe>& Data, bool bMaster);
+		bool Process(int32 TasksToSpawn, TSharedRef<TParallelForData, ESPMode::ThreadSafe>& Data, ENamedThreads::Type InDesiredThread, bool bMaster);
 	};
 
 	template<typename FunctionType>
 	class TParallelForTask
 	{
 		TSharedRef<TParallelForData<FunctionType>, ESPMode::ThreadSafe> Data;
+		ENamedThreads::Type DesiredThread;
 		int32 TasksToSpawn;
 	public:
-		TParallelForTask(TSharedRef<TParallelForData<FunctionType>, ESPMode::ThreadSafe>& InData, int32 InTasksToSpawn = 0)
+		TParallelForTask(TSharedRef<TParallelForData<FunctionType>, ESPMode::ThreadSafe>& InData, ENamedThreads::Type InDesiredThread, int32 InTasksToSpawn = 0)
 			: Data(InData) 
+			, DesiredThread(InDesiredThread)
 			, TasksToSpawn(InTasksToSpawn)
 		{
 		}
@@ -113,10 +118,12 @@ namespace ParallelForImpl
 		{
 			return GET_STATID(STAT_ParallelForTask);
 		}
-		static FORCEINLINE ENamedThreads::Type GetDesiredThread()
+		
+		FORCEINLINE ENamedThreads::Type GetDesiredThread()
 		{
-			return ENamedThreads::AnyHiPriThreadHiPriTask;
+			return DesiredThread;
 		}
+
 		static FORCEINLINE ESubsequentsMode::Type GetSubsequentsMode() 
 		{ 
 			return ESubsequentsMode::FireAndForget; 
@@ -124,7 +131,7 @@ namespace ParallelForImpl
 		void DoTask(ENamedThreads::Type CurrentThread, const FGraphEventRef& MyCompletionGraphEvent)
 		{
 			FMemMark Mark(FMemStack::Get());
-			if (Data->Process(TasksToSpawn, Data, false))
+			if (Data->Process(TasksToSpawn, Data, DesiredThread, false))
 			{
 				checkSlow(!Data->bTriggered);
 				Data->bTriggered = true;
@@ -134,13 +141,13 @@ namespace ParallelForImpl
 	};
 
 	template<typename FunctionType>
-	inline bool TParallelForData<FunctionType>::Process(int32 TasksToSpawn, TSharedRef<TParallelForData<FunctionType>, ESPMode::ThreadSafe>& Data, bool bMaster)
+	inline bool TParallelForData<FunctionType>::Process(int32 TasksToSpawn, TSharedRef<TParallelForData<FunctionType>, ESPMode::ThreadSafe>& Data, ENamedThreads::Type InDesiredThread, bool bMaster)
 	{
 		int32 MaybeTasksLeft = Num - IndexToDo.GetValue();
 		if (TasksToSpawn && MaybeTasksLeft > 0)
 		{
 			TasksToSpawn = FMath::Min<int32>(TasksToSpawn, MaybeTasksLeft);
-			TGraphTask<TParallelForTask<FunctionType>>::CreateTask().ConstructAndDispatchWhenReady(Data, TasksToSpawn - 1);
+			TGraphTask<TParallelForTask<FunctionType>>::CreateTask().ConstructAndDispatchWhenReady(Data, InDesiredThread, TasksToSpawn - 1);
 		}
 		int32 LocalBlockSize = BlockSize;
 		int32 LocalNum = Num;
@@ -209,12 +216,15 @@ namespace ParallelForImpl
 			return;
 		}
 
-		const bool bPumpRenderingThread = (Flags & EParallelForFlags::PumpRenderingThread) != EParallelForFlags::None;
+		const bool bPumpRenderingThread         = (Flags & EParallelForFlags::PumpRenderingThread) != EParallelForFlags::None;
+		const bool bBackgroundPriority          = (Flags & EParallelForFlags::BackgroundPriority) != EParallelForFlags::None;
+		const ENamedThreads::Type DesiredThread = bBackgroundPriority ? ENamedThreads::AnyBackgroundThreadNormalTask : ENamedThreads::AnyHiPriThreadHiPriTask;
+
 		TParallelForData<FunctionType>* DataPtr = new TParallelForData<FunctionType>(Num, AnyThreadTasks + 1, (Num > AnyThreadTasks + 1) && bPumpRenderingThread, Body, Flags);
 		TSharedRef<TParallelForData<FunctionType>, ESPMode::ThreadSafe> Data = MakeShareable(DataPtr);
-		TGraphTask<TParallelForTask<FunctionType>>::CreateTask().ConstructAndDispatchWhenReady(Data, AnyThreadTasks - 1);
+		TGraphTask<TParallelForTask<FunctionType>>::CreateTask().ConstructAndDispatchWhenReady(Data, DesiredThread, AnyThreadTasks - 1);
 		// this thread can help too and this is important to prevent deadlock on recursion 
-		if (!Data->Process(0, Data, true))
+		if (!Data->Process(0, Data, DesiredThread, true))
 		{
 			if (bPumpRenderingThread && IsInActualRenderingThread())
 			{
@@ -270,13 +280,17 @@ namespace ParallelForImpl
 			return;
 		}
 		check(Num);
+
+		const bool bBackgroundPriority = (Flags & EParallelForFlags::BackgroundPriority) != EParallelForFlags::None;
+		const ENamedThreads::Type DesiredThread = bBackgroundPriority ? ENamedThreads::AnyBackgroundThreadNormalTask : ENamedThreads::AnyHiPriThreadHiPriTask;
+
 		TParallelForData<FunctionType>* DataPtr = new TParallelForData<FunctionType>(Num, AnyThreadTasks, false, Body, Flags);
 		TSharedRef<TParallelForData<FunctionType>, ESPMode::ThreadSafe> Data = MakeShareable(DataPtr);
-		TGraphTask<TParallelForTask<FunctionType>>::CreateTask().ConstructAndDispatchWhenReady(Data, AnyThreadTasks - 1);
+		TGraphTask<TParallelForTask<FunctionType>>::CreateTask().ConstructAndDispatchWhenReady(Data, DesiredThread, AnyThreadTasks - 1);
 		// do the prework
 		CurrentThreadWorkToDoBeforeHelping();
 		// this thread can help too and this is important to prevent deadlock on recursion 
-		if (!Data->Process(0, Data, true))
+		if (!Data->Process(0, Data, DesiredThread, true))
 		{
 			if (IsInRenderingThread() && (Flags & EParallelForFlags::PumpRenderingThread) != EParallelForFlags::None)
 			{
