Added function to query whether a path actually exists in the asset registry

#rnx

diff --git a/Engine/Source/Runtime/AssetRegistry/Private/AssetRegistry.cpp b/Engine/Source/Runtime/AssetRegistry/Private/AssetRegistry.cpp
--- a/Engine/Source/Runtime/AssetRegistry/Private/AssetRegistry.cpp
+++ b/Engine/Source/Runtime/AssetRegistry/Private/AssetRegistry.cpp
@@ -1324,6 +1324,16 @@ bool UAssetRegistryImpl::RemovePath(const FString& PathToRemove)
 	return RemoveAssetPath(FName(*PathToRemove));
 }
 
+bool UAssetRegistryImpl::PathExists(const FString& PathToTest) const
+{
+	return PathExists(FName(*PathToTest));
+}
+
+bool UAssetRegistryImpl::PathExists(const FName PathToTest) const
+{
+	return CachedPathTree.PathExists(PathToTest);
+}
+
 void UAssetRegistryImpl::ScanPathsSynchronous(const TArray<FString>& InPaths, bool bForceRescan)
 {
 	ScanPathsAndFilesSynchronous(InPaths, TArray<FString>(), bForceRescan, EAssetDataCacheMode::UseModularCache);
diff --git a/Engine/Source/Runtime/AssetRegistry/Private/AssetRegistry.h b/Engine/Source/Runtime/AssetRegistry/Private/AssetRegistry.h
--- a/Engine/Source/Runtime/AssetRegistry/Private/AssetRegistry.h
+++ b/Engine/Source/Runtime/AssetRegistry/Private/AssetRegistry.h
@@ -74,6 +74,8 @@ public:
 	virtual void PrioritizeAssetInstall(const FAssetData& AssetData) const override;
 	virtual bool AddPath(const FString& PathToAdd) override;
 	virtual bool RemovePath(const FString& PathToRemove) override;
+	virtual bool PathExists(const FString& PathToTest) const override;
+	virtual bool PathExists(const FName PathToTest) const override;
 	virtual void SearchAllAssets(bool bSynchronousSearch) override;
 	virtual void ScanPathsSynchronous(const TArray<FString>& InPaths, bool bForceRescan = false) override;
 	virtual void ScanFilesSynchronous(const TArray<FString>& InFilePaths, bool bForceRescan = false) override;
diff --git a/Engine/Source/Runtime/AssetRegistry/Private/PathTree.cpp b/Engine/Source/Runtime/AssetRegistry/Private/PathTree.cpp
--- a/Engine/Source/Runtime/AssetRegistry/Private/PathTree.cpp
+++ b/Engine/Source/Runtime/AssetRegistry/Private/PathTree.cpp
@@ -149,6 +149,31 @@ bool FPathTree::RemovePath(FName Path)
 	return true;
 }
 
+bool FPathTree::PathExists(FName Path) const
+{
+	if (Path.IsNone())
+	{
+		return false;
+	}
+
+	const TSet<FName>* ChildPathsPtr = ParentPathToChildPaths.Find(Path);
+	if (!ChildPathsPtr)
+	{
+		// Paths are cached without their trailing slash, so if the given path has a trailing slash, test it again now as it may already be cached
+		// We do this after the initial map test as: a) Most paths are well formed, b) This avoids string allocations until we know we need them
+		FString PathStr = Path.ToString();
+		if (PathStr[PathStr.Len() - 1] == '/')
+		{
+			PathStr.RemoveAt(PathStr.Len() - 1, 1, /*bAllowShrinking*/false);
+			Path = *PathStr;
+
+			ChildPathsPtr = ParentPathToChildPaths.Find(Path);
+		}
+	}
+
+	return ChildPathsPtr != nullptr;
+}
+
 bool FPathTree::GetAllPaths(TSet<FName>& OutPaths) const
 {
 	OutPaths.Reset();
diff --git a/Engine/Source/Runtime/AssetRegistry/Private/PathTree.h b/Engine/Source/Runtime/AssetRegistry/Private/PathTree.h
--- a/Engine/Source/Runtime/AssetRegistry/Private/PathTree.h
+++ b/Engine/Source/Runtime/AssetRegistry/Private/PathTree.h
@@ -13,6 +13,9 @@ public:
 	/** Removes the specified path in the tree relative to this node. Returns true if the path was found and removed. */
 	bool RemovePath(FName Path);
 
+	/** Checks whether the given path is one that we know about */
+	bool PathExists(FName Path) const;
+
 	/** Get all of the paths we know about */
 	bool GetAllPaths(TSet<FName>& OutPaths) const;
 
diff --git a/Engine/Source/Runtime/AssetRegistry/Public/IAssetRegistry.h b/Engine/Source/Runtime/AssetRegistry/Public/IAssetRegistry.h
--- a/Engine/Source/Runtime/AssetRegistry/Public/IAssetRegistry.h
+++ b/Engine/Source/Runtime/AssetRegistry/Public/IAssetRegistry.h
@@ -378,6 +378,10 @@ public:
 	/** Attempts to remove the specified path to the set of cached paths. This will only succeed if there are no assets left in the specified path. */
 	virtual bool RemovePath(const FString& PathToRemove) = 0;
 
+	/** Queries whether the given path exists in the set of cached paths */
+	virtual bool PathExists(const FString& PathToTest) const = 0;
+	virtual bool PathExists(const FName PathToTest) const = 0;
+
 	/** Scan the supplied paths recursively right now and populate the asset registry. If bForceRescan is true, the paths will be scanned again, even if they were previously scanned */
 	UFUNCTION(BlueprintCallable, Category = "AssetRegistry")
 	virtual void ScanPathsSynchronous(const TArray<FString>& InPaths, bool bForceRescan = false) = 0;
