#UE4 Make a "getter" function const

diff --git a/Engine/Source/Runtime/AssetRegistry/Private/AssetRegistryState.cpp b/Engine/Source/Runtime/AssetRegistry/Private/AssetRegistryState.cpp
--- a/Engine/Source/Runtime/AssetRegistry/Private/AssetRegistryState.cpp
+++ b/Engine/Source/Runtime/AssetRegistry/Private/AssetRegistryState.cpp
@@ -1473,9 +1473,9 @@ bool FAssetRegistryState::RemoveAssetData(FAssetData* AssetData, bool bRemoveDep
 	return bRemoved;
 }
 
-FDependsNode* FAssetRegistryState::FindDependsNode(const FAssetIdentifier& Identifier)
+FDependsNode* FAssetRegistryState::FindDependsNode(const FAssetIdentifier& Identifier) const
 {
-	FDependsNode** FoundNode = CachedDependsNodes.Find(Identifier);
+	FDependsNode*const* FoundNode = CachedDependsNodes.Find(Identifier);
 	if (FoundNode)
 	{
 		return *FoundNode;
diff --git a/Engine/Source/Runtime/AssetRegistry/Public/AssetRegistryState.h b/Engine/Source/Runtime/AssetRegistry/Public/AssetRegistryState.h
--- a/Engine/Source/Runtime/AssetRegistry/Public/AssetRegistryState.h
+++ b/Engine/Source/Runtime/AssetRegistry/Public/AssetRegistryState.h
@@ -332,7 +332,7 @@ private:
 	FDependsNode* ResolveRedirector(FDependsNode* InDependency, TMap<FName, FAssetData*>& InAllowedAssets, TMap<FDependsNode*, FDependsNode*>& InCache);
 
 	/** Finds an existing node for the given package and returns it, or returns null if one isn't found */
-	FDependsNode* FindDependsNode(const FAssetIdentifier& Identifier);
+	FDependsNode* FindDependsNode(const FAssetIdentifier& Identifier) const;
 
 	/** Creates a node in the CachedDependsNodes map or finds the existing node and returns it */
 	FDependsNode* CreateOrFindDependsNode(const FAssetIdentifier& Identifier);
