Asset Registry
- Blacklist scan filters will always be honored when synchronously scanning paths.
- Added the possibility to ignore blacklist scan filters when synchronously scanning paths: in such case, remove the scanned folder from the blacklist, so we can receive assets added/removed/updated events.

diff --git a/Engine/Source/Editor/UnrealEd/Private/Commandlets/ConvertLevelsToExternalActorsCommandlet.cpp b/Engine/Source/Editor/UnrealEd/Private/Commandlets/ConvertLevelsToExternalActorsCommandlet.cpp
--- a/Engine/Source/Editor/UnrealEd/Private/Commandlets/ConvertLevelsToExternalActorsCommandlet.cpp
+++ b/Engine/Source/Editor/UnrealEd/Private/Commandlets/ConvertLevelsToExternalActorsCommandlet.cpp
@@ -71,6 +71,10 @@ bool UConvertLevelsToExternalActorsCommandlet::CheckExternalActors(const FString
 		TMap<FName, FName> ActorFiles;
 
 		IAssetRegistry& AssetRegistry = FModuleManager::LoadModuleChecked<FAssetRegistryModule>(TEXT("AssetRegistry")).Get();
+
+		// Let the asset registry process any pending operations before continuing
+		AssetRegistry.Tick(-1);
+
 		FDelegateHandle AddedCheckHandle = AssetRegistry.OnAssetAdded().AddLambda([&ActorFiles](const FAssetData& AssetData)
 		{
 			check(!ActorFiles.Contains(AssetData.ObjectPath));
@@ -88,7 +92,7 @@ bool UConvertLevelsToExternalActorsCommandlet::CheckExternalActors(const FString
 			DuplicatedActorFiles.Add(AssetData.ObjectPath, AssetData.PackageName);			
 		});
 
-		AssetRegistry.ScanPathsSynchronous({LevelExternalPathActors});
+		AssetRegistry.ScanPathsSynchronous({LevelExternalPathActors}, /*bForceRescan*/true, /*bIgnoreBlackListScanFilters*/true);
 
 		AssetRegistry.OnAssetAdded().Remove(AddedCheckHandle);
 		AssetRegistry.OnAssetUpdated().Remove(UpdatedCheckHandle);
diff --git a/Engine/Source/Runtime/AssetRegistry/Private/AssetRegistry.cpp b/Engine/Source/Runtime/AssetRegistry/Private/AssetRegistry.cpp
--- a/Engine/Source/Runtime/AssetRegistry/Private/AssetRegistry.cpp
+++ b/Engine/Source/Runtime/AssetRegistry/Private/AssetRegistry.cpp
@@ -1552,14 +1552,30 @@ bool UAssetRegistryImpl::PathExists(const FName PathToTest) const
 	return CachedPathTree.PathExists(PathToTest);
 }
 
-void UAssetRegistryImpl::ScanPathsSynchronous(const TArray<FString>& InPaths, bool bForceRescan)
+void UAssetRegistryImpl::ScanPathsSynchronous(const TArray<FString>& InPaths, bool bForceRescan, bool bIgnoreBlackListScanFilters)
 {
-	ScanPathsAndFilesSynchronous(InPaths, TArray<FString>(), TArray<FString>(), bForceRescan, EAssetDataCacheMode::UseModularCache);
+	// When ignoring blacklist scan filters, and forcing a rescan of one of the blacklisted folders, remove that folder from the blacklist.
+	// That way, we can now receive assets events such as added / removed / updated.
+	if (bForceRescan && bIgnoreBlackListScanFilters)
+	{
+		for (const FString& Path : InPaths)
+		{
+			for (int i = 0; i < BlacklistScanFilters.Num(); ++i)
+			{
+				if (Path.StartsWith(BlacklistScanFilters[i]))
+				{
+					BlacklistScanFilters.RemoveAt(i--);
+				}
+			}
+		}
+	}
+
+	ScanPathsAndFilesSynchronous(InPaths, TArray<FString>(), BlacklistScanFilters, bForceRescan, EAssetDataCacheMode::UseModularCache);
 }
 
 void UAssetRegistryImpl::ScanFilesSynchronous(const TArray<FString>& InFilePaths, bool bForceRescan)
 {
-	ScanPathsAndFilesSynchronous(TArray<FString>(), InFilePaths, TArray<FString>(), bForceRescan, EAssetDataCacheMode::UseModularCache);
+	ScanPathsAndFilesSynchronous(TArray<FString>(), InFilePaths, BlacklistScanFilters, bForceRescan, EAssetDataCacheMode::UseModularCache);
 }
 
 void UAssetRegistryImpl::PrioritizeSearchPath(const FString& PathToPrioritize)
diff --git a/Engine/Source/Runtime/AssetRegistry/Private/AssetRegistry.h b/Engine/Source/Runtime/AssetRegistry/Private/AssetRegistry.h
--- a/Engine/Source/Runtime/AssetRegistry/Private/AssetRegistry.h
+++ b/Engine/Source/Runtime/AssetRegistry/Private/AssetRegistry.h
@@ -88,7 +88,7 @@ public:
 	virtual bool PathExists(const FName PathToTest) const override;
 	virtual void SearchAllAssets(bool bSynchronousSearch) override;
 	virtual void WaitForCompletion() override;
-	virtual void ScanPathsSynchronous(const TArray<FString>& InPaths, bool bForceRescan = false) override;
+	virtual void ScanPathsSynchronous(const TArray<FString>& InPaths, bool bForceRescan = false, bool bIgnoreBlackListScanFilters = false) override;
 	virtual void ScanFilesSynchronous(const TArray<FString>& InFilePaths, bool bForceRescan = false) override;
 	virtual void PrioritizeSearchPath(const FString& PathToPrioritize) override;
 	virtual void ScanModifiedAssetFiles(const TArray<FString>& InFilePaths) override;
diff --git a/Engine/Source/Runtime/AssetRegistry/Public/AssetRegistry/IAssetRegistry.h b/Engine/Source/Runtime/AssetRegistry/Public/AssetRegistry/IAssetRegistry.h
--- a/Engine/Source/Runtime/AssetRegistry/Public/AssetRegistry/IAssetRegistry.h
+++ b/Engine/Source/Runtime/AssetRegistry/Public/AssetRegistry/IAssetRegistry.h
@@ -386,7 +386,7 @@ public:
 
 	/** Scan the supplied paths recursively right now and populate the asset registry. If bForceRescan is true, the paths will be scanned again, even if they were previously scanned */
 	UFUNCTION(BlueprintCallable, Category = "AssetRegistry")
-	virtual void ScanPathsSynchronous(const TArray<FString>& InPaths, bool bForceRescan = false) = 0;
+	virtual void ScanPathsSynchronous(const TArray<FString>& InPaths, bool bForceRescan = false, bool bIgnoreBlackListScanFilters = false) = 0;
 
 	/** Scan the specified individual files right now and populate the asset registry. If bForceRescan is true, the paths will be scanned again, even if they were previously scanned */
 	UFUNCTION(BlueprintCallable, Category = "AssetRegistry")
diff --git a/Engine/Source/Runtime/Engine/Private/WorldPartition/WorldPartition.cpp b/Engine/Source/Runtime/Engine/Private/WorldPartition/WorldPartition.cpp
--- a/Engine/Source/Runtime/Engine/Private/WorldPartition/WorldPartition.cpp
+++ b/Engine/Source/Runtime/Engine/Private/WorldPartition/WorldPartition.cpp
@@ -218,7 +218,7 @@ void UWorldPartition::Initialize(UWorld* InWorld, const FTransform& InTransform)
 
 				// Do a synchronous scan of the level external actors path.
 				IAssetRegistry& AssetRegistry = FModuleManager::LoadModuleChecked<FAssetRegistryModule>(TEXT("AssetRegistry")).Get();
-				AssetRegistry.ScanPathsSynchronous({LevelExternalActorsPath}, /*bForceRescan*/true);
+				AssetRegistry.ScanPathsSynchronous({LevelExternalActorsPath}, /*bForceRescan*/true, /*bIgnoreBlackListScanFilters*/true);
 
 				FARFilter Filter;
 				Filter.bRecursivePaths = true;
