Update the background search asset gatherer's blacklist filter when changing it.

diff --git a/Engine/Source/Runtime/AssetRegistry/Private/AssetRegistry.cpp b/Engine/Source/Runtime/AssetRegistry/Private/AssetRegistry.cpp
--- a/Engine/Source/Runtime/AssetRegistry/Private/AssetRegistry.cpp
+++ b/Engine/Source/Runtime/AssetRegistry/Private/AssetRegistry.cpp
@@ -1568,6 +1568,11 @@ void UAssetRegistryImpl::ScanPathsSynchronous(const TArray<FString>& InPaths, bo
 				}
 			}
 		}
+
+		if (BackgroundAssetSearch.IsValid())
+		{
+			BackgroundAssetSearch->SetBlacklistScanFilters(BlacklistScanFilters);
+		}
 	}
 
 	ScanPathsAndFilesSynchronous(InPaths, TArray<FString>(), BlacklistScanFilters, bForceRescan, EAssetDataCacheMode::UseModularCache);
