Asset Registry: properly react to newly created paths/files from directory watcher events when the asset registry is in synchronous mode (commandlets).

diff --git a/Engine/Source/Runtime/AssetRegistry/Private/AssetRegistry.cpp b/Engine/Source/Runtime/AssetRegistry/Private/AssetRegistry.cpp
--- a/Engine/Source/Runtime/AssetRegistry/Private/AssetRegistry.cpp
+++ b/Engine/Source/Runtime/AssetRegistry/Private/AssetRegistry.cpp
@@ -2650,6 +2650,10 @@ void UAssetRegistryImpl::AddPathToSearch(const FString& Path)
 	{
 		BackgroundAssetSearch->AddPathToSearch(Path);
 	}
+	else if (!FPlatformProperties::RequiresCookedData())
+	{
+		ScanPathsSynchronous({Path});
+	}
 }
 
 void UAssetRegistryImpl::AddFilesToSearch(const TArray<FString>& Files)
@@ -2658,6 +2662,10 @@ void UAssetRegistryImpl::AddFilesToSearch(const TArray<FString>& Files)
 	{
 		BackgroundAssetSearch->AddFilesToSearch(Files);
 	}
+	else if (!FPlatformProperties::RequiresCookedData())
+	{
+		ScanFilesSynchronous(Files);
+	}
 }
 
 #if WITH_EDITOR
