Use GIsEditor instead of !FPlatformProperties::RequiresCookedData() when synchronously scanning new paths/files to avoid impacting game builds.

diff --git a/Engine/Source/Runtime/AssetRegistry/Private/AssetRegistry.cpp b/Engine/Source/Runtime/AssetRegistry/Private/AssetRegistry.cpp
--- a/Engine/Source/Runtime/AssetRegistry/Private/AssetRegistry.cpp
+++ b/Engine/Source/Runtime/AssetRegistry/Private/AssetRegistry.cpp
@@ -2650,7 +2650,7 @@ void UAssetRegistryImpl::AddPathToSearch(const FString& Path)
 	{
 		BackgroundAssetSearch->AddPathToSearch(Path);
 	}
-	else if (!FPlatformProperties::RequiresCookedData())
+	else if (GIsEditor)
 	{
 		ScanPathsSynchronous({Path});
 	}
@@ -2662,7 +2662,7 @@ void UAssetRegistryImpl::AddFilesToSearch(const TArray<FString>& Files)
 	{
 		BackgroundAssetSearch->AddFilesToSearch(Files);
 	}
-	else if (!FPlatformProperties::RequiresCookedData())
+	else if (GIsEditor)
 	{
 		ScanFilesSynchronous(Files);
 	}
