Fix typo in asset registry warning

diff --git a/Engine/Source/Runtime/AssetRegistry/Private/AssetDataGatherer.cpp b/Engine/Source/Runtime/AssetRegistry/Private/AssetDataGatherer.cpp
--- a/Engine/Source/Runtime/AssetRegistry/Private/AssetDataGatherer.cpp
+++ b/Engine/Source/Runtime/AssetRegistry/Private/AssetDataGatherer.cpp
@@ -465,7 +465,7 @@ FAssetDataGatherer::FAssetDataGatherer(const TArray<FString>& InPaths, const TAr
 	if (!bIsSynchronous && !FPlatformProcess::SupportsMultithreading())
 	{
 		bIsSynchronous = true;
-		UE_LOG(LogAssetRegistry, Warning, TEXT("Requested asyncronous asset data gather, but threading support is disabled. Performing a synchronous gather instead!"));
+		UE_LOG(LogAssetRegistry, Warning, TEXT("Requested asynchronous asset data gather, but threading support is disabled. Performing a synchronous gather instead!"));
 	}
 
 	if (bIsSynchronous)
