Fix BlacklistScanFilter going out of sync preventing some AssetAdded on external actors (NewFiles events being filtered out)

diff --git a/Engine/Source/Runtime/AssetRegistry/Private/AssetDataGatherer.cpp b/Engine/Source/Runtime/AssetRegistry/Private/AssetDataGatherer.cpp
--- a/Engine/Source/Runtime/AssetRegistry/Private/AssetDataGatherer.cpp
+++ b/Engine/Source/Runtime/AssetRegistry/Private/AssetDataGatherer.cpp
@@ -930,6 +930,11 @@ void FAssetDataGatherer::PrioritizeSearchPath(const FString& PathToPrioritize)
 
 void FAssetDataGatherer::SetBlacklistScanFilters(const TArray<FString>& InBlacklistScanFilters)
 {
+	{
+		FScopeLock CritSectionLock(&WorkerThreadCriticalSection);
+		BlacklistScanFilters = InBlacklistScanFilters;
+	}
+
 	if (BackgroundPackageFileDiscovery.IsValid())
 	{
 		BackgroundPackageFileDiscovery->SetBlacklistScanFilters(InBlacklistScanFilters);
