Fixing PVS 7.7 Issues:

warning V792: The 'IsRedundantChild' function located to the right of the operator '|' will be called regardless of the value of the left operand. Perhaps, it is better to use '||'.

Testing code so will allow it to aggressively call IsRedundantChild.

diff --git a/Engine/Source/Runtime/AssetRegistry/Private/PropertyCombinationSetTest.cpp b/Engine/Source/Runtime/AssetRegistry/Private/PropertyCombinationSetTest.cpp
--- a/Engine/Source/Runtime/AssetRegistry/Private/PropertyCombinationSetTest.cpp
+++ b/Engine/Source/Runtime/AssetRegistry/Private/PropertyCombinationSetTest.cpp
@@ -150,7 +150,7 @@ protected:
 						{
 							uint32 ExistingValue = ExistingValues[ExistingIndex];
 							bAddedAlreadyExists = bAddedAlreadyExists | (AddValue == ExistingValue);
-							bAddedIsRedundant = bAddedIsRedundant | IsRedundantChild(AddValue, ExistingValue);
+							bAddedIsRedundant = bAddedIsRedundant | IsRedundantChild(AddValue, ExistingValue); //-V792
 							bool bExistingIsRedundant = IsRedundantChild(ExistingValue, AddValue);
 							if (bExistingIsRedundant)
 							{
