AssetRegistry: Change AssetRegistry to use IPluginManager::OnLoadingPhaseComplete to tell when initial plugins have been loaded, rather than GIsRunning, so that it will work in commandlets in a future change.
Fix AssetRegistry load of script package dependencies to deterministically load all script packages rather than allowing a race between report of the dependency using the script package and load of the script package.
#rnx

diff --git a/Engine/Source/Runtime/AssetRegistry/Private/AssetDataGatherer.cpp b/Engine/Source/Runtime/AssetRegistry/Private/AssetDataGatherer.cpp
--- a/Engine/Source/Runtime/AssetRegistry/Private/AssetDataGatherer.cpp
+++ b/Engine/Source/Runtime/AssetRegistry/Private/AssetDataGatherer.cpp
@@ -428,6 +428,7 @@ FAssetDataGatherer::FAssetDataGatherer(const TArray<FString>& InPaths, const TAr
 	, Thread(nullptr)
 {
 	bGatherDependsData = (GIsEditor && !FParse::Param( FCommandLine::Get(), TEXT("NoDependsGathering") )) || FParse::Param(FCommandLine::Get(),TEXT("ForceDependsGathering"));
+	bInitialPluginsLoaded.store(false);
 
 	if (FParse::Param(FCommandLine::Get(), TEXT("NoAssetRegistryCache")) || FParse::Param(FCommandLine::Get(), TEXT("multiprocess")))
 	{
@@ -972,6 +973,16 @@ void FAssetDataGatherer::SetBlacklistScanFilters(const TArray<FString>& InBlackl
 	}
 }
 
+bool FAssetDataGatherer::IsGatheringDependencies() const
+{
+	return bGatherDependsData;
+}
+
+void FAssetDataGatherer::SetInitialPluginsLoaded()
+{
+	bInitialPluginsLoaded.store(true);
+}
+
 void FAssetDataGatherer::SortPathsByPriority(const int32 MaxNumToSort)
 {
 	FScopeLock CritSectionLock(&WorkerThreadCriticalSection);
@@ -1015,7 +1026,7 @@ bool FAssetDataGatherer::ReadAssetFile(const FString& AssetFilename, TArray<FAss
 		//		have no guarantee that a commandlet or an initialized editor is going to load any more modules/plugins
 		//   -	Likewise, we can only attempt a retry for asynchronous scans, as during a synchronous scan we won't be loading any 
 		//		modules/plugins so it would last forever
-		const bool bAllowRetry = GIsEditor && !IsRunningCommandlet() && !GIsRunning && !bIsSynchronous;
+		const bool bAllowRetry = GIsEditor && !IsRunningCommandlet() && !bInitialPluginsLoaded.load() && !bIsSynchronous;
 		OutCanRetry = bAllowRetry && OpenPackageResult == FPackageReader::EOpenPackageResult::CustomVersionMissing;
 		return false;
 	}
diff --git a/Engine/Source/Runtime/AssetRegistry/Private/AssetDataGatherer.h b/Engine/Source/Runtime/AssetRegistry/Private/AssetDataGatherer.h
--- a/Engine/Source/Runtime/AssetRegistry/Private/AssetDataGatherer.h
+++ b/Engine/Source/Runtime/AssetRegistry/Private/AssetDataGatherer.h
@@ -9,6 +9,7 @@
 #include "HAL/Runnable.h"
 #include "DiskCachedAssetData.h"
 #include "BackgroundGatherResults.h"
+#include <atomic>
 
 /**
  * Minimal amount of information needed about a discovered asset file
@@ -193,6 +194,12 @@ public:
 	/** Set the blacklist filters to use during scanning. */
 	void SetBlacklistScanFilters(const TArray<FString>& InBlacklistScanFilters);
 
+	/** Report whether the gatherer is configured to load depends data in addition to asset data. */
+	bool IsGatheringDependencies() const;
+
+	/** Inform the gatherer that initial plugins have finished loading and it should not retry failed loads in hope of a missing custom version */
+	void SetInitialPluginsLoaded();
+
 private:
 	/** Sort the paths so that items belonging to the current priority path is processed first */
 	void SortPathsByPriority(const int32 MaxNumToSort);
@@ -259,6 +266,9 @@ private:
 	/** True if dependency data should be gathered */
 	bool bGatherDependsData;
 
+	/** Async only. Set to true once initial plugins have been loaded. */
+	std::atomic<bool> bInitialPluginsLoaded;
+
 	/** The cached value of the NumPathsToSearch returned by FAssetDataDiscovery the last time we synchronized with it */
 	int32 NumPathsToSearchAtLastSyncPoint;
 
diff --git a/Engine/Source/Runtime/AssetRegistry/Private/AssetRegistry.cpp b/Engine/Source/Runtime/AssetRegistry/Private/AssetRegistry.cpp
--- a/Engine/Source/Runtime/AssetRegistry/Private/AssetRegistry.cpp
+++ b/Engine/Source/Runtime/AssetRegistry/Private/AssetRegistry.cpp
@@ -26,6 +26,7 @@
 #include "HAL/ThreadHeartBeat.h"
 #include "HAL/PlatformMisc.h"
 #include "AssetRegistryConsoleCommands.h"
+#include "Interfaces/IPluginManager.h"
 
 #if WITH_EDITOR
 #include "IDirectoryWatcher.h"
@@ -491,6 +492,63 @@ void UAssetRegistryImpl::InitRedirectors()
 	}
 }
 
+void UAssetRegistryImpl::RegisterReadOfScriptPackages()
+{
+	if (!BackgroundAssetSearch.IsValid())
+	{
+		return;
+	}
+	IPluginManager& PluginManager = IPluginManager::Get();
+	ELoadingPhase::Type LoadingPhase = PluginManager.GetLastCompletedLoadingPhase();
+	if (LoadingPhase != ELoadingPhase::None && LoadingPhase >= ELoadingPhase::PostEngineInit)
+	{
+		ReadScriptPackages();
+	}
+	else
+	{
+		PluginManager.OnLoadingPhaseComplete().AddUObject(this, &UAssetRegistryImpl::OnPluginLoadingPhaseComplete);
+	}
+}
+
+void UAssetRegistryImpl::OnPluginLoadingPhaseComplete(ELoadingPhase::Type LoadingPhase, bool bPhaseSuccessful)
+{
+	if (LoadingPhase != ELoadingPhase::PostEngineInit)
+	{
+		return;
+	}
+	ReadScriptPackages();
+}
+
+void UAssetRegistryImpl::ReadScriptPackages()
+{
+	if (BackgroundAssetSearch.IsValid())
+	{
+		BackgroundAssetSearch->SetInitialPluginsLoaded();
+		if (BackgroundAssetSearch->IsGatheringDependencies())
+		{
+			// Now that all scripts have been loaded, we need to create AssetPackageDatas for every script
+			// This is also done whenever scripts are referenced in our gather of existing packages,
+			// but we need to complete it for all scripts that were referenced but not yet loaded for packages
+			// that we already gathered
+			// CODEREVIEWTODO: Is there a faster way to get the list of loaded script packages?
+			for (TObjectIterator<UPackage> It; It; ++It)
+			{
+				UPackage* Package = *It;
+				if (Package)
+				{
+					if (Package && FPackageName::IsScriptPackage(Package->GetName()))
+					{
+						FAssetPackageData* ScriptPackageData = State.CreateOrGetAssetPackageData(Package->GetFName());
+						// Get the guid off the script package, it is updated when script is changed so we need to refresh it every run
+						ScriptPackageData->PackageGuid = Package->GetGuid();
+					}
+				}
+			}
+		}
+	}
+	IPluginManager::Get().OnLoadingPhaseComplete().RemoveAll(this);
+}
+
 void UAssetRegistryImpl::InitializeSerializationOptions(FAssetRegistrySerializationOptions& Options, const FString& PlatformIniName) const
 {
 	if (PlatformIniName.IsEmpty())
@@ -696,6 +754,7 @@ UAssetRegistryImpl::~UAssetRegistryImpl()
 	FPackageName::OnContentPathMounted().RemoveAll(this);
 	FPackageName::OnContentPathDismounted().RemoveAll(this);
 	FCoreDelegates::OnPostEngineInit.RemoveAll(this);
+	IPluginManager::Get().OnLoadingPhaseComplete().RemoveAll(this);
 
 #if WITH_EDITOR
 	if (GIsEditor)
@@ -784,6 +843,7 @@ void UAssetRegistryImpl::SearchAllAssets(bool bSynchronousSearch)
 	{
 		// if the BackgroundAssetSearch is already valid then we have already called it before
 		BackgroundAssetSearch = MakeShared<FAssetDataGatherer>(PathsToSearch, TArray<FString>(), BlacklistScanFilters, bSynchronousSearch, EAssetDataCacheMode::UseMonolithicCache);
+		RegisterReadOfScriptPackages();
 	}
 }
 
@@ -2440,7 +2500,7 @@ void UAssetRegistryImpl::DependencyDataGathered(const double TickStartTime, TBac
 
 					if (FPackageName::IsScriptPackage(PackageName))
 					{
-						// Get the guid off the script package, this is updated when script is changed
+						// Get the guid off the script package, it is updated when script is changed so we need to refresh it every run
 						UPackage* Package = FindPackage(nullptr, *PackageName);
 
 						if (Package)
diff --git a/Engine/Source/Runtime/AssetRegistry/Private/AssetRegistry.h b/Engine/Source/Runtime/AssetRegistry/Private/AssetRegistry.h
--- a/Engine/Source/Runtime/AssetRegistry/Private/AssetRegistry.h
+++ b/Engine/Source/Runtime/AssetRegistry/Private/AssetRegistry.h
@@ -10,6 +10,7 @@
 #include "PackageDependencyData.h"
 #include "AssetDataGatherer.h"
 #include "BackgroundGatherResults.h"
+#include "ModuleDescriptor.h"
 #include "AssetRegistry.generated.h"
 
 class FDependsNode;
@@ -151,6 +152,11 @@ protected:
 private:
 
 	void InitRedirectors();
+	/** Read script packages if all initial plugins have been loaded, otherwise register with the callback to read them once read. */
+	void RegisterReadOfScriptPackages();
+	void OnPluginLoadingPhaseComplete(ELoadingPhase::Type LoadingPhase, bool bPhaseSuccessful);
+	/** If collecting dependencies, create an FAssetPackageData for every script package, to make sure dependencies can find a data for them. */
+	void ReadScriptPackages();
 
 	/**
 	 * Internal implementation for RunAssetsThroughFilter and UseFilterToExcludeAssets
