AssetRegistryState::Dump: Fix missing lines in last page.
#rnx

diff --git a/Engine/Source/Runtime/AssetRegistry/Private/AssetRegistryState.cpp b/Engine/Source/Runtime/AssetRegistry/Private/AssetRegistryState.cpp
--- a/Engine/Source/Runtime/AssetRegistry/Private/AssetRegistryState.cpp
+++ b/Engine/Source/Runtime/AssetRegistry/Private/AssetRegistryState.cpp
@@ -1813,7 +1813,7 @@ void FAssetRegistryState::Dump(const TArray<FString>& Arguments, TArray<FString>
 	int32 NumLinesInPage = 0;
 	const int32 LineTerminatorLen = TCString<TCHAR>::Strlen(LINE_TERMINATOR);
 
-	auto FinishPage = [&PageBuffer, &NumLinesInPage, HashStartValue, HashMultiplier, PageEndSearchLength, &OutPages, &OverflowText, LineTerminatorLen]()
+	auto FinishPage = [&PageBuffer, &NumLinesInPage, HashStartValue, HashMultiplier, PageEndSearchLength, &OutPages, &OverflowText, LineTerminatorLen](bool bLastPage)
 	{
 		int32 PageEndIndex = PageBuffer.Len();
 		const TCHAR* BufferEnd = PageBuffer.GetData() + PageEndIndex;
@@ -1822,7 +1822,7 @@ void FAssetRegistryState::Dump(const TArray<FString>& Arguments, TArray<FString>
 		// because after one missing or added line, every page from that point on will be offset and therefore different, making false positive differences
 		// To make pages after one missing or added line the same, we look for a good page ending based on the text of all the lines near the end of the current page
 		// By choosing specific-valued texts as page breaks, we will usually randomly get lucky and have the two diffs pick the same line for the end of the page
-		if (NumLinesInPage > PageEndSearchLength)
+		if (!bLastPage && NumLinesInPage > PageEndSearchLength)
 		{
 			const TCHAR* WinningLineEnd = BufferEnd;
 			uint32 WinningLineValue = 0;
@@ -1885,7 +1885,7 @@ void FAssetRegistryState::Dump(const TArray<FString>& Arguments, TArray<FString>
 			}
 			else
 			{
-				FinishPage();
+				FinishPage(false);
 			}
 		}
 	};
@@ -2145,7 +2145,7 @@ void FAssetRegistryState::Dump(const TArray<FString>& Arguments, TArray<FString>
 		}
 		else
 		{
-			FinishPage();
+			FinishPage(true);
 		}
 	}
 }
