Fixed outstanding issue in the PackageReader that would assume ownership of archive even if unwarranted;
Exposed a way to read the FAssetData from a file in the AssetDataGatherer api
Augmented IAssetRegistry api to support getting registry data from a filename, not only an archive
Replaced use of "old" code so that there is one path for asset data recovery only.

diff --git a/Engine/Source/Runtime/AssetRegistry/Private/AssetDataGatherer.cpp b/Engine/Source/Runtime/AssetRegistry/Private/AssetDataGatherer.cpp
--- a/Engine/Source/Runtime/AssetRegistry/Private/AssetDataGatherer.cpp
+++ b/Engine/Source/Runtime/AssetRegistry/Private/AssetDataGatherer.cpp
@@ -1030,7 +1030,14 @@ bool FAssetDataGatherer::ReadAssetFile(const FString& AssetFilename, TArray<FAss
 		OutCanRetry = bAllowRetry && OpenPackageResult == FPackageReader::EOpenPackageResult::CustomVersionMissing;
 		return false;
 	}
+	else
+	{
+		return ReadAssetFile(PackageReader, AssetDataList, (bGatherDependsData ? &DependencyData : nullptr), CookedPackageNamesWithoutAssetData);
+	}
+}
 
+bool FAssetDataGatherer::ReadAssetFile(FPackageReader& PackageReader, TArray<FAssetData*>& AssetDataList, FPackageDependencyData* DependencyData, TArray<FString>& CookedPackageNamesWithoutAssetData)
+{
 	if ( PackageReader.ReadAssetRegistryDataIfCookedPackage(AssetDataList, CookedPackageNamesWithoutAssetData) )
 	{
 		// Cooked data is special. No further data is found in these packages
@@ -1046,9 +1053,9 @@ bool FAssetDataGatherer::ReadAssetFile(const FString& AssetFilename, TArray<FAss
 		}
 	}
 
-	if ( bGatherDependsData )
+	if ( DependencyData )
 	{
-		if ( !PackageReader.ReadDependencyData(DependencyData) )
+		if ( !PackageReader.ReadDependencyData(*DependencyData) )
 		{
 			return false;
 		}
@@ -1059,7 +1066,7 @@ bool FAssetDataGatherer::ReadAssetFile(const FString& AssetFilename, TArray<FAss
 		FName RedirectorClassName = UObjectRedirector::StaticClass()->GetFName();
 		if (Algo::AnyOf(AssetDataList, [RedirectorClassName](FAssetData* AssetData) { return AssetData->AssetClass == RedirectorClassName; }))
 		{
-			TBitArray<>& ImportUsedInGame = DependencyData.ImportUsedInGame;
+			TBitArray<>& ImportUsedInGame = DependencyData->ImportUsedInGame;
 			for (int32 ImportNum = ImportUsedInGame.Num(), Index = 0; Index < ImportNum; ++Index)
 			{
 				ImportUsedInGame[Index] = true;
diff --git a/Engine/Source/Runtime/AssetRegistry/Private/AssetDataGatherer.h b/Engine/Source/Runtime/AssetRegistry/Private/AssetDataGatherer.h
--- a/Engine/Source/Runtime/AssetRegistry/Private/AssetDataGatherer.h
+++ b/Engine/Source/Runtime/AssetRegistry/Private/AssetDataGatherer.h
@@ -11,6 +11,8 @@
 #include "BackgroundGatherResults.h"
 #include <atomic>
 
+class FPackageReader;
+
 /**
  * Minimal amount of information needed about a discovered asset file
  */
@@ -200,6 +202,16 @@ public:
 	/** Inform the gatherer that initial plugins have finished loading and it should not retry failed loads in hope of a missing custom version */
 	void SetInitialPluginsLoaded();
 
+	/**
+	 * Reads FAssetData information out of a previously initialized package reader
+	 * 
+	 * @param PackageReader the previously opened package reader
+	 * @param AssetDataList the FAssetData for every asset found in the file
+	 * @param DependencyData the FPackageDependencyData for every asset found in the file, can be null
+	 * @param CookedPackagesToLoadUponDiscovery the list of cooked packages to be loaded if any
+	 */
+	static bool ReadAssetFile(FPackageReader& PackageReader, TArray<FAssetData*>& AssetDataList, FPackageDependencyData* DependencyData, TArray<FString>& CookedPackagesToLoadUponDiscovery);
+
 private:
 	/** Sort the paths so that items belonging to the current priority path is processed first */
 	void SortPathsByPriority(const int32 MaxNumToSort);
@@ -210,6 +222,7 @@ private:
 	 * @param AssetFilename the name of the file to read
 	 * @param AssetDataList the FAssetData for every asset found in the file
 	 * @param DependencyData the FPackageDependencyData for every asset found in the file
+	 * @param CookedPackagesToLoadUponDiscovery the list of cooked packages to be loaded if any
 	 * @param OutCanRetry Set to true if this file failed to load, but might be loadable later (due to missing modules)
 	 *
 	 * @return true if the file was successfully read
diff --git a/Engine/Source/Runtime/AssetRegistry/Private/AssetRegistry.cpp b/Engine/Source/Runtime/AssetRegistry/Private/AssetRegistry.cpp
--- a/Engine/Source/Runtime/AssetRegistry/Private/AssetRegistry.cpp
+++ b/Engine/Source/Runtime/AssetRegistry/Private/AssetRegistry.cpp
@@ -2147,16 +2147,27 @@ uint32 UAssetRegistryImpl::GetAllocatedSize(bool bLogDetailed) const
 void UAssetRegistryImpl::LoadPackageRegistryData(FArchive& Ar, TArray<FAssetData*> &AssetDataList) const
 {
 	FPackageReader Reader;
-	Reader.OpenPackageFile(&Ar);
-
-	Reader.ReadAssetRegistryData(AssetDataList);
+	if (Reader.OpenPackageFile(&Ar))
+	{
+		ReadAssetFile(Reader, AssetDataList);
+	}
+}
 
-	Reader.ReadAssetDataFromThumbnailCache(AssetDataList);
+void UAssetRegistryImpl::LoadPackageRegistryData(const FString& PackageFilename, TArray<FAssetData*>& AssetDataList) const
+{
+	FPackageReader Reader;
+	if (Reader.OpenPackageFile(PackageFilename))
+	{
+		ReadAssetFile(Reader, AssetDataList);
+	}
+}
 
+bool UAssetRegistryImpl::ReadAssetFile(FPackageReader& PackageReader, TArray<FAssetData*>& AssetDataList) const
+{
 	TArray<FString> CookedPackageNamesWithoutAssetDataGathered;
-	Reader.ReadAssetRegistryDataIfCookedPackage(AssetDataList, CookedPackageNamesWithoutAssetDataGathered);
-
 	// Dependency Data not currently read for this function
+	FPackageDependencyData* DependencyData = nullptr;
+	return FAssetDataGatherer::ReadAssetFile(PackageReader, AssetDataList, DependencyData, CookedPackageNamesWithoutAssetDataGathered);
 }
 
 void UAssetRegistryImpl::InitializeTemporaryAssetRegistryState(FAssetRegistryState& OutState, const FAssetRegistrySerializationOptions& Options, bool bRefreshExisting, const TMap<FName, FAssetData*>& OverrideData) const
diff --git a/Engine/Source/Runtime/AssetRegistry/Private/AssetRegistry.h b/Engine/Source/Runtime/AssetRegistry/Private/AssetRegistry.h
--- a/Engine/Source/Runtime/AssetRegistry/Private/AssetRegistry.h
+++ b/Engine/Source/Runtime/AssetRegistry/Private/AssetRegistry.h
@@ -96,6 +96,7 @@ public:
 	virtual void AppendState(const FAssetRegistryState& InState) override;
 	virtual uint32 GetAllocatedSize(bool bLogDetailed = false) const override;
 	virtual void LoadPackageRegistryData(FArchive& Ar, TArray<FAssetData*>& Data) const override;
+	virtual void LoadPackageRegistryData(const FString& PackageFilename, TArray<FAssetData*>& Data) const override;
 	virtual void InitializeTemporaryAssetRegistryState(FAssetRegistryState& OutState, const FAssetRegistrySerializationOptions& Options, bool bRefreshExisting = false, const TMap<FName, FAssetData*>& OverrideData = TMap<FName, FAssetData*>()) const override;
 	virtual const FAssetRegistryState* GetAssetRegistryState() const override;
 	virtual const TSet<FName>& GetCachedEmptyPackages() const override;
@@ -314,6 +315,13 @@ private:
 	 */
 	bool IsPathMounted(const FString& Path, const TSet<FString>& MountPointsNoTrailingSlashes, FString& StringBuffer) const;
 
+	/**
+	 * Reads Asset file from a previously initialized package reader
+	 * 
+	 * @param PackageReader	Previsouly initialized package reader that maps to the data to be read
+	 * @param AssetDataList	List of the read asset data
+	 */
+	bool ReadAssetFile(FPackageReader& PackageReader, TArray<FAssetData*>& AssetDataList) const;
 private:
 	
 	/** Internal state of the cached asset registry */
diff --git a/Engine/Source/Runtime/AssetRegistry/Private/PackageReader.cpp b/Engine/Source/Runtime/AssetRegistry/Private/PackageReader.cpp
--- a/Engine/Source/Runtime/AssetRegistry/Private/PackageReader.cpp
+++ b/Engine/Source/Runtime/AssetRegistry/Private/PackageReader.cpp
@@ -14,6 +14,7 @@ FPackageReader::FPackageReader()
 	: Loader(nullptr)
 	, PackageFileSize(0)
 	, AssetRegistryDependencyDataOffset(INDEX_NONE)
+	, bLoaderOwner(false)
 {
 	this->SetIsLoading(true);
 	this->SetIsPersistent(true);
@@ -21,7 +22,7 @@ FPackageReader::FPackageReader()
 
 FPackageReader::~FPackageReader()
 {
-	if (Loader)
+	if (Loader && bLoaderOwner)
 	{
 		delete Loader;
 	}
@@ -29,14 +30,18 @@ FPackageReader::~FPackageReader()
 
 bool FPackageReader::OpenPackageFile(const FString& InPackageFilename, EOpenPackageResult* OutErrorCode)
 {
+	check(!Loader);
 	PackageFilename = InPackageFilename;
 	Loader = IFileManager::Get().CreateFileReader(*PackageFilename);
+	bLoaderOwner = true;
 	return OpenPackageFile(OutErrorCode);
 }
 
 bool FPackageReader::OpenPackageFile(FArchive* InLoader, EOpenPackageResult* OutErrorCode)
 {
+	check(!Loader);
 	Loader = InLoader;
+	bLoaderOwner = false;
 	PackageFilename = Loader->GetArchiveName();
 	return OpenPackageFile(OutErrorCode);
 }
diff --git a/Engine/Source/Runtime/AssetRegistry/Private/PackageReader.h b/Engine/Source/Runtime/AssetRegistry/Private/PackageReader.h
--- a/Engine/Source/Runtime/AssetRegistry/Private/PackageReader.h
+++ b/Engine/Source/Runtime/AssetRegistry/Private/PackageReader.h
@@ -77,4 +77,5 @@ private:
 	TArray<FName> NameMap;
 	int64 PackageFileSize;
 	int64 AssetRegistryDependencyDataOffset;
+	bool bLoaderOwner;
 };
diff --git a/Engine/Source/Runtime/AssetRegistry/Public/AssetRegistry/IAssetRegistry.h b/Engine/Source/Runtime/AssetRegistry/Public/AssetRegistry/IAssetRegistry.h
--- a/Engine/Source/Runtime/AssetRegistry/Public/AssetRegistry/IAssetRegistry.h
+++ b/Engine/Source/Runtime/AssetRegistry/Public/AssetRegistry/IAssetRegistry.h
@@ -511,6 +511,8 @@ public:
 	/** Load FPackageRegistry data from the supplied package */
 	virtual void LoadPackageRegistryData(FArchive& Ar, TArray<FAssetData*>& Data) const = 0;
 	
+	/** Load FAssetData from the specified package filename */
+	virtual void LoadPackageRegistryData(const FString& PackageFilename, TArray<FAssetData*>& Data) const = 0;
 	
 protected:
 	// Functions specifically for calling from the asset manager
