Prototype Wrapped Object Pointers (ObjectHandle/ObjectPtr)

Disabling itemized package access dump to log when using wrapped object pointers, will print only aggregate quantity instead.  Farm preflight shows the package-by-package dump can add minutes of time at the end of cooking some projects.

eg: Note the timestamps on log statements for package accesses
Farm cook EngineTest w/o wrapped pointers: https://horde.devtools.epicgames.com/log/600c7d6e3548fa00015a6872
Preflight cook EngineTest w wrapped poitners: https://horde.devtools.epicgames.com/log/600bde4c5803b30001801685

diff --git a/Engine/Source/Runtime/AssetRegistry/Private/PackageBuildDependencyTracker.cpp b/Engine/Source/Runtime/AssetRegistry/Private/PackageBuildDependencyTracker.cpp
--- a/Engine/Source/Runtime/AssetRegistry/Private/PackageBuildDependencyTracker.cpp
+++ b/Engine/Source/Runtime/AssetRegistry/Private/PackageBuildDependencyTracker.cpp
@@ -25,7 +25,8 @@ void FPackageBuildDependencyTracker::DumpData() const
 			++ReferenceCount;
 		}
 	}
-	UE_LOG(LogPackageBuildDependencyTracker, Display, TEXT("Package Accesses (%u referencing packages with a total of %u unique accesses):"), ReferencingPackageCount, ReferenceCount);
+	UE_LOG(LogPackageBuildDependencyTracker, Display, TEXT("Package Accesses (%u referencing packages with a total of %u unique accesses)"), ReferencingPackageCount, ReferenceCount);
+#if 0
 	UE_LOG(LogPackageBuildDependencyTracker, Display, TEXT("========================================================================="));
 	for (const auto& PackageAccessRecord : Records)
 	{
@@ -35,6 +36,7 @@ void FPackageBuildDependencyTracker::DumpData() const
 			UE_LOG(LogPackageBuildDependencyTracker, Display, TEXT("    %s"), *AccessedData.ToString());
 		}
 	}
+#endif
 }
 
 FPackageBuildDependencyTracker::FPackageBuildDependencyTracker()
