AssetRegistry - Change AssetDataGatherer to always use asynchronous background scanning. A tree structure in FAssetDataDiscovery keeps track of what has been scanned, so that ScanPathsSynchronous requests can wait on only what is missing.
#rn Minor AssetRegistry

diff --git a/Engine/Source/Runtime/AssetRegistry/Private/AssetDataGatherer.cpp b/Engine/Source/Runtime/AssetRegistry/Private/AssetDataGatherer.cpp
--- a/Engine/Source/Runtime/AssetRegistry/Private/AssetDataGatherer.cpp
+++ b/Engine/Source/Runtime/AssetRegistry/Private/AssetDataGatherer.cpp
@@ -1,1151 +1,3421 @@
 // Copyright Epic Games, Inc. All Rights Reserved.
 
 #include "AssetDataGatherer.h"
+#include "AssetDataGathererPrivate.h"
+
+#include "Algo/AnyOf.h"
+#include "Algo/Find.h"
+#include "AssetRegistry/AssetData.h"
+#include "AssetRegistryArchive.h"
+#include "AssetRegistryPrivate.h"
+#include "Async/ParallelFor.h"
+#include "Containers/BinaryHeap.h"
 #include "HAL/PlatformProcess.h"
 #include "HAL/RunnableThread.h"
+#include "Math/NumericLimits.h"
 #include "Misc/AsciiSet.h"
+#include "Misc/Char.h"
 #include "Misc/CommandLine.h"
 #include "Misc/Paths.h"
+#include "Misc/PathViews.h"
 #include "Misc/ScopeLock.h"
-#include "AssetRegistryArchive.h"
-#include "AssetRegistryPrivate.h"
+#include "Misc/ScopeExit.h"
 #include "PackageReader.h"
-#include "Async/ParallelFor.h"
-#include "Algo/AnyOf.h"
+#include "Serialization/Archive.h"
+#include "Templates/UniquePtr.h"
+#include "Templates/UnrealTemplate.h"
 
 namespace AssetDataGathererConstants
 {
-	static const int32 MaxFilesToDiscoverBeforeFlush = 2500;
-	static const int32 MaxFilesToGatherBeforeFlush = 250;
-	static const int32 MinSecondsToElapseBeforeCacheWrite = 60;
+	constexpr int32 SingleThreadFilesPerBatch = 3;
+	constexpr int32 ExpectedMaxBatchSize = 100;
+	constexpr int32 MinSecondsToElapseBeforeCacheWrite = 60;
+	static constexpr uint32 CacheSerializationMagic = 0xCBA78339; 
 }
 
-namespace
+namespace UE
+{
+namespace AssetDataGather
+{
+namespace Private
 {
-	class FLambdaDirectoryStatVisitor : public IPlatformFile::FDirectoryStatVisitor
-	{
-	public:
-		typedef TFunctionRef<bool(const TCHAR*, const FFileStatData&)> FLambdaRef;
-		FLambdaRef Callback;
-		explicit FLambdaDirectoryStatVisitor(FLambdaRef InCallback)
-			: Callback(MoveTemp(InCallback))
-		{
-		}
-		virtual bool Visit(const TCHAR* FilenameOrDirectory, const FFileStatData& StatData) override
-		{
-			return Callback(FilenameOrDirectory, StatData);
-		}
-	};
-
-	static constexpr FAsciiSet InvalidLongPackageCharacters(INVALID_LONGPACKAGE_CHARACTERS);
 
-	static bool ConvertToValidLongPackageName(const FString& Filename, /* out */ FString& LongPackageName)
+/** InOutResult = Value, but without shrinking the string to fit. */
+void AssignStringWithoutShrinking(FString& InOutResult, FStringView Value)
+{
+	TArray<TCHAR>& Result = InOutResult.GetCharArray();
+	if (Value.IsEmpty())
 	{
-		// Make sure the path does not contain invalid characters. These packages will not be successfully loaded or read later.
-		return FPackageName::TryConvertFilenameToLongPackageName(Filename, LongPackageName) &&
-			FAsciiSet::HasNone(*LongPackageName, InvalidLongPackageCharacters);
+		Result.Reset();
 	}
-}
-
-namespace AssetDataDiscoveryUtil
-{
-	bool PassesScanFilters(const TArray<FString>& InBlacklistFilters, const FString& InPath)
+	else
 	{
-		for (const FString& Filter : InBlacklistFilters)
-		{
-			if (InPath.StartsWith(Filter))
-			{
-				return false;
-			}
-		}
-		return true;
+		Result.SetNumUninitialized(Value.Len() + 1, false /* bAllowShrinking */);
+		FMemory::Memcpy(Result.GetData(), Value.GetData(), Value.Len() * sizeof(Value[0]));
+		Result[Value.Len()] = '\0';
 	}
 }
 
-FAssetDataDiscovery::FAssetDataDiscovery(const TArray<FString>& InPaths, const TArray<FString>& InBlacklistScanFilters, bool bInIsSynchronous)
-	: bIsSynchronous(bInIsSynchronous)
-	, bIsDiscoveringFiles(true)
-	, BlacklistScanFilters(InBlacklistScanFilters)
-	, StopTaskCounter(0)
-	, Thread(nullptr)
+/** Adapter to allow us to use a lambda for IterateDirectoryStat */
+class FLambdaDirectoryStatVisitor : public IPlatformFile::FDirectoryStatVisitor
 {
-	DirectoriesToSearch.Reserve(InPaths.Num());
-	for (const FString& Path : InPaths)
+public:
+	typedef TFunctionRef<bool(const TCHAR*, const FFileStatData&)> FLambdaRef;
+	FLambdaRef Callback;
+	explicit FLambdaDirectoryStatVisitor(FLambdaRef InCallback)
+		: Callback(MoveTemp(InCallback))
 	{
-		// Convert the package path to a filename with no extension (directory)
-		DirectoriesToSearch.Add(FPackageName::LongPackageNameToFilename(Path / TEXT("")));
 	}
-
-	if (!bIsSynchronous && !FPlatformProcess::SupportsMultithreading())
+	virtual bool Visit(const TCHAR* FilenameOrDirectory, const FFileStatData& StatData) override
 	{
-		bIsSynchronous = true;
-		UE_LOG(LogAssetRegistry, Warning, TEXT("Requested asyncronous asset data discovery, but threading support is disabled. Performing a synchronous discovery instead!"));
+		return Callback(FilenameOrDirectory, StatData);
 	}
+};
 
-	if (bIsSynchronous)
-	{
-		Run();
-	}
-	else
-	{
-		Thread = FRunnableThread::Create(this, TEXT("FAssetDataDiscovery"), 0, TPri_BelowNormal);
-		checkf(Thread, TEXT("Failed to create asset data discovery thread"));
-	}
+FDiscoveredPathData::FDiscoveredPathData(FStringView InLocalAbsPath, FStringView InLongPackageName, FStringView InRelPath, const FDateTime& InPackageTimestamp)
+	: LocalAbsPath(InLocalAbsPath)
+	, LongPackageName(InLongPackageName)
+	, RelPath(InRelPath)
+	, PackageTimestamp(InPackageTimestamp)
+{
 }
 
-FAssetDataDiscovery::~FAssetDataDiscovery()
+FDiscoveredPathData::FDiscoveredPathData(FStringView InLocalAbsPath, FStringView InLongPackageName, FStringView InRelPath)
+	: LocalAbsPath(InLocalAbsPath)
+	, LongPackageName(InLongPackageName)
+	, RelPath(InRelPath)
 {
 }
 
-bool FAssetDataDiscovery::Init()
+void FDiscoveredPathData::Assign(FStringView InLocalAbsPath, FStringView InLongPackageName, FStringView InRelPath)
 {
-	return true;
+	AssignStringWithoutShrinking(LocalAbsPath, InLocalAbsPath);
+	AssignStringWithoutShrinking(LongPackageName, InLongPackageName);
+	AssignStringWithoutShrinking(RelPath, InRelPath);
 }
 
-uint32 FAssetDataDiscovery::Run()
+void FDiscoveredPathData::Assign(FStringView InLocalAbsPath, FStringView InLongPackageName, FStringView InRelPath, const FDateTime& InPackageTimestamp)
 {
-	// Commenting out the code below as it causes cook-on-the-side to wait indefinitely in FAssetRegistryGenerator::Initialize
-	//if (!bIsSynchronous)
-	//{
-	//	// If we're running asynchronous, don't allow these tasks to start until the engine is "running"
-	//	// as we may still be manipulating global lists during start-up configuration
-	//	while (!GIsRunning && StopTaskCounter.GetValue() == 0)
-	//	{
-	//		FPlatformProcess::SleepNoStats(0.1f);
-	//	}
-	//}
-
-	double DiscoverStartTime = FPlatformTime::Seconds();
-	int32 NumDiscoveredFiles = 0;
-
-	FString LocalFilenamePathToPrioritize;
-
-	TSet<FString> LocalDiscoveredPathsSet;
-	TArray<FString> LocalDiscoveredDirectories;
+	Assign(InLocalAbsPath, InLongPackageName, InRelPath);
+	PackageTimestamp = InPackageTimestamp;
+}
 
-	TArray<FDiscoveredPackageFile> LocalPriorityFilesToSearch;
-	TArray<FDiscoveredPackageFile> LocalNonPriorityFilesToSearch;
+uint32 FDiscoveredPathData::GetAllocatedSize() const
+{
+	return LocalAbsPath.GetAllocatedSize() + LongPackageName.GetAllocatedSize() + RelPath.GetAllocatedSize();
+}
 
-	TArray<FString> LocalBlacklistScanFilters;
+FGatheredPathData::FGatheredPathData(FStringView InLocalAbsPath, FStringView InLongPackageName, const FDateTime& InPackageTimestamp)
+	: LocalAbsPath(InLocalAbsPath)
+	, LongPackageName(InLongPackageName)
+	, PackageTimestamp(InPackageTimestamp)
+{
+}
 
-	// This set contains the folders that we should hide by default unless they contain assets
-	TSet<FString> PathsToHideIfEmpty;
-	PathsToHideIfEmpty.Add(TEXT("/Game/Collections"));
+FGatheredPathData::FGatheredPathData(const FDiscoveredPathData& DiscoveredData)
+	:FGatheredPathData(DiscoveredData.LocalAbsPath, DiscoveredData.LongPackageName, DiscoveredData.PackageTimestamp)
+{
+}
 
-	auto FlushLocalResultsIfRequired = [&]()
-	{
-		if (LocalPriorityFilesToSearch.Num() > 0 || LocalNonPriorityFilesToSearch.Num() > 0 || LocalDiscoveredPathsSet.Num() > 0)
-		{
-			TArray<FString> LocalDiscoveredPathsArray = LocalDiscoveredPathsSet.Array();
+void FGatheredPathData::Assign(FStringView InLocalAbsPath, FStringView InLongPackageName, const FDateTime& InPackageTimestamp)
+{
+	AssignStringWithoutShrinking(LocalAbsPath, InLocalAbsPath);
+	AssignStringWithoutShrinking(LongPackageName, InLongPackageName);
+	PackageTimestamp = InPackageTimestamp;
+}
 
-			{
-				FScopeLock CritSectionLock(&WorkerThreadCriticalSection);
+void FGatheredPathData::Assign(const FDiscoveredPathData& DiscoveredData)
+{
+	Assign(DiscoveredData.LocalAbsPath, DiscoveredData.LongPackageName, DiscoveredData.PackageTimestamp);
+}
 
-				// Work around for TArray::Append(const TArray&) not growing expontentially.
-				// This causes O(N^2) allocations when continuously appending small arrays,
-				// at least on platforms with allocators that use QuantizeSize().
-				const auto ReserveAndAppend = [](auto& To, auto&& From)
-				{
-					To.Reserve(FMath::RoundUpToPowerOfTwo(To.Num() + From.Num()));
-					To.Append(MoveTemp(From));
-				};
-
-				// Place all the discovered files into the files to search list
-				ReserveAndAppend(DiscoveredPaths, MoveTemp(LocalDiscoveredPathsArray));
-				ReserveAndAppend(PriorityDiscoveredFiles, MoveTemp(LocalPriorityFilesToSearch));
-				ReserveAndAppend(NonPriorityDiscoveredFiles, MoveTemp(LocalNonPriorityFilesToSearch));
-			}
-		}
+uint32 FGatheredPathData::GetAllocatedSize() const
+{
+	return LocalAbsPath.GetAllocatedSize() + LongPackageName.GetAllocatedSize();
+}
 
-		LocalDiscoveredPathsSet.Reset();
+FScanDir::FScanDir(FMountDir& InMountDir, FScanDir* InParent, FStringView InRelPath)
+	: MountDir(&InMountDir)
+	, Parent(InParent)
+	, RelPath(InRelPath)
+{
+	InMountDir.GetDiscovery().NumDirectoriesToScan.Increment();
+}
 
-		LocalPriorityFilesToSearch.Reset();
-		LocalNonPriorityFilesToSearch.Reset();
-	};
+FScanDir::~FScanDir()
+{
+	// Assert that Shutdown has been called to confirm that the parent no longer has a reference we need to clear.
+	check(!MountDir);
+}
 
-	auto IsPriorityFile = [&](const FString& InPackageFilename) -> bool
+void FScanDir::Shutdown()
+{
+	if (!MountDir)
 	{
-		return !bIsSynchronous && !LocalFilenamePathToPrioritize.IsEmpty() && InPackageFilename.StartsWith(LocalFilenamePathToPrioritize);
-	};
+		// Already shutdown
+		return;
+	}
 
-	auto OnIterateDirectoryItem = [&](const TCHAR* InPackageFilename, const FFileStatData& InPackageStatData) -> bool
+	// Shutdown all children
+	for (TRefCountPtr<FScanDir>& ScanDir : SubDirs)
 	{
-		if (StopTaskCounter.GetValue() != 0)
-		{
-			// Requested to stop - break out of the directory iteration
-			return false;
-		}
-
-		const FString PackageFilenameStr = InPackageFilename;
-		FString PackagePath;
-		bool bIsLongPackagePath = ConvertToValidLongPackageName(PackageFilenameStr, PackagePath);
-		if (AssetDataDiscoveryUtil::PassesScanFilters(LocalBlacklistScanFilters, PackagePath))
-		{
-			if (InPackageStatData.bIsDirectory)
-			{
-				LocalDiscoveredDirectories.Add(PackageFilenameStr / TEXT(""));
-				if (bIsLongPackagePath && !PathsToHideIfEmpty.Contains(PackagePath))
-				{
-					LocalDiscoveredPathsSet.Add(PackagePath);
-				}
-			}
-			else if (FPackageName::IsPackageFilename(PackageFilenameStr))
-			{
-				if (bIsLongPackagePath)
-				{
-					if (IsPriorityFile(PackageFilenameStr))
-					{
-						LocalPriorityFilesToSearch.Add(FDiscoveredPackageFile(PackageFilenameStr, InPackageStatData.ModificationTime));
-					}
-					else
-					{
-						LocalNonPriorityFilesToSearch.Add(FDiscoveredPackageFile(PackageFilenameStr, InPackageStatData.ModificationTime));
-					}
-
-					LocalDiscoveredPathsSet.Add(FPackageName::GetLongPackagePath(PackagePath));
-
-					++NumDiscoveredFiles;
-
-					// Flush the data if we've processed enough
-					if (!bIsSynchronous && (LocalPriorityFilesToSearch.Num() + LocalNonPriorityFilesToSearch.Num()) >= AssetDataGathererConstants::MaxFilesToDiscoverBeforeFlush)
-					{
-						FlushLocalResultsIfRequired();
-					}
-				}
-			}
-		}
-		return true;
-	};
-
-	bool bIsIdle = true;
+		// Destruction contract for FScanDir requires that the parent calls Shutdown before dropping the reference
+		ScanDir->Shutdown();
+		ScanDir.SafeRelease();
+	}
+	SubDirs.Empty();
 
-	while (StopTaskCounter.GetValue() == 0)
+	// Update MountDir data that we influence
+	if (!bIsComplete)
 	{
-		FString LocalDirectoryToSearch;
-		{
-			FScopeLock CritSectionLock(&WorkerThreadCriticalSection);
-
-			if (DirectoriesToSearch.Num() > 0)
-			{
-				bIsDiscoveringFiles = true;
-
-				LocalFilenamePathToPrioritize = FilenamePathToPrioritize;
-
-				// Pop off the first path to search
-				LocalDirectoryToSearch = DirectoriesToSearch[0];
-				DirectoriesToSearch.RemoveAt(0, 1, false);
-			}
-
-			if (BlacklistScanFilters.Num() > 0)
-			{
-				LocalBlacklistScanFilters = MoveTemp(BlacklistScanFilters);
-			}
-		}
-
-		if (LocalDirectoryToSearch.Len() > 0)
-		{
-			if (bIsIdle)
-			{
-				bIsIdle = false;
-
-				// About to start work - reset these
-				DiscoverStartTime = FPlatformTime::Seconds();
-				NumDiscoveredFiles = 0;
-			}
-
-			// Iterate the current search directory
-			FLambdaDirectoryStatVisitor Visitor(OnIterateDirectoryItem);
-			IFileManager::Get().IterateDirectoryStat(*LocalDirectoryToSearch, Visitor);
-
-			{
-				FScopeLock CritSectionLock(&WorkerThreadCriticalSection);
-
-				// Push back any newly discovered sub-directories
-				if (LocalDiscoveredDirectories.Num() > 0)
-				{
-					// Use LocalDiscoveredDirectories as scratch space, then move it back out - this puts the directories we just 
-					// discovered at the start of the list for the next iteration, which can help with disk locality
-					LocalDiscoveredDirectories.Append(MoveTemp(DirectoriesToSearch));
-					DirectoriesToSearch = MoveTemp(LocalDiscoveredDirectories);
-				}
-				LocalDiscoveredDirectories.Reset();
+		MountDir->GetDiscovery().NumDirectoriesToScan.Decrement();
+	}
 
-				if (!bIsSynchronous)
-				{
-					FlushLocalResultsIfRequired();
-					SortPathsByPriority(1);
-				}
-			}
-		}
-		else
+	// Update Parent data that we influence
+	if (Parent) // Root ScanDir has no parent
+	{
+		if (AccumulatedPriority != EPriority::Normal)
 		{
-			if (!bIsIdle)
-			{
-				bIsIdle = true;
-
-				{
-					FScopeLock CritSectionLock(&WorkerThreadCriticalSection);
-					bIsDiscoveringFiles = false;
-				}
-
-				UE_LOG(LogAssetRegistry, Verbose, TEXT("Discovery took %0.6f seconds and found %d files to process"), FPlatformTime::Seconds() - DiscoverStartTime, NumDiscoveredFiles);
-			}
-
-			// Ran out of things to do... if we have any pending results, flush those now
-			FlushLocalResultsIfRequired();
-
-			if (bIsSynchronous)
-			{
-				// This is synchronous. Since our work is done, we should safely exit
-				Stop();
-			}
-			else
-			{
-				// No work to do. Sleep for a little and try again later.
-				FPlatformProcess::Sleep(0.1);
-			}
+			Parent->OnChildPriorityChanged(AccumulatedPriority, -1);
 		}
 	}
 
-	return 0;
+	// Clear backpointers (which also marks us as shutdown)
+	MountDir = nullptr;
+	Parent = nullptr;
 }
 
-void FAssetDataDiscovery::Stop()
+bool FScanDir::IsValid() const
 {
-	StopTaskCounter.Increment();
+	return MountDir != nullptr;
 }
 
-void FAssetDataDiscovery::Exit()
+FMountDir* FScanDir::GetMountDir() const
 {
+	return MountDir;
 }
 
-void FAssetDataDiscovery::EnsureCompletion()
+FStringView FScanDir::GetRelPath() const
 {
-	{
-		FScopeLock CritSectionLock(&WorkerThreadCriticalSection);
-		DirectoriesToSearch.Empty();
-	}
-
-	Stop();
-
-	if (Thread != nullptr)
-	{
-		Thread->WaitForCompletion();
-		delete Thread;
-		Thread = nullptr;
-	}
+	return RelPath;
 }
 
-bool FAssetDataDiscovery::GetAndTrimSearchResults(TArray<FString>& OutDiscoveredPaths, TArray<FDiscoveredPackageFile>& OutDiscoveredFiles, int32& OutNumPathsToSearch)
+EPriority FScanDir::GetPriority() const
 {
-	FScopeLock CritSectionLock(&WorkerThreadCriticalSection);
+	return AccumulatedPriority;
+}
 
-	OutDiscoveredPaths.Append(MoveTemp(DiscoveredPaths));
-	DiscoveredPaths.Reset();
+void FScanDir::AppendLocalAbsPath(FStringBuilderBase& OutFullPath) const
+{
+	if (!MountDir)
+	{
+		return;
+	}
 
-	if (PriorityDiscoveredFiles.Num() > 0)
+	if (Parent)
 	{
-		// Use PriorityDiscoveredFiles as scratch space, then move it back out - this puts the priority files at the start of the final list
-		PriorityDiscoveredFiles.Append(MoveTemp(OutDiscoveredFiles));
-		PriorityDiscoveredFiles.Append(MoveTemp(NonPriorityDiscoveredFiles));
-		OutDiscoveredFiles = MoveTemp(PriorityDiscoveredFiles);
+		Parent->AppendLocalAbsPath(OutFullPath);
+		FPathViews::AppendPath(OutFullPath, RelPath);
 	}
 	else
 	{
-		OutDiscoveredFiles.Append(MoveTemp(NonPriorityDiscoveredFiles));
+		// The root ScanDir should have an empty RelPath from the MountDir
+		check(RelPath.IsEmpty());
+		OutFullPath << MountDir->GetLocalAbsPath();
 	}
-	PriorityDiscoveredFiles.Reset();
-	NonPriorityDiscoveredFiles.Reset();
-
-	OutNumPathsToSearch = DirectoriesToSearch.Num();
-
-	return bIsDiscoveringFiles;
 }
 
-void FAssetDataDiscovery::AddPathToSearch(const FString& Path)
+FString FScanDir::GetLocalAbsPath() const
 {
-	FScopeLock CritSectionLock(&WorkerThreadCriticalSection);
-
-	// Convert the package path to a filename with no extension (directory)
-	DirectoriesToSearch.Add(FPackageName::LongPackageNameToFilename(Path / TEXT("")));
+	TStringBuilder<128> Result;
+	AppendLocalAbsPath(Result);
+	return FString(Result);
 }
 
-void FAssetDataDiscovery::PrioritizeSearchPath(const FString& PathToPrioritize)
+void FScanDir::AppendMountRelPath(FStringBuilderBase& OutRelPath) const
 {
-	FString LocalFilenamePathToPrioritize;
-	if (FPackageName::TryConvertLongPackageNameToFilename(PathToPrioritize / TEXT(""), LocalFilenamePathToPrioritize))
+	if (!MountDir)
 	{
-		FScopeLock CritSectionLock(&WorkerThreadCriticalSection);
+		return;
+	}
 
-		FilenamePathToPrioritize = LocalFilenamePathToPrioritize;
-		SortPathsByPriority(INDEX_NONE);
+	if (Parent)
+	{
+		Parent->AppendMountRelPath(OutRelPath);
+		FPathViews::AppendPath(OutRelPath, RelPath);
+	}
+	else
+	{
+		// The root ScanDir should have an empty RelPath from the MountDir
+		check(RelPath.IsEmpty());
 	}
 }
 
-void FAssetDataDiscovery::SetBlacklistScanFilters(const TArray<FString>& InBlacklistScanFilters)
+FString FScanDir::GetMountRelPath() const
 {
-	FScopeLock CritSectionLock(&WorkerThreadCriticalSection);
-	BlacklistScanFilters = InBlacklistScanFilters;
+	TStringBuilder<128> Result;
+	AppendMountRelPath(Result);
+	return FString(Result);
 }
 
-void FAssetDataDiscovery::SortPathsByPriority(const int32 MaxNumToSort)
+bool FScanDir::IsBlacklisted() const
 {
-	FScopeLock CritSectionLock(&WorkerThreadCriticalSection);
+	if (!MountDir)
+	{
+		return false;
+	}
 
-	// Critical section. This code needs to be as fast as possible since it is in a critical section!
-	// Swap all priority files to the top of the list
-	if (FilenamePathToPrioritize.Len() > 0)
+	const TSet<FString>& Blacklist = MountDir->GetBlacklist();
+	if (Blacklist.Num())
 	{
-		int32 LowestNonPriorityPathIdx = 0;
-		int32 NumSorted = 0;
-		const int32 NumToSort = (MaxNumToSort == INDEX_NONE) ? DirectoriesToSearch.Num() : FMath::Min(DirectoriesToSearch.Num(), MaxNumToSort);
-		for (int32 DirIdx = 0; DirIdx < DirectoriesToSearch.Num(); ++DirIdx)
+		TStringBuilder<128> MountRelPath;
+		AppendMountRelPath(MountRelPath);
+		FStringView MountRelPathSV(MountRelPath);
+		if (Blacklist.ContainsByHash(GetTypeHash(MountRelPathSV), MountRelPathSV))
 		{
-			if (DirectoriesToSearch[DirIdx].StartsWith(FilenamePathToPrioritize))
-			{
-				DirectoriesToSearch.Swap(DirIdx, LowestNonPriorityPathIdx);
-				LowestNonPriorityPathIdx++;
-
-				if (++NumSorted >= NumToSort)
-				{
-					break;
-				}
-			}
+			return true;
 		}
 	}
+	return false;
 }
 
-
-FAssetDataGatherer::FAssetDataGatherer(const TArray<FString>& InPaths, const TArray<FString>& InSpecificFiles, const TArray<FString>& InBlacklistScanFilters, bool bInIsSynchronous, EAssetDataCacheMode AssetDataCacheMode)
-	: StopTaskCounter( 0 )
-	, bIsSynchronous( bInIsSynchronous )
-	, bIsDiscoveringFiles( false )
-	, SearchStartTime( 0 )
-	, BlacklistScanFilters(InBlacklistScanFilters)
-	, NumPathsToSearchAtLastSyncPoint( InPaths.Num() )
-	, bLoadAndSaveCache( false )
-	, bFinishedInitialDiscovery( false )
-	, Thread(nullptr)
+bool FScanDir::IsDirectWhitelisted() const
 {
-	bGatherDependsData = (GIsEditor && !FParse::Param( FCommandLine::Get(), TEXT("NoDependsGathering") )) || FParse::Param(FCommandLine::Get(),TEXT("ForceDependsGathering"));
-	bInitialPluginsLoaded.store(false);
+	return bIsDirectWhitelisted;
+}
 
-	if (FParse::Param(FCommandLine::Get(), TEXT("NoAssetRegistryCache")) || FParse::Param(FCommandLine::Get(), TEXT("multiprocess")))
+bool FScanDir::IsRecursiveMonitored(bool bParentIsWhitelisted) const
+{
+	if (!MountDir)
 	{
-		bLoadAndSaveCache = false;
+		return false;
 	}
-	else if (AssetDataCacheMode != EAssetDataCacheMode::NoCache)
+	if (!bParentIsWhitelisted && !bIsDirectWhitelisted)
 	{
-		if (AssetDataCacheMode == EAssetDataCacheMode::UseMonolithicCache)
-		{
-			bLoadAndSaveCache = true;
-			CacheFilename = FPaths::ProjectIntermediateDir() / (bGatherDependsData ? TEXT("CachedAssetRegistry.bin") : TEXT("CachedAssetRegistryNoDeps.bin"));
-		}
-		else if (InPaths.Num() > 0)
-		{
-			// Sort the paths to try and build a consistent hash for this input
-			TArray<FString> SortedPaths = InPaths;
-			SortedPaths.StableSort();
-
-			// todo: handle hash collisions?
-			uint32 CacheHash = GetTypeHash(SortedPaths[0]);
-			for (int32 PathIndex = 1; PathIndex < SortedPaths.Num(); ++PathIndex)
-			{
-				CacheHash = HashCombine(CacheHash, GetTypeHash(SortedPaths[PathIndex]));
-			}
-
-			bLoadAndSaveCache = true;
-			CacheFilename = FPaths::ProjectIntermediateDir() / TEXT("AssetRegistryCache") / FString::Printf(TEXT("%08x%s.bin"), CacheHash, bGatherDependsData ? TEXT("") : TEXT("NoDeps"));
-		}
+		return false;
 	}
-
-	// Add any specific files before doing search
-	AddFilesToSearch(InSpecificFiles);
-	if (!bIsSynchronous && !FPlatformProcess::SupportsMultithreading())
+	if (IsBlacklisted())
 	{
-		bIsSynchronous = true;
-		UE_LOG(LogAssetRegistry, Warning, TEXT("Requested asynchronous asset data gather, but threading support is disabled. Performing a synchronous gather instead!"));
+		return false;
 	}
+	return true;
+}
 
-	if (bIsSynchronous)
+bool FScanDir::IsPathWhitelisted(FStringView InRelPath, bool bParentIsWhitelisted) const
+{
+	bool bIsRecursiveWhitelisted = bParentIsWhitelisted || bIsDirectWhitelisted;
+	if (bIsRecursiveWhitelisted)
 	{
-		// Run the package file discovery synchronously
-		FAssetDataDiscovery PackageFileDiscovery(InPaths, BlacklistScanFilters, bIsSynchronous);
-		PackageFileDiscovery.GetAndTrimSearchResults(DiscoveredPaths, FilesToSearch, NumPathsToSearchAtLastSyncPoint);
-
-		Run();
+		return true;
 	}
-	else
+	const FScanDir* SubDir = nullptr;
+	if (!InRelPath.IsEmpty())
 	{
-		BackgroundPackageFileDiscovery = MakeShared<FAssetDataDiscovery>(InPaths, BlacklistScanFilters, bIsSynchronous);
-		Thread = FRunnableThread::Create(this, TEXT("FAssetDataGatherer"), 0, TPri_BelowNormal);
-		checkf(Thread, TEXT("Failed to create asset data gatherer thread"));
+		FStringView FirstComponent;
+		FPathViews::SplitFirstComponent(InRelPath, FirstComponent, InRelPath);
+		SubDir = FindSubDir(FirstComponent);
+	}
+	if (!SubDir)
+	{
+		return false;
 	}
+	return SubDir->IsPathWhitelisted(InRelPath, false /* bParentIsWhitelisted */);
 }
 
-FAssetDataGatherer::~FAssetDataGatherer()
+bool FScanDir::ShouldScan(bool bParentIsWhitelisted) const
 {
-	NewCachedAssetDataMap.Empty();
-	DiskCachedAssetDataMap.Empty();
-
-	for ( auto CacheIt = NewCachedAssetData.CreateConstIterator(); CacheIt; ++CacheIt )
-	{
-		delete (*CacheIt);
-	}
-	NewCachedAssetData.Empty();
+	return !bHasScanned && IsRecursiveMonitored(bParentIsWhitelisted);
 }
 
-bool FAssetDataGatherer::Init()
+bool FScanDir::HasScanned() const
 {
-	return true;
+	return bHasScanned;
 }
 
-uint32 FAssetDataGatherer::Run()
+bool FScanDir::IsComplete() const
 {
-	TRACE_CPUPROFILER_EVENT_SCOPE(FAssetDataGatherer::Run)
-	
-	static constexpr uint32 RemovedCacheSerializationVersionMagic = 0xCBA78339;
+	return bIsComplete;
+}
 
-	if ( bLoadAndSaveCache )
+uint32 FScanDir::GetAllocatedSize() const
+{
+	uint32 Result = 0;
+	Result += SubDirs.GetAllocatedSize();
+	for (const TRefCountPtr<FScanDir>& Value : SubDirs)
+	{
+		Result += sizeof(*Value);
+		Result += Value->GetAllocatedSize();
+	}
+	Result += AlreadyScannedFiles.GetAllocatedSize();
+	for (const FString& Value : AlreadyScannedFiles)
 	{
-		TRACE_CPUPROFILER_EVENT_SCOPE(ReadAssetCacheFile)
+		Result += Value.GetAllocatedSize();
+	}
+	Result += RelPath.GetAllocatedSize();
+	return Result;
+}
+
+FScanDir* FScanDir::GetControllingDir(FStringView InRelPath, bool bIsDirectory, bool bParentIsWhitelisted, bool& bOutIsWhitelisted, FString& OutRelPath)
+{
+	// GetControllingDir can only be called on valid ScanDirs, which we rely on since we need to call FindOrAddSubDir which relies on that
+	check(IsValid());
 
-		// load the cached data
-		TUniquePtr<FArchive> FileAr(IFileManager::Get().CreateFileReader(*CacheFilename, FILEREAD_Silent));
-		if (FileAr && !FileAr->IsError() && FileAr->TotalSize() > 2 * sizeof(uint32))
+	bool bIsWhitelisted = bParentIsWhitelisted || bIsDirectWhitelisted;
+	if (InRelPath.IsEmpty())
+	{
+		if (!bIsDirectory)
+		{
+			UE_LOG(LogAssetRegistry, Warning, TEXT("GetControllingDir called on %s with !bIsDirectory, but we have it recorded as a directory. Returning null."), *GetLocalAbsPath());
+			bOutIsWhitelisted = false;
+			OutRelPath.Reset();
+			return nullptr;
+		}
+		else
 		{
-			uint32 MagicNumber = 0;
-			*FileAr << MagicNumber;
+			bOutIsWhitelisted = bIsWhitelisted;
+			OutRelPath = InRelPath;
+			return this;
+		}
+	}
 
-			if (!FileAr->IsError() && MagicNumber == RemovedCacheSerializationVersionMagic)
+	FStringView FirstComponent;
+	FStringView RemainingPath;
+	FPathViews::SplitFirstComponent(InRelPath, FirstComponent, RemainingPath);
+	if (RemainingPath.IsEmpty() && !bIsDirectory)
+	{
+		bOutIsWhitelisted = bIsWhitelisted;
+		OutRelPath = InRelPath;
+		return this;
+	}
+	else
+	{
+		FScanDir* SubDir = nullptr;
+		if (ShouldScan(bParentIsWhitelisted))
+		{
+			SubDir = &FindOrAddSubDir(FirstComponent);
+		}
+		else
+		{
+			SubDir = FindSubDir(FirstComponent);
+			if (!SubDir)
 			{
-				FAssetRegistryVersion::Type RegistryVersion;
-				if (FAssetRegistryVersion::SerializeVersion(*FileAr, RegistryVersion) && RegistryVersion == FAssetRegistryVersion::LatestVersion)
-				{
-					FAssetRegistryReader RegistryReader(*FileAr);
-					if (!RegistryReader.IsError())
-					{
-						SerializeCache(RegistryReader);
-					}	
-				}
+				bOutIsWhitelisted = bIsWhitelisted;
+				OutRelPath = InRelPath;
+				return this;
 			}
 		}
+		return SubDir->GetControllingDir(RemainingPath, bIsDirectory, bIsWhitelisted, bOutIsWhitelisted, OutRelPath);
 	}
+}
 
-	TArray<FDiscoveredPackageFile> LocalFilesToSearch;
-	TArray<FAssetData*> LocalAssetResults;
-	TArray<FPackageDependencyData> LocalDependencyResults;
-	TArray<FString> LocalCookedPackageNamesWithoutAssetDataResults;
-
-	const double InitialScanStartTime = FPlatformTime::Seconds();
-	double LastCacheWriteTime = InitialScanStartTime;
-	int32 NumCachedFiles = 0;
-	int32 NumUncachedFiles = 0;
+bool FScanDir::TrySetDirectoryProperties(FStringView InRelPath, const FSetPathProperties& InProperties, bool bConfirmedExists)
+{
+	// TrySetDirectoryProperties can only be called on valid ScanDirs, which we rely on so we can call FindOrAddSubDir which requires that
+	check(IsValid()); 
 
-	auto WriteAssetCacheFile = [&]()
+	SetComplete(false);
+	if (InRelPath.IsEmpty())
 	{
-		TRACE_CPUPROFILER_EVENT_SCOPE(WriteAssetCacheFile);
-
-		// Save to a temp file first, then move to the destination to avoid corruption
-		FString TempFilename(CacheFilename + TEXT(".tmp"));
-
-		TUniquePtr<FArchive> FileAr(IFileManager::Get().CreateFileWriter(*TempFilename, 0));
-		if (FileAr)
+		// The properties apply to this entire directory
+		if (InProperties.IsWhitelisted.IsSet() && bIsDirectWhitelisted != *InProperties.IsWhitelisted)
 		{
-			int32 MagicNumber = RemovedCacheSerializationVersionMagic;
-			*FileAr << MagicNumber;
+			if (bScanInFlight)
+			{
+				bScanInFlightInvalidated = true;
+			}
+			bIsDirectWhitelisted = *InProperties.IsWhitelisted;
 
-			FAssetRegistryVersion::Type RegistryVersion = FAssetRegistryVersion::LatestVersion;
-			FAssetRegistryVersion::SerializeVersion(*FileAr, RegistryVersion);
+			if (bIsDirectWhitelisted)
+			{
+				// Since we are setting this directory to be monitored, we need to implement the guarantee that all Monitored flags of its children are set to false
+				// We also need to SetComplete false on all directories in between this and a previously whitelisted directory, since those non-whitelisted parent directories
+				// marked themselves complete once their whitelisted children finished
+				ForEachDescendent([](FScanDir& ScanDir)
+					{
+						ScanDir.bIsDirectWhitelisted = false;
+						ScanDir.SetComplete(false);
+					});
+			}
+			else
+			{
+				// Cancel any scans since they are no longer whitelisted
+				ForEachDescendent([](FScanDir& ScanDir)
+					{
+						if (ScanDir.bScanInFlight)
+						{
+							ScanDir.bScanInFlightInvalidated = true;
+						}
+					});
+			}
+		}
+		if (InProperties.HasScanned.IsSet())
+		{
+			bool bNewValue = *InProperties.HasScanned;
+			auto SetProperties = [bNewValue](FScanDir& ScanDir)
+			{
+				if (ScanDir.bScanInFlight)
+				{
+					ScanDir.bScanInFlightInvalidated = true;
+				}
+				ScanDir.bHasScanned = bNewValue;
+				ScanDir.AlreadyScannedFiles.Reset();
+			};
+			SetProperties(*this);
+			ForEachDescendent(SetProperties);
+		}
+		if (InProperties.Priority.IsSet() && DirectPriority != *InProperties.Priority)
+		{
+			SetDirectPriority(*InProperties.Priority);
+		}
+		// InProperties.IgnoreBlacklist does not require an action on ScanDirs; it is implemented on the MountDir level
+		return true;
+	}
+	else
+	{
+		TOptional<FSetPathProperties> ModifiedProperties;
+		const FSetPathProperties* Properties = &InProperties;
+		if (Properties->IsWhitelisted.IsSet())
+		{
+			if (bIsDirectWhitelisted)
+			{
+				// If this directory is set to be monitored, all Monitored flags of its children are unused, are guaranteed set to false, and should not be changed
+				ModifiedProperties = *Properties;
+				ModifiedProperties->IsWhitelisted.Reset();
+				if (!ModifiedProperties->IsSet())
+				{
+					return false;
+				}
+				Properties = &ModifiedProperties.GetValue();
+			}
+		}
+
+		FStringView FirstComponent;
+		FStringView Remainder;
+		FPathViews::SplitFirstComponent(InRelPath, FirstComponent, Remainder);
+
+		FScanDir* SubDir = nullptr;
+		if (bHasScanned &&
+			(!Properties->HasScanned.IsSet() || *Properties->HasScanned == true) &&
+			(!Properties->IsWhitelisted.IsSet()))
+		{
+			// If this parent directory has already been scanned and we are not changing the target directory's has-been-scanned value,
+			// and the next child subdirectory does not exist, then the child directory has already been scanned and we do not need to set the properties on it.
+			// Therefore call FindSubDir instead of FindOrAddSubDir, and abort if we do not find it
+			SubDir = FindSubDir(FirstComponent);
+			if (!SubDir)
+			{
+				return false;
+			}
+		}
+		else
+		{
+			SubDir = FindSubDir(FirstComponent);
+			if (!SubDir)
+			{
+				if (!bConfirmedExists)
+				{
+					TStringBuilder<256> LocalAbsPath;
+					AppendLocalAbsPath(LocalAbsPath);
+					FPathViews::AppendPath(LocalAbsPath, InRelPath);
+					FFileStatData StatData = IFileManager::Get().GetStatData(LocalAbsPath.ToString());
+					if (!StatData.bIsValid || !StatData.bIsDirectory)
+					{
+						UE_LOG(LogAssetRegistry, Warning, TEXT("SetDirectoryProperties called on %s path %.*s. Ignoring the call."),
+							StatData.bIsValid ? TEXT("file") : TEXT("non-existent"), LocalAbsPath.Len(), LocalAbsPath.GetData());
+						return false;
+					}
+					bConfirmedExists = true;
+				}
+				SubDir = &FindOrAddSubDir(FirstComponent);
+			}
+		}
+		return SubDir->TrySetDirectoryProperties(Remainder, *Properties, bConfirmedExists);
+	}
+}
+
+void FScanDir::MarkFileAlreadyScanned(FStringView BaseName)
+{
+	if (bHasScanned)
+	{
+		return;
+	}
+	check(FPathViews::IsPathLeaf(BaseName));
+	for (const FString& AlreadyScannedFile : AlreadyScannedFiles)
+	{
+		if (FStringView(AlreadyScannedFile).Equals(BaseName, ESearchCase::IgnoreCase))
+		{
+			return;
+		}
+	}
+	AlreadyScannedFiles.Emplace(BaseName);
+}
+
+void FScanDir::SetDirectPriority(EPriority InPriority)
+{
+	DirectPriority = InPriority;
+	UpdateAccumulatedPriority();
+}
+
+void FScanDir::UpdateAccumulatedPriority()
+{
+	uint32 LocalAccumulated = static_cast<uint32>(DirectPriority);
+	for (uint32 PriorityLevel = 0; PriorityLevel < CountEPriority; ++PriorityLevel)
+	{
+		if (PriorityRefCounts[PriorityLevel] > 0 && PriorityLevel < LocalAccumulated)
+		{
+			LocalAccumulated = PriorityLevel;
+		}
+	}
+
+	EPriority LocalEPriority = static_cast<EPriority>(LocalAccumulated);
+	if (LocalEPriority != AccumulatedPriority)
+	{
+		if (Parent)
+		{
+			if (AccumulatedPriority != EPriority::Normal)
+			{
+				Parent->OnChildPriorityChanged(AccumulatedPriority, -1);
+			}
+			if (LocalEPriority != EPriority::Normal)
+			{
+				Parent->OnChildPriorityChanged(LocalEPriority, 1);
+			}
+		}
+		AccumulatedPriority = LocalEPriority;
+	}
+}
+
+void FScanDir::OnChildPriorityChanged(EPriority InPriority, int32 Delta)
+{
+	check(-(int32)TNumericLimits<uint8>::Max() < Delta && Delta < TNumericLimits<uint8>::Max());
+	uint8& PriorityRefCount = PriorityRefCounts[static_cast<uint32>(InPriority)];
+	check(Delta > 0 || PriorityRefCount >= -Delta);
+
+	if (Delta > 0 && PriorityRefCount >= TNumericLimits<uint8>::Max() - Delta)
+	{
+		// Mark that the count is now stuck
+		PriorityRefCount = TNumericLimits<uint8>::Max();
+	}
+	else if (Delta < 0 && PriorityRefCount == TNumericLimits<uint8>::Max())
+	{
+		// The count is stuck, do not decrement it
+	}
+	else
+	{
+		PriorityRefCount += Delta;
+	}
+	UpdateAccumulatedPriority();
+}
+
+void FScanDir::SetScanResults(FStringView LocalAbsPath, TArrayView<FDiscoveredPathData>& InOutSubDirs, TArrayView<FDiscoveredPathData>& InOutFiles)
+{
+	// Note that by contract SetScanResults is only called on Paths with ShouldScan == true, so we do not need to check IsWhitelisted for the files in this directory or for any of its subdirs
+	// We do still need to check IsBlacklisted for subdirs, since the blacklist can be true for a subdirectory even if not true for the parent
+	SetComplete(false);
+	check(!bScanInFlightInvalidated);
+	check(MountDir);
+
+	if (!ensure(!bHasScanned))
+	{
+		return;
+	}
+
+	// Add SubDirectories in the tree for the directories found by the scan, and report the directories as discovered directory paths as well
+	// Remove any SubDirectories from the tree that were previously present but are not in the latest scan results
+	TSet<FScanDir*> SubDirsToRemove;
+	ForEachSubDir([&SubDirsToRemove](FScanDir& SubScanDir) { SubDirsToRemove.Add(&SubScanDir); });
+	for (int32 Index = 0; Index < InOutSubDirs.Num(); )
+	{
+		FDiscoveredPathData& SubDirPath = InOutSubDirs[Index];
+		bool bReportResult = false;
+		if (!MountDir->IsBlacklisted(SubDirPath.LocalAbsPath))
+		{
+			FScanDir& SubScanDir = FindOrAddSubDir(SubDirPath.RelPath);
+			SubDirsToRemove.Remove(&SubScanDir);
+			bReportResult = MountDir->GetDiscovery().ShouldDirBeReported(SubDirPath.LongPackageName);
+		}
+		if (!bReportResult)
+		{
+			Swap(SubDirPath, InOutSubDirs.Last());
+			InOutSubDirs = InOutSubDirs.Slice(0, InOutSubDirs.Num() - 1);
+		}
+		else
+		{
+			++Index;
+		}
+	}
+	for (FScanDir* ScanDir : SubDirsToRemove)
+	{
+		RemoveSubDir(ScanDir->GetRelPath());
+	}
+
+	// Add the files that were found in the scan, skipping any files that have already been scanned
+	if (InOutFiles.Num())
+	{
+		auto IsAlreadyScanned = [this, &LocalAbsPath](const FDiscoveredPathData& InFile)
+		{
+			return Algo::AnyOf(AlreadyScannedFiles, [&InFile](const FString& AlreadyScannedFileRelPath) { return FPathViews::Equals(AlreadyScannedFileRelPath, InFile.RelPath); });
+		};
+		bool bScanAll = AlreadyScannedFiles.Num() == 0;
+		for (int32 Index = 0; Index < InOutFiles.Num(); )
+		{
+			FDiscoveredPathData& InFile = InOutFiles[Index];
+			if (!bScanAll && IsAlreadyScanned(InFile))
+			{
+				// Remove this file from InOutFiles
+				Swap(InFile, InOutFiles.Last());
+				InOutFiles = InOutFiles.Slice(0, InOutFiles.Num() - 1);
+			}
+			else
+			{
+				++Index;
+			}
+		}
+	}
+	AlreadyScannedFiles.Empty();
+
+	MountDir->SetHasStartedScanning();
+	bHasScanned = true;
+}
+
+void FScanDir::Update(FScanDir*& OutCursor, bool& bInOutParentIsWhitelisted)
+{
+	check(MountDir);
+	if (bIsComplete)
+	{
+		return;
+	}
+	if (ShouldScan(bInOutParentIsWhitelisted))
+	{
+		OutCursor = this;
+		return;
+	}
+
+	if (SubDirs.Num())
+	{
+		FScanDir* SubDirToScan = FindHighestPrioritySubDir();
+		if (SubDirToScan)
+		{
+			OutCursor = SubDirToScan;
+			bInOutParentIsWhitelisted = bInOutParentIsWhitelisted || bIsDirectWhitelisted;
+			return;
+		}
+	}
+	SetComplete(true);
+
+	OutCursor = Parent; // Note this will be null for the root ScanDir
+	if (!Parent)
+	{
+		bInOutParentIsWhitelisted = false;
+	}
+	else if (Parent->bIsDirectWhitelisted)
+	{
+		// We have a contract bIsDirectWhitelisted is only set on the highest-level directory to monitor and applies to all directories under it.
+		// So we only need to change bInOutParentIsWhitelisted from true to false when we move up the tree into a parent directory with bIsDirectWhitelisted = true.
+		check(!this->bIsDirectWhitelisted); // Verify the contract that children below the ScanDir with bIsDirectWhitelisted = true have bIsDirectWhitelisted = false
+		check(!Parent->Parent || !Parent->Parent->bIsDirectWhitelisted); // Verify the contract that the parent above the ScanDir with bIsDirectWhitelisted = true has bIsDirectWhitelisted = false
+		check(bInOutParentIsWhitelisted); // Verify that the original bInOutParentIsWhitelisted was set to true; it should have been since the parent's bIsDirectWhitelisted is true
+		bInOutParentIsWhitelisted = false;
+	}
+}
+
+bool FScanDir::IsScanInFlight() const
+{
+	return bScanInFlight;
+}
+
+void FScanDir::SetScanInFlight(bool bInScanInFlight)
+{
+	bScanInFlight = bInScanInFlight;
+}
+
+bool FScanDir::IsScanInFlightInvalidated() const
+{
+	return bScanInFlightInvalidated;
+}
+
+void FScanDir::SetScanInFlightInvalidated(bool bInvalidated)
+{
+	bScanInFlightInvalidated = bInvalidated;
+}
+
+void FScanDir::MarkDirty(bool bMarkDescendents)
+{
+	if (bMarkDescendents)
+	{
+		ForEachDescendent([](FScanDir& Descendent) { Descendent.SetComplete(false); });
+	}
+	FScanDir* Current = this;
+	while (Current)
+	{
+		Current->SetComplete(false);
+		Current = Current->Parent;
+	}
+}
+
+void FScanDir::Shrink()
+{
+	ForEachSubDir([](FScanDir& SubDir) {SubDir.Shrink(); });
+	SubDirs.Shrink();
+	AlreadyScannedFiles.Shrink();
+}
+
+
+void FScanDir::SetComplete(bool bInIsComplete)
+{
+	if (!MountDir || bIsComplete == bInIsComplete)
+	{
+		return;
+	}
+
+	bIsComplete = bInIsComplete;
+	if (bIsComplete)
+	{
+		MountDir->GetDiscovery().NumDirectoriesToScan.Decrement();
+		// If we were given a priority, remove it when we complete
+		SetDirectPriority(EPriority::Normal);
+		// All subDirs are complete, so all of their priorities should be set back to normal, so we can unstick any stuck priorities now by setting them all to 0
+#if DO_CHECK
+		bool bHasPriority = false;
+		ForEachSubDir([&bHasPriority](FScanDir& SubDir) { if (SubDir.GetPriority() != EPriority::Normal) { bHasPriority = true; }});
+		if (bHasPriority)
+		{
+			UE_LOG(LogAssetRegistry, Warning, TEXT("ScanDir %s is marked complete, but it has subdirectories with still-set priorities."), *GetLocalAbsPath());
+		}
+		else
+#endif
+		{
+			bool bModifiedRefCount = false;
+			for (uint8& PriorityRefCount : PriorityRefCounts)
+			{
+				bModifiedRefCount = bModifiedRefCount | (PriorityRefCount != 0);
+				PriorityRefCount = 0;
+			}
+			if (bModifiedRefCount)
+			{
+				UpdateAccumulatedPriority();
+			}
+		}
+		// Upon completion, subdirs that do not need to be maintained are deleted, which is done by removing them from the parent
+		// ScanDirs need to be maintained if they are the root, or are whitelisted, or have child ScanDirs that need to be maintained.
+		if (Parent != nullptr && !bIsDirectWhitelisted && SubDirs.IsEmpty())
+		{
+			Parent->RemoveSubDir(GetRelPath());
+			// *this is Shutdown (e.g. Parent is now null) and it may also have been deallocated
+			return;
+		}
+	}
+	else
+	{
+		MountDir->GetDiscovery().NumDirectoriesToScan.Increment();
+	}
+}
+
+FScanDir* FScanDir::FindSubDir(FStringView SubDirBaseName)
+{
+	int32 Index = FindLowerBoundSubDir(SubDirBaseName);
+	if (Index == SubDirs.Num() || !FPathViews::Equals(SubDirs[Index]->GetRelPath(), SubDirBaseName))
+	{
+		return nullptr;
+	}
+	else
+	{
+		return SubDirs[Index].GetReference();
+	}
+}
+
+const FScanDir* FScanDir::FindSubDir(FStringView SubDirBaseName) const
+{
+	return const_cast<FScanDir*>(this)->FindSubDir(SubDirBaseName);
+}
+
+FScanDir& FScanDir::FindOrAddSubDir(FStringView SubDirBaseName)
+{
+	// FindOrAddSubDir is only allowed to be called on valid FScanDirs, which we rely on since we need a non-null MountDir which valid ScanDirs have
+	check(MountDir != nullptr);
+
+	int32 Index = FindLowerBoundSubDir(SubDirBaseName);
+	if (Index == SubDirs.Num() || !FPathViews::Equals(SubDirs[Index]->GetRelPath(), SubDirBaseName))
+	{
+		return *SubDirs.EmplaceAt_GetRef(Index, new FScanDir(*MountDir, this, SubDirBaseName));
+	}
+	else
+	{
+		return *SubDirs[Index];
+	}
+}
+
+void FScanDir::RemoveSubDir(FStringView SubDirBaseName)
+{
+	int32 Index = FindLowerBoundSubDir(SubDirBaseName);
+	if (Index < SubDirs.Num() && FPathViews::Equals(SubDirs[Index]->GetRelPath(), SubDirBaseName))
+	{
+		// Destruction contract for FScanDir requires that the parent calls Shutdown before dropping the reference
+		SubDirs[Index]->Shutdown();
+		SubDirs.RemoveAt(Index);
+	}
+}
+
+int32 FScanDir::FindLowerBoundSubDir(FStringView SubDirBaseName)
+{
+	return Algo::LowerBound(SubDirs, SubDirBaseName,
+		[](const TRefCountPtr<FScanDir>& SubDir, FStringView BaseName)
+		{
+			return FPathViews::Less(SubDir->GetRelPath(), BaseName);
+		}
+	);
+}
+
+FScanDir* FScanDir::FindHighestPrioritySubDir()
+{
+	if (SubDirs.Num() == 0)
+	{
+		return nullptr;
+	}
+
+	FScanDir* WinningSubDir = nullptr;
+	EPriority WinningPriority = EPriority::Normal;
+
+	for (const TRefCountPtr<FScanDir>& SubDir : SubDirs)
+	{
+		if (SubDir->bIsComplete)
+		{
+			continue;
+		}
+		if (WinningSubDir == nullptr || SubDir->AccumulatedPriority < WinningPriority)
+		{
+			WinningSubDir = SubDir.GetReference();
+			WinningPriority = SubDir->AccumulatedPriority;
+		}
+	}
+	return WinningSubDir;
+}
+
+template <typename CallbackType>
+void FScanDir::ForEachSubDir(const CallbackType& Callback)
+{
+	for (TRefCountPtr<FScanDir>& Ptr : SubDirs)
+	{
+		Callback(*Ptr);
+	}
+}
+
+/** Depth-first-search traversal of all descedent subdirs under this (not including this). Callback is called on parents before children. */
+template <typename CallbackType>
+void FScanDir::ForEachDescendent(const CallbackType& Callback)
+{
+	TArray<TPair<FScanDir*, int32>, TInlineAllocator<10>> Stack; // 10 chosen arbitrarily as a depth that is greater than most of our content root directory tree depths
+	Stack.Add(TPair<FScanDir*, int32>(this, 0));
+	while (Stack.Num())
+	{
+		TPair<FScanDir*, int32>& Top = Stack.Last();
+		FScanDir* ParentOnStack = Top.Get<0>();
+		int32& NextIndex = Top.Get<1>();
+		if (NextIndex == ParentOnStack->SubDirs.Num())
+		{
+			Stack.SetNum(Stack.Num() - 1, false /* bAllowShrinking */);
+			continue;
+		}
+		FScanDir* Child = ParentOnStack->SubDirs[NextIndex++];
+		Callback(*Child);
+		Stack.Add(TPair<FScanDir*, int32>(Child, 0));
+	}
+}
+
+FMountDir::FMountDir(FAssetDataDiscovery& InDiscovery, FStringView InLocalAbsPath, FStringView InLongPackageName)
+	: LocalAbsPath(InLocalAbsPath)
+	, LongPackageName(InLongPackageName)
+	, Discovery(InDiscovery)
+{
+	Root = new FScanDir(*this, nullptr, FStringView());
+	UpdateBlacklist();
+}
+
+FMountDir::~FMountDir()
+{
+	// ScanDir's destruction contract requires that the parent calls Shutdown on it before dropping the reference
+	Root->Shutdown();
+	Root.SafeRelease();
+}
+
+FStringView FMountDir::GetLocalAbsPath() const
+{
+	return LocalAbsPath;
+}
+
+FStringView FMountDir::GetLongPackageName() const
+{
+	return LongPackageName;
+}
+
+const TSet<FString>& FMountDir::GetBlacklist() const
+{
+	return BlacklistedRelPaths;
+}
+
+FAssetDataDiscovery& FMountDir::GetDiscovery() const
+{
+	return Discovery;
+}
+
+FScanDir* FMountDir::GetControllingDir(FStringView InLocalAbsPath, bool bIsDirectory, bool& bOutIsWhitelisted, FString& OutRelPath)
+{
+	FStringView RemainingPath;
+	if (!FPathViews::TryMakeChildPathRelativeTo(InLocalAbsPath, GetLocalAbsPath(), RemainingPath))
+	{
+		return nullptr;
+	}
+	return Root->GetControllingDir(RemainingPath, bIsDirectory, false /* bParentIsWhitelisted */, bOutIsWhitelisted, OutRelPath);
+}
+
+uint32 FMountDir::GetAllocatedSize() const
+{
+	uint32 Result = sizeof(*Root);
+	Result += Root->GetAllocatedSize();
+	Result += PathDatas.GetAllocatedSize();
+	for (const FPathData& Value : PathDatas)
+	{
+		Result += Value.GetAllocatedSize();
+	}
+	Result += LongPackageName.GetAllocatedSize();
+	Result += BlacklistedRelPaths.GetAllocatedSize();
+	for (const FString& Value : BlacklistedRelPaths)
+	{
+		Result += Value.GetAllocatedSize();
+	}
+	return Result;
+}
+
+void FMountDir::Shrink()
+{
+	Root->Shrink();
+	PathDatas.Shrink();
+	BlacklistedRelPaths.Shrink();
+}
+
+bool FMountDir::IsComplete() const
+{
+	return Root->IsComplete();
+}
+
+EPriority FMountDir::GetPriority() const
+{
+	return Root->GetPriority();
+}
+
+bool FMountDir::IsPathWhitelisted(FStringView InLocalAbsPath) const
+{
+	FStringView QueryRelPath;
+	if (!ensure(FPathViews::TryMakeChildPathRelativeTo(InLocalAbsPath, GetLocalAbsPath(), QueryRelPath)))
+	{
+		return false;
+	}
+
+	return Root->IsPathWhitelisted(QueryRelPath, false);
+}
+
+bool FMountDir::IsBlacklisted(FStringView InLocalAbsPath) const
+{
+	FStringView QueryRelPath;
+	verify(FPathViews::TryMakeChildPathRelativeTo(InLocalAbsPath, GetLocalAbsPath(), QueryRelPath));
+	for (const FString& Blacklist : BlacklistedRelPaths)
+	{
+		if (FPathViews::IsParentPathOf(Blacklist, QueryRelPath))
+		{
+			return true;
+		}
+	}
+	return false;
+}
+
+bool FMountDir::IsMonitored(FStringView InLocalAbsPath) const
+{
+	FStringView QueryRelPath;
+	if (!ensure(FPathViews::TryMakeChildPathRelativeTo(InLocalAbsPath, GetLocalAbsPath(), QueryRelPath)))
+	{
+		return false;
+	}
+	if (IsBlacklisted(InLocalAbsPath))
+	{
+		return false;
+	}
+
+	return Root->IsPathWhitelisted(QueryRelPath, false);
+}
+
+bool FMountDir::TrySetDirectoryProperties(FStringView InLocalAbsPath, const FSetPathProperties& InProperties, bool bConfirmedExists)
+{
+	FStringView RelPath;
+	if (!ensure(FPathViews::TryMakeChildPathRelativeTo(InLocalAbsPath, GetLocalAbsPath(), RelPath)))
+	{
+		return false;
+	}
+	bool bResult = false;
+
+	const FSetPathProperties* Properties = &InProperties;
+	TOptional<FSetPathProperties> ModifiedProperties;
+	if (Properties->IgnoreBlacklist.IsSet())
+	{
+		// IgnoreBlacklist is applied at the MountDir level, so we handle it separately
+		bool bIgnoreBlacklist = *Properties->IgnoreBlacklist;
+		FPathData* PathData = nullptr;
+		if (bIgnoreBlacklist)
+		{
+			PathData = &FindOrAddPathData(RelPath);
+		}
+		else
+		{
+			PathData = FindPathData(RelPath);
+		}
+		if (PathData)
+		{
+			PathData->bIgnoreBlacklist = bIgnoreBlacklist;
+			if (PathData->IsEmpty())
+			{
+				RemovePathData(RelPath);
+			}
+		}
+		UpdateBlacklist();
+		
+		MarkDirty(RelPath);
+
+		ModifiedProperties = InProperties;
+		ModifiedProperties->IgnoreBlacklist.Reset();
+		if (!ModifiedProperties->IsSet())
+		{
+			return true;
+		}
+		Properties = &ModifiedProperties.GetValue();
+		bResult = true;
+	}
+
+	bResult = Root->TrySetDirectoryProperties(RelPath, InProperties, bConfirmedExists) || bResult;
+	return bResult;
+}
+
+void FMountDir::UpdateBlacklist()
+{
+	BlacklistedRelPaths.Empty(Discovery.BlacklistMountRelativePaths.Num());
+	for (const FString& BlacklistName : Discovery.BlacklistLongPackageNames)
+	{
+		FStringView MountRelPath;
+		if (FPathViews::TryMakeChildPathRelativeTo(BlacklistName, LongPackageName, MountRelPath))
+		{
+			FPathData* PathData = FindPathData(MountRelPath);
+			if (!PathData || !PathData->bIgnoreBlacklist)
+			{
+				// Note that an empty RelPath means we blacklist the entire mountpoint
+				BlacklistedRelPaths.Emplace(MountRelPath);
+			}
+		}
+	}
+	for (const FString& MountRelPath : Discovery.BlacklistMountRelativePaths)
+	{
+		FPathData* PathData = FindPathData(MountRelPath);
+		if (!PathData || !PathData->bIgnoreBlacklist)
+		{
+			BlacklistedRelPaths.Emplace(MountRelPath);
+		}
+	}
+	for (const FPathData& PathData : PathDatas)
+	{
+		if (PathData.bIsChildPath)
+		{
+			BlacklistedRelPaths.Emplace(PathData.RelPath);
+		}
+	}
+}
+
+void FMountDir::Update(FScanDir*& OutCursor, bool& bOutCursorParentIsWhitelisted)
+{
+	bOutCursorParentIsWhitelisted = false;
+	Root->Update(OutCursor, bOutCursorParentIsWhitelisted);
+}
+
+void FMountDir::SetHasStartedScanning()
+{
+	bHasStartedScanning = true;
+}
+
+void FMountDir::AddChildMount(FMountDir* ChildMount)
+{
+	if (!ChildMount)
+	{
+		return;
+	}
+	FStringView RelPath;
+	if (!FPathViews::TryMakeChildPathRelativeTo(ChildMount->GetLocalAbsPath(), LocalAbsPath, RelPath))
+	{
+		return;
+	}
+	FindOrAddPathData(RelPath).bIsChildPath = true;
+	if (bHasStartedScanning)
+	{
+		UE_LOG(LogAssetRegistry, Warning, TEXT("AssetDataGatherer directory %.*s has already started scanning when a new mountpoint was added under it at %.*s. ")
+			TEXT("Assets in the new mount point may exist twice in the AssetRegistry under two different package names."),
+			LocalAbsPath.Len(), *LocalAbsPath, ChildMount->LocalAbsPath.Len(), *ChildMount->LocalAbsPath);
+	}
+	UpdateBlacklist();
+	MarkDirty(RelPath);
+}
+
+void FMountDir::RemoveChildMount(FMountDir* ChildMount)
+{
+	if (!ChildMount)
+	{
+		return;
+	}
+	FStringView RelPath;
+	if (!FPathViews::TryMakeChildPathRelativeTo(ChildMount->GetLocalAbsPath(), LocalAbsPath, RelPath))
+	{
+		return;
+	}
+	FPathData* PathData = FindPathData(RelPath);
+	if (!PathData)
+	{
+		return;
+	}
+	PathData->bIsChildPath = false;
+	if (PathData->IsEmpty())
+	{
+		RemovePathData(RelPath);
+	}
+	if (ChildMount->bHasStartedScanning)
+	{
+		UE_LOG(LogAssetRegistry, Warning, TEXT("AssetDataGatherer directory %.*s has already started scanning when it was removed and merged into its parent mount at %.*s. ")
+			TEXT("Assets in the new mount point may exist twice in the AssetRegistry under two different package names."),
+			ChildMount->LocalAbsPath.Len(), *ChildMount->LocalAbsPath, LocalAbsPath.Len(), *LocalAbsPath);
+	}
+	UpdateBlacklist();
+	MarkDirty(RelPath);
+}
+
+void FMountDir::OnDestroyClearChildMounts()
+{
+	// This function deletes more data than just the child paths; if it becomes used for purposes other than deleting the MountDir,
+	// we will need to change it to only set PathData->bIsChildPath = false;
+	PathDatas.Empty();
+}
+
+void FMountDir::SetParentMount(FMountDir* Parent)
+{
+	ParentMount = Parent;
+}
+
+FMountDir* FMountDir::GetParentMount() const
+{
+	return ParentMount;
+}
+
+TArray<FMountDir*> FMountDir::GetChildMounts() const
+{
+	// Called within Discovery's TreeLock
+	TArray<FMountDir*> Result;
+	for (const FPathData& PathData : PathDatas)
+	{
+		if (PathData.bIsChildPath)
+		{
+			TStringBuilder<256> ChildAbsPath;
+			ChildAbsPath << LocalAbsPath;
+			FPathViews::AppendPath(ChildAbsPath, PathData.RelPath);
+			FMountDir* ChildMount = Discovery.FindMountPoint(ChildAbsPath);
+			if (ensure(ChildMount)) // This PathData information should have been removed with RemoveChildMount when the child MountDir was removed from the Discovery
+			{
+				Result.Add(ChildMount);
+			}
+		}
+	}
+	return Result;
+}
+
+void FMountDir::MarkDirty(FStringView MountRelPath)
+{
+	bool bIsWhitelisted = true;
+	FString ControlRelPath;
+	FScanDir* ScanDir = Root->GetControllingDir(MountRelPath, true /* bIsDirectory */, false /* bParentIsWhitelisted */, bIsWhitelisted, ControlRelPath);
+	if (ScanDir)
+	{
+		// If a ScanDir exists for the directory that is being marked dirty, mark all of its descendents dirty as well.
+		// If the control dir is a parent directory of the requested path, just mark it and its parents dirty
+		// Mark all parent directories that exist as incomplete
+		bool bDirtyAllDescendents = ControlRelPath.IsEmpty();
+		ScanDir->MarkDirty(bDirtyAllDescendents);
+	}
+}
+
+FMountDir::FPathData* FMountDir::FindPathData(FStringView MountRelPath)
+{
+	return Algo::FindByPredicate(PathDatas, [MountRelPath](const FPathData& PathData) { return FPathViews::Equals(PathData.RelPath, MountRelPath); });
+}
+
+FMountDir::FPathData& FMountDir::FindOrAddPathData(FStringView MountRelPath)
+{
+	FPathData* PathData = FindPathData(MountRelPath);
+	if (!PathData)
+	{
+		PathData = &PathDatas.Emplace_GetRef(MountRelPath);
+	}
+	return *PathData;
+}
+
+void FMountDir::RemovePathData(FStringView MountRelPath)
+{
+	PathDatas.RemoveAllSwap([MountRelPath](const FPathData& PathData) { return FPathViews::Equals(PathData.RelPath, MountRelPath);  });
+}
+
+FMountDir::FPathData::FPathData(FStringView MountRelPath)
+:RelPath(MountRelPath)
+{
+}
+
+bool FMountDir::FPathData::IsEmpty() const
+{
+	return !bIgnoreBlacklist && !bIsChildPath;
+}
+
+uint32 FMountDir::FPathData::GetAllocatedSize() const
+{
+	return RelPath.GetAllocatedSize();
+}
+
+FAssetDataDiscovery::FAssetDataDiscovery(const TArray<FString>& InBlacklistLongPackageNames, const TArray<FString>& InBlacklistMountRelativePaths, bool bInIsSynchronous)
+	: BlacklistLongPackageNames(InBlacklistLongPackageNames)
+	, BlacklistMountRelativePaths(InBlacklistMountRelativePaths)
+	, Thread(nullptr)
+	, bIsSynchronous(bInIsSynchronous)
+	, bIsIdle(false)
+	, IsStopped(0)
+	, IsPaused(0)
+	, NumDirectoriesToScan(0)
+{
+	DirLongPackageNamesToNotReport.Add(TEXT("/Game/Collections"));
+
+	if (!bIsSynchronous && !FPlatformProcess::SupportsMultithreading())
+	{
+		bIsSynchronous = true;
+		UE_LOG(LogAssetRegistry, Warning, TEXT("Requested asyncronous asset data discovery, but threading support is disabled. Performing a synchronous discovery instead!"));
+	}
+}
+
+FAssetDataDiscovery::~FAssetDataDiscovery()
+{
+	EnsureCompletion();
+	Cursor.SafeRelease();
+	// Remove pointers to other MountDirs before we delete any of them
+	for (TUniquePtr<FMountDir>& MountDir : MountDirs)
+	{
+		MountDir->SetParentMount(nullptr);
+		MountDir->OnDestroyClearChildMounts();
+	}
+	MountDirs.Empty();
+}
+
+void FAssetDataDiscovery::StartAsync()
+{
+	if (!bIsSynchronous && !Thread)
+	{
+		Thread = FRunnableThread::Create(this, TEXT("FAssetDataDiscovery"), 0, TPri_BelowNormal);
+		checkf(Thread, TEXT("Failed to create asset data discovery thread"));
+	}
+}
+
+bool FAssetDataDiscovery::Init()
+{
+	return true;
+}
+
+uint32 FAssetDataDiscovery::Run()
+{
+	constexpr double IdleSleepTime = 0.1;
+	{
+		FGathererScopeLock ResultsScopeLock(&ResultsLock);
+		DiscoverStartTime = FPlatformTime::Seconds();
+		NumDiscoveredFiles = 0;
+	}
+
+	while (!IsStopped)
+	{
+		{
+			CHECK_IS_NOT_LOCKED_CURRENT_THREAD(TreeLock);
+			CHECK_IS_NOT_LOCKED_CURRENT_THREAD(ResultsLock);
+			FGathererScopeLock TickScopeLock(&TickLock);
+			while (!IsStopped && !bIsIdle && !IsPaused)
+			{
+				TickInternal();
+			}
+		}
+
+		while (!IsStopped && (IsPaused || bIsIdle))
+		{
+			// No work to do. Sleep for a little and try again later.
+			// TODO: Need IsPaused to be a condition variable so we avoid sleeping while waiting for it and then taking a long time to wake after it is unset.
+			FPlatformProcess::Sleep(IdleSleepTime);
+		} 
+	}
+	return 0;
+}
+
+FAssetDataDiscovery::FScopedPause::FScopedPause(const FAssetDataDiscovery& InOwner)
+	:Owner(InOwner)
+{
+	if (!Owner.bIsSynchronous)
+	{
+		Owner.IsPaused++;
+	}
+}
+
+FAssetDataDiscovery::FScopedPause::~FScopedPause()
+{
+	if (!Owner.bIsSynchronous)
+	{
+		check(Owner.IsPaused > 0);
+		Owner.IsPaused--;
+	}
+}
+
+void FAssetDataDiscovery::Stop()
+{
+	IsStopped++;
+}
+
+void FAssetDataDiscovery::Exit()
+{
+}
+
+void FAssetDataDiscovery::EnsureCompletion()
+{
+	Stop();
+
+	if (Thread != nullptr)
+	{
+		Thread->WaitForCompletion();
+		delete Thread;
+		Thread = nullptr;
+	}
+}
+
+
+void FAssetDataDiscovery::TickInternal()
+{
+	TStringBuilder<256> DirLocalAbsPath;
+	TStringBuilder<128> DirLongPackageName;
+	TStringBuilder<128> DirMountRelPath;
+	int32 DirLongPackageNameRootLen;
+	TRefCountPtr<FScanDir> LocalCursor = nullptr;
+	{
+		CHECK_IS_NOT_LOCKED_CURRENT_THREAD(ResultsLock);
+		FGathererScopeLock TreeScopeLock(&TreeLock);
+		for (;;)
+		{
+			// Start at the existing cursor (initializing it if necessary) at call Update until we find a ScanTree that requires scanning
+			if (!Cursor || !Cursor->IsValid())
+			{
+				FScanDir* NewCursor;
+				FindFirstCursor(NewCursor, bCursorParentIsWhitelisted);
+				Cursor = NewCursor;
+				if (!NewCursor)
+				{
+					SetIsIdle(true);
+					return;
+				}
+			}
+			if (Cursor->ShouldScan(bCursorParentIsWhitelisted))
+			{
+				break;
+			}
+
+			FScanDir* NewCursor = Cursor.GetReference();
+			NewCursor->Update(NewCursor, bCursorParentIsWhitelisted);
+			check(NewCursor != Cursor);
+			Cursor = NewCursor;
+		}
+		// IsScanInFlight must be false, because it is not valid to have two TickInternals run at the same time, and we set ScanInFlight back to false after each TickInternal.
+		// If ScanInFlight were true here we would not be able to proceed since we currently don't have a way to find the next ScanTree to update without scanning the current one.
+		check(!Cursor->IsScanInFlight());
+
+		Cursor->SetScanInFlight(true);
+		FMountDir* MountDir = Cursor->GetMountDir();
+		check(MountDir);
+		Cursor->AppendMountRelPath(DirMountRelPath);
+		DirLocalAbsPath << MountDir->GetLocalAbsPath();
+		FPathViews::AppendPath(DirLocalAbsPath, DirMountRelPath);
+		DirLongPackageName << MountDir->GetLongPackageName();
+		FPathViews::AppendPath(DirLongPackageName, DirMountRelPath);
+		DirLongPackageNameRootLen = DirLongPackageName.Len();
+		LocalCursor = Cursor;
+	}
+
+	int32 NumIteratedDirs = 0;
+	int32 NumIteratedFiles = 0;
+	FLambdaDirectoryStatVisitor Visitor(FLambdaDirectoryStatVisitor::FLambdaRef(
+		[this, &DirLocalAbsPath, &DirLongPackageName, DirLongPackageNameRootLen, &NumIteratedDirs, &NumIteratedFiles]
+	(const TCHAR* InPackageFilename, const FFileStatData& InPackageStatData)
+		{
+			FStringView LocalAbsPath(InPackageFilename);
+			FStringView RelPath;
+			FString Buffer;
+			if (!FPathViews::TryMakeChildPathRelativeTo(InPackageFilename, DirLocalAbsPath, RelPath))
+			{
+				// Try again with the path converted to the absolute path format that we passed in; some IFileManagers can send relative paths to the visitor even though the search path is absolute
+				Buffer = FPaths::ConvertRelativePathToFull(FString(InPackageFilename));
+				LocalAbsPath = Buffer;
+				if (!FPathViews::TryMakeChildPathRelativeTo(Buffer, DirLocalAbsPath, RelPath))
+				{
+					UE_LOG(LogAssetRegistry, Warning, TEXT("IterateDirectoryStat returned unexpected result %s which is not a child of the requested path %s."), InPackageFilename, DirLocalAbsPath.ToString());
+					return true;
+				}
+			}
+			if (FPathViews::GetPathLeaf(RelPath).Len() != RelPath.Len())
+			{
+				UE_LOG(LogAssetRegistry, Warning, TEXT("IterateDirectoryStat returned unexpected result %* which is not a direct child of the requested path %s."), InPackageFilename, DirLocalAbsPath.ToString());
+				return true;
+			}
+			ON_SCOPE_EXIT{ DirLongPackageName.RemoveSuffix(DirLongPackageName.Len() - DirLongPackageNameRootLen); };
+
+			if (InPackageStatData.bIsDirectory)
+			{
+				FPathViews::AppendPath(DirLongPackageName, RelPath);
+				// Don't enter directories that contain invalid packagepath characters (including '.'; extensions are not valid in content directories because '.' is not valid in a packagepath)
+				if (!FPackageName::DoesPackageNameContainInvalidCharacters(RelPath))
+				{
+					if (IteratedSubDirs.Num() < NumIteratedDirs + 1)
+					{
+						check(IteratedSubDirs.Num() == NumIteratedDirs);
+						IteratedSubDirs.Emplace();
+					}
+					IteratedSubDirs[NumIteratedDirs++].Assign(LocalAbsPath, DirLongPackageName, RelPath);
+				}
+			}
+			else
+			{
+				FStringView BaseName = FPathViews::GetBaseFilename(RelPath);
+				FPathViews::AppendPath(DirLongPackageName, BaseName);
+				// Don't record files that contain invalid packagepath characters (not counting their extension) or that do not end with a recognized extension
+				if (!FPackageName::DoesPackageNameContainInvalidCharacters(BaseName) && FPackageName::IsPackageFilename(RelPath))
+				{
+					if (IteratedFiles.Num() < NumIteratedFiles + 1)
+					{
+						check(IteratedFiles.Num() == NumIteratedFiles);
+						IteratedFiles.Emplace();
+					}
+					IteratedFiles[NumIteratedFiles++].Assign(LocalAbsPath, DirLongPackageName, RelPath, InPackageStatData.ModificationTime);
+				}
+			}
+			return true;
+		}));
+
+	IFileManager::Get().IterateDirectoryStat(DirLocalAbsPath.ToString(), Visitor);
+	TArrayView<FDiscoveredPathData> LocalSubDirs(IteratedSubDirs.GetData(), NumIteratedDirs);
+	TArrayView<FDiscoveredPathData> LocalDiscoveredFiles(IteratedFiles.GetData(), NumIteratedFiles);
+	bool bValid = false;
+	{
+		CHECK_IS_NOT_LOCKED_CURRENT_THREAD(ResultsLock);
+		FGathererScopeLock TreeScopeLock(&TreeLock);
+		if (!LocalCursor->IsValid())
+		{
+			// The ScanDir has been shutdown, and it is only still allocated to prevent us from crashing. Drop our reference and allow it to delete.
+		}
+		else if (LocalCursor->IsScanInFlightInvalidated())
+		{
+			// Some setting has been applied to the ScanDir that requires a new scan
+			// Consume the invalidated flag and ignore the results of our scan
+			LocalCursor->SetScanInFlightInvalidated(false);
+		}
+		else
+		{
+			LocalCursor->SetScanResults(DirLocalAbsPath, LocalSubDirs, LocalDiscoveredFiles);
+			bValid = true;
+			FScanDir* NewCursor = Cursor.GetReference();
+			// Other thread may have set the cursor to a new spot; in that case do not update and on the next tick start at the new cursor
+			if (LocalCursor == NewCursor)
+			{
+				LocalCursor->Update(NewCursor, bCursorParentIsWhitelisted);
+				Cursor = NewCursor;
+			}
+		}
+		LocalCursor->SetScanInFlight(false);
+	}
+
+	if (bValid && (!LocalSubDirs.IsEmpty() || !LocalDiscoveredFiles.IsEmpty()))
+	{
+		AddDiscovered(LocalSubDirs, LocalDiscoveredFiles);
+	}
+}
+
+void FAssetDataDiscovery::FindFirstCursor(FScanDir*& OutCursor, bool& bOutCursorParentIsWhitelisted)
+{
+	CHECK_IS_LOCKED_CURRENT_THREAD(TreeLock);
+	OutCursor = nullptr;
+	while (!OutCursor)
+	{
+		EPriority WinningPriority = EPriority::Normal;
+		FMountDir* WinningMountDir = nullptr;
+		for (TUniquePtr<FMountDir>& MountDir : MountDirs)
+		{
+			if (MountDir->IsComplete())
+			{
+				continue;
+			}
+			if (WinningMountDir == nullptr || MountDir->GetPriority() < WinningPriority)
+			{
+				WinningMountDir = MountDir.Get();
+				WinningPriority = MountDir->GetPriority();
+			}
+		}
+
+		if (!WinningMountDir)
+		{
+			OutCursor = nullptr;
+			bOutCursorParentIsWhitelisted = false;
+			break;
+		}
+
+		WinningMountDir->Update(OutCursor, bOutCursorParentIsWhitelisted);
+		check(OutCursor != nullptr || WinningMountDir->IsComplete()); // The WinningMountDir's update should either return something to update or it should mark itself complete
+	}
+}
+
+void FAssetDataDiscovery::InvalidateCursor()
+{
+	if (Cursor)
+	{
+		if (Cursor->IsScanInFlight())
+		{
+			Cursor->SetScanInFlightInvalidated(true);
+		}
+		Cursor.SafeRelease();
+	}
+}
+
+void FAssetDataDiscovery::SetIsIdle(bool bInIsIdle)
+{
+	CHECK_IS_LOCKED_CURRENT_THREAD(TreeLock);
+
+	// Caller is responsible for holding TreeLock around this function; writes of SetIsIdle are done inside the TreeLock
+	// If bIsIdle is true, caller holds TickLock and TreeLock
+	if (bIsIdle == bInIsIdle)
+	{
+		return;
+	}
+	FGathererScopeLock ResultsScopeLock(&ResultsLock);
+	bIsIdle = bInIsIdle;
+	if (!bIsSynchronous)
+	{
+		if (bIsIdle)
+		{
+			UE_LOG(LogAssetRegistry, Verbose, TEXT("Discovery took %0.6f seconds and found %d files to process"), FPlatformTime::Seconds() - DiscoverStartTime, NumDiscoveredFiles);
+		}
+		else
+		{
+			DiscoverStartTime = FPlatformTime::Seconds();
+			NumDiscoveredFiles = 0;
+		}
+	}
+
+	if (bIsIdle)
+	{
+		CHECK_IS_LOCKED_CURRENT_THREAD(TickLock);
+		Shrink();
+	}
+}
+
+
+void FAssetDataDiscovery::GetAndTrimSearchResults(bool& bOutIsComplete, TArray<FString>& OutDiscoveredPaths, TRingBuffer<FGatheredPathData>& OutDiscoveredFiles, int32& OutNumPathsToSearch)
+{
+	FGathererScopeLock ResultsScopeLock(&ResultsLock);
+
+	OutDiscoveredPaths.Append(MoveTemp(DiscoveredDirectories));
+	DiscoveredDirectories.Reset();
+	OutDiscoveredFiles.Reserve(OutDiscoveredFiles.Num() + DiscoveredFiles.Num());
+	for (FGatheredPathData& DiscoveredFile : DiscoveredFiles)
+	{
+		OutDiscoveredFiles.Add(MoveTemp(DiscoveredFile));
+	}
+	DiscoveredFiles.Reset();
+
+	OutNumPathsToSearch = NumDirectoriesToScan.GetValue();
+	bOutIsComplete = bIsIdle;
+}
+
+void FAssetDataDiscovery::WaitForIdle()
+{
+	if (bIsIdle)
+	{
+		return;
+	}
+	FScopedPause ScopedPause(*this);
+	CHECK_IS_NOT_LOCKED_CURRENT_THREAD(TreeLock);
+	CHECK_IS_NOT_LOCKED_CURRENT_THREAD(ResultsLock);
+	FGathererScopeLock TickScopeLock(&TickLock);
+	while (!bIsIdle)
+	{
+		TickInternal();
+	}
+}
+
+void FAssetDataDiscovery::SetPropertiesAndWait(const FString& LocalAbsPath, bool bAddToWhitelist, bool bForceRescan, bool bIgnoreBlackListScanFilters)
+{
+	FFileStatData StatData = IFileManager::Get().GetStatData(*LocalAbsPath);
+	if (!StatData.bIsValid)
+	{
+		// SetPropertiesAndWait is called for every ScanPathsSynchronous, and this is the first spot that checks for existence. Some systems call ScanPathsSynchronous
+		// speculatively to scan whatever is present, so this is not a significant enough occurrence for a log.
+		UE_LOG(LogAssetRegistry, Verbose, TEXT("SetPropertiesAndWait called on non-existent path %.*s. Call will be ignored."), LocalAbsPath.Len(), *LocalAbsPath);
+		return;
+	}
+
+	{
+		CHECK_IS_NOT_LOCKED_CURRENT_THREAD(ResultsLock);
+		FGathererScopeLock TreeScopeLock(&TreeLock);
+		FMountDir* MountDir = FindContainingMountPoint(LocalAbsPath);
+		if (!MountDir)
+		{
+			UE_LOG(LogAssetRegistry, Log, TEXT("SetPropertiesAndWait called on %.*s which is not in a mounted directory. Call will be ignored."), LocalAbsPath.Len(), *LocalAbsPath);
+			return;
+		}
+		bool bIsBlacklisted = MountDir->IsBlacklisted(LocalAbsPath);
+		if (bIsBlacklisted && !bIgnoreBlackListScanFilters)
+		{
+			UE_LOG(LogAssetRegistry, Log, TEXT("SetPropertiesAndWait called on %.*s which is blacklisted. Call will be ignored."), LocalAbsPath.Len(), *LocalAbsPath);
+			return;
+		}
+
+		if (StatData.bIsDirectory)
+		{
+			FSetPathProperties Properties;
+			if (bAddToWhitelist)
+			{
+				Properties.IsWhitelisted = bAddToWhitelist;
+			}
+			if (bForceRescan)
+			{
+				Properties.HasScanned = false;
+			}
+			if (bIgnoreBlackListScanFilters)
+			{
+				Properties.IgnoreBlacklist = true;
+			}
+			if (Properties.IsSet())
+			{
+				SetIsIdle(false);
+				MountDir->TrySetDirectoryProperties(LocalAbsPath, Properties, true /* bConfirmedExists */);
+			}
+		}
+
+		FString RelPath;
+		bool bIsWhitelisted;
+		TRefCountPtr<FScanDir> ScanDir = MountDir->GetControllingDir(LocalAbsPath, StatData.bIsDirectory, bIsWhitelisted, RelPath);
+		if (!ScanDir || (!bIsWhitelisted && !bAddToWhitelist))
+		{
+			UE_LOG(LogAssetRegistry, Log, TEXT("SetPropertiesAndWait called on %.*s which is not whitelisted. Call will be ignored."), LocalAbsPath.Len(), *LocalAbsPath);
+			return;
+		}
+
+		if (StatData.bIsDirectory)
+		{
+			// If Relpath from the controlling dir to the requested dir is not empty then we have found a parent directory rather than the requested directory.
+			// This can only occur for a monitored directory when the requested directory is already complete and we do not need to wait on it.
+			if (RelPath.IsEmpty() && !ScanDir->IsComplete())
+			{
+				// We are going to wait on the path, so set its priority to blocking
+				SetIsIdle(false);
+				ScanDir->SetDirectPriority(EPriority::Blocking);
+				InvalidateCursor();
+
+				TreeScopeLock.Unlock(); // Entering the ticklock, as well as any long duration task such as a tick, has to be done outside of any locks
+
+				FScopedPause ScopedPause(*this);
+				CHECK_IS_NOT_LOCKED_CURRENT_THREAD(TreeLock);
+				CHECK_IS_NOT_LOCKED_CURRENT_THREAD(ResultsLock);
+				FGathererScopeLock TickScopeLock(&TickLock);
+				for (;;)
+				{
+					TickInternal();
+					FGathererScopeLock LoopTreeScopeLock(&TreeLock);
+					if (!ScanDir->IsValid() || ScanDir->IsComplete())
+					{
+						break;
+					}
+					else if (!ensureMsgf(!bIsIdle, TEXT("It should not be possible for the Discovery to go idle while there is an incomplete ScanDir.")))
+					{
+						break;
+					}
+				}
+			}
+		}
+		else
+		{
+			bool bAlreadyScanned = ScanDir->HasScanned() && !bIsBlacklisted;
+			if (!bAlreadyScanned || bForceRescan)
+			{
+				FStringView RelPathFromParentDir = FPathViews::GetCleanFilename(RelPath);
+				FStringView FileRelPathNoExt = FPathViews::GetBaseFilenameWithPath(RelPath);
+				if (!FPackageName::DoesPackageNameContainInvalidCharacters(FileRelPathNoExt) && FPackageName::IsPackageFilename(RelPathFromParentDir))
+				{
+					TStringBuilder<256> LongPackageName;
+					LongPackageName << MountDir->GetLongPackageName();
+					FPathViews::AppendPath(LongPackageName, ScanDir->GetMountRelPath());
+					FPathViews::AppendPath(LongPackageName, FileRelPathNoExt);
+					AddDiscovered(TConstArrayView<FDiscoveredPathData>(), TConstArrayView<FDiscoveredPathData> { FDiscoveredPathData(LocalAbsPath, LongPackageName, RelPathFromParentDir, StatData.ModificationTime) });
+					if (FPathViews::IsPathLeaf(RelPath) && !ScanDir->HasScanned())
+					{
+						SetIsIdle(false);
+						ScanDir->MarkFileAlreadyScanned(RelPath);
+					}
+				}
+			}
+		}
+	}
+}
+
+bool FAssetDataDiscovery::TrySetDirectoryProperties(const FString& LocalAbsPath, const FSetPathProperties& InProperties, bool bConfirmedExists)
+{
+	if (!InProperties.IsSet())
+	{
+		return false;
+	}
+	CHECK_IS_NOT_LOCKED_CURRENT_THREAD(ResultsLock);
+	FGathererScopeLock TreeScopeLock(&TreeLock);
+	if (!TrySetDirectoryPropertiesInternal(LocalAbsPath, InProperties, bConfirmedExists))
+	{
+		return false;
+	}
+	SetIsIdle(false);
+	InvalidateCursor();
+	return true;
+}
+
+bool FAssetDataDiscovery::TrySetDirectoryPropertiesInternal(const FString& LocalAbsPath, const FSetPathProperties& InProperties, bool bConfirmedExists)
+{
+	CHECK_IS_LOCKED_CURRENT_THREAD(TreeLock);
+	FMountDir* MountDir = FindContainingMountPoint(LocalAbsPath);
+	if (!MountDir)
+	{
+		UE_LOG(LogAssetRegistry, Warning, TEXT("FAssetDataGatherer::SetDirectoryProperties called on unmounted path %.*s. Call will be ignored."), LocalAbsPath.Len(), *LocalAbsPath);
+		return false;
+	}
+
+	return MountDir->TrySetDirectoryProperties(LocalAbsPath, InProperties, bConfirmedExists);
+}
+
+bool FAssetDataDiscovery::IsWhitelisted(FStringView LocalAbsPath) const
+{
+	CHECK_IS_NOT_LOCKED_CURRENT_THREAD(ResultsLock);
+	FGathererScopeLock TreeScopeLock(&TreeLock);
+	const FMountDir* MountDir = FindContainingMountPoint(LocalAbsPath);
+	return MountDir && MountDir->IsPathWhitelisted(LocalAbsPath);
+}
+
+bool FAssetDataDiscovery::IsBlacklisted(FStringView LocalAbsPath) const
+{
+	CHECK_IS_NOT_LOCKED_CURRENT_THREAD(ResultsLock);
+	FGathererScopeLock TreeScopeLock(&TreeLock);
+	const FMountDir* MountDir = FindContainingMountPoint(LocalAbsPath);
+	if (!MountDir)
+	{
+		// Not mounted, which we report as not blacklisted
+		return false;
+	}
+
+	return MountDir->IsBlacklisted(LocalAbsPath);
+}
+
+bool FAssetDataDiscovery::IsMonitored(FStringView LocalAbsPath) const
+{
+	CHECK_IS_NOT_LOCKED_CURRENT_THREAD(ResultsLock);
+	FGathererScopeLock TreeScopeLock(&TreeLock);
+	const FMountDir* MountDir = FindContainingMountPoint(LocalAbsPath);
+	return MountDir && MountDir->IsMonitored(LocalAbsPath);
+}
+
+uint32 FAssetDataDiscovery::GetAllocatedSize() const
+{
+	auto GetArrayRecursiveAllocatedSize = [](auto Container)
+	{
+		int32 Result = Container.GetAllocatedSize();
+		for (const auto& Value : Container)
+		{
+			Result += Value.GetAllocatedSize();
+		}
+		return Result;
+	};
 
-#if ALLOW_NAME_BATCH_SAVING
+	FScopedPause ScopedPause(*this);
+	CHECK_IS_NOT_LOCKED_CURRENT_THREAD(TreeLock);
+	CHECK_IS_NOT_LOCKED_CURRENT_THREAD(ResultsLock);
+	FGathererScopeLock TickScopeLock(&TickLock);
+	FGathererScopeLock TreeScopeLock(&TreeLock);
+	FGathererScopeLock ResultsScopeLock(&ResultsLock);
+
+	uint32 Result = 0;
+	Result += GetArrayRecursiveAllocatedSize(BlacklistLongPackageNames);
+	Result += GetArrayRecursiveAllocatedSize(BlacklistMountRelativePaths);
+	Result += GetArrayRecursiveAllocatedSize(DirLongPackageNamesToNotReport);
+	if (Thread)
+	{
+		// TODO: Thread->GetAllocatedSize()
+		Result += sizeof(*Thread);
+	}
+
+	Result += GetArrayRecursiveAllocatedSize(DiscoveredDirectories);
+	Result += GetArrayRecursiveAllocatedSize(DiscoveredFiles);
+
+	Result += MountDirs.GetAllocatedSize();
+	for (const TUniquePtr<FMountDir>& Value : MountDirs)
+	{
+		Result += sizeof(*Value);
+		Result += Value->GetAllocatedSize();
+	}
+	Result += GetArrayRecursiveAllocatedSize(IteratedSubDirs);
+	Result += GetArrayRecursiveAllocatedSize(IteratedFiles);
+	return Result;
+}
+
+void FAssetDataDiscovery::Shrink()
+{
+	CHECK_IS_LOCKED_CURRENT_THREAD(TickLock);
+	CHECK_IS_LOCKED_CURRENT_THREAD(TreeLock);
+	CHECK_IS_LOCKED_CURRENT_THREAD(ResultsLock);
+	DirLongPackageNamesToNotReport.Shrink();
+	DiscoveredDirectories.Shrink();
+	DiscoveredFiles.Shrink();
+	MountDirs.Shrink();
+	for (TUniquePtr<FMountDir>& MountDir : MountDirs)
+	{
+		MountDir->Shrink();
+	}
+	IteratedSubDirs.Shrink();
+	IteratedFiles.Shrink();
+}
+
+void FAssetDataDiscovery::AddMountPoint(const FString& LocalAbsPath, FStringView LongPackageName)
+{
+	CHECK_IS_NOT_LOCKED_CURRENT_THREAD(ResultsLock);
+	FGathererScopeLock TreeScopeLock(&TreeLock);
+	SetIsIdle(false);
+	AddMountPointInternal(LocalAbsPath, LongPackageName);
+
+	InvalidateCursor();
+}
+
+void FAssetDataDiscovery::AddMountPointInternal(const FString& LocalAbsPath, FStringView LongPackageName)
+{
+	CHECK_IS_LOCKED_CURRENT_THREAD(TreeLock);
+	TArray<FMountDir*> ChildMounts;
+	FMountDir* ParentMount = nullptr;
+	bool bExists = false;
+	for (TUniquePtr<FMountDir>& ExistingMount : MountDirs)
+	{
+		if (FPathViews::Equals(ExistingMount->GetLocalAbsPath(), LocalAbsPath))
+		{
+			bExists = true;
+			break;
+		}
+		else if (FPathViews::IsParentPathOf(ExistingMount->GetLocalAbsPath(), LocalAbsPath))
+		{
+			// Overwrite any earlier ParentMount; later mounts are more direct parents than earlier mounts
+			ParentMount = ExistingMount.Get();
+		}
+		else if (FPathViews::IsParentPathOf(LocalAbsPath, ExistingMount->GetLocalAbsPath()))
+		{
+			// A mount under the new directory might be a grandchild mount.
+			// Don't add it as a child mount unless there is no other mount in between the new mount and the mount.
+			FMountDir* ExistingParentMount = ExistingMount->GetParentMount();
+			if (!ExistingParentMount || ExistingParentMount == ParentMount)
+			{
+				ChildMounts.Add(ExistingMount.Get());
+			}
+		}
+	}
+	if (bExists)
+	{
+		return;
+	}
+
+	FMountDir& Mount = FindOrAddMountPoint(LocalAbsPath, LongPackageName);
+	if (ParentMount)
+	{
+		FStringView RelPath;
+		verify(FPathViews::TryMakeChildPathRelativeTo(LocalAbsPath, ParentMount->GetLocalAbsPath(), RelPath));
+		ParentMount->AddChildMount(&Mount);
+		for (FMountDir* ChildMount : ChildMounts)
+		{
+			ParentMount->RemoveChildMount(ChildMount);
+		}
+	}
+	for (FMountDir* ChildMount : ChildMounts)
+	{
+		Mount.AddChildMount(ChildMount);
+		ChildMount->SetParentMount(ParentMount);
+	}
+}
+
+void FAssetDataDiscovery::RemoveMountPoint(const FString& LocalAbsPath)
+{
+	CHECK_IS_NOT_LOCKED_CURRENT_THREAD(ResultsLock);
+	FGathererScopeLock TreeScopeLock(&TreeLock);
+	SetIsIdle(false);
+	RemoveMountPointInternal(LocalAbsPath);
+
+	InvalidateCursor();
+}
+
+void FAssetDataDiscovery::RemoveMountPointInternal(const FString& LocalAbsPath)
+{
+	CHECK_IS_LOCKED_CURRENT_THREAD(TreeLock);
+	int32 ExistingIndex = FindLowerBoundMountPoint(LocalAbsPath);
+	if (ExistingIndex == MountDirs.Num() || !FPathViews::Equals(MountDirs[ExistingIndex]->GetLocalAbsPath(), LocalAbsPath))
+	{
+		return;
+	}
+	TUniquePtr<FMountDir> Mount = MoveTemp(MountDirs[ExistingIndex]);
+	MountDirs.RemoveAt(ExistingIndex);
+	FMountDir* ParentMount = Mount->GetParentMount();
+
+	if (ParentMount)
+	{
+		for (FMountDir* ChildMount : Mount->GetChildMounts())
+		{
+			ParentMount->AddChildMount(ChildMount);
+			ChildMount->SetParentMount(ParentMount);
+		}
+		ParentMount->RemoveChildMount(Mount.Get());
+	}
+	else
+	{
+		for (FMountDir* ChildMount : Mount->GetChildMounts())
+		{
+			ChildMount->SetParentMount(nullptr);
+		}
+	}
+}
+
+void FAssetDataDiscovery::OnDirectoryCreated(FStringView LocalAbsPath)
+{
+	CHECK_IS_NOT_LOCKED_CURRENT_THREAD(ResultsLock);
+	FGathererScopeLock TreeScopeLock(&TreeLock);
+	FMountDir* MountDir = FindContainingMountPoint(LocalAbsPath);
+	if (!MountDir || !MountDir->IsMonitored(LocalAbsPath))
+	{
+		return;
+	}
+
+	FStringView MountRelPath;
+	verify(FPathViews::TryMakeChildPathRelativeTo(LocalAbsPath, MountDir->GetLocalAbsPath(), MountRelPath));
+	TStringBuilder<128> LongPackageName;
+	LongPackageName << MountDir->GetLongPackageName();
+	FPathViews::AppendPath(LongPackageName, MountRelPath);
+	if (FPackageName::DoesPackageNameContainInvalidCharacters(LongPackageName))
+	{
+		return;
+	}
+
+	// Skip reporting the directory if it is in the blacklists of directories to not report
+	if (!ShouldDirBeReported(LongPackageName))
+	{
+		return;
+	}
+
+	FDiscoveredPathData DirData;
+	DirData.LocalAbsPath = LocalAbsPath;
+	DirData.LongPackageName = LongPackageName;
+	DirData.RelPath = FPathViews::GetCleanFilename(MountRelPath);
+
+	// Note that we AddDiscovered but do not scan the directory
+	// Any files and paths under it will be added by their own event from the directory watcher, so a scan is unnecessary.
+	// The directory may also be scanned in the future because a parent directory is still yet pending to scan,
+	// we do not try to prevent that wasteful rescan because this is a rare event and it does not cause a behavior problem
+	AddDiscovered(TConstArrayView<FDiscoveredPathData>(&DirData, 1), TConstArrayView<FDiscoveredPathData>());
+	SetIsIdle(false);
+}
+
+void FAssetDataDiscovery::OnFilesCreated(TConstArrayView<FString> LocalAbsPaths)
+{
+	CHECK_IS_NOT_LOCKED_CURRENT_THREAD(ResultsLock);
+	FGathererScopeLock TreeScopeLock(&TreeLock);
+	SetIsIdle(false);
+	for (const FString& LocalAbsPath : LocalAbsPaths)
+	{
+		OnFileCreated(LocalAbsPath);
+	}
+}
+
+void FAssetDataDiscovery::OnFileCreated(const FString& LocalAbsPath)
+{
+	CHECK_IS_LOCKED_CURRENT_THREAD(TreeLock);
+	// Detect whether the file should be scanned and if so pass it through to the gatherer
+	FMountDir* MountDir = FindContainingMountPoint(LocalAbsPath);
+	if (!MountDir || MountDir->IsBlacklisted(LocalAbsPath))
+	{
+		// The content root of the file is not registered, or the file is blacklisted; ignore it
+		return;
+	}
+	FFileStatData StatData = IFileManager::Get().GetStatData(*LocalAbsPath);
+	if (!StatData.bIsValid || StatData.bIsDirectory)
+	{
+		// The caller has erroneously told us a file exists that doesn't exist (perhaps due to create/delete hysteresis); ignore it
+		return;
+	}
+
+	FString FileRelPath;
+	bool bIsWhitelisted;
+	FScanDir* ScanDir = MountDir->GetControllingDir(LocalAbsPath, false /* bIsDirectory */, bIsWhitelisted, FileRelPath);
+	if (!ScanDir || !bIsWhitelisted)
+	{
+		// The new file is in an unmonitored directory; ignore it
+		return;
+	}
+
+	FStringView RelPathFromParentDir = FPathViews::GetCleanFilename(FileRelPath);
+	FStringView FileRelPathNoExt = FPathViews::GetBaseFilenameWithPath(FileRelPath);
+	if (!FPackageName::DoesPackageNameContainInvalidCharacters(FileRelPathNoExt) && FPackageName::IsPackageFilename(RelPathFromParentDir))
+	{
+		TStringBuilder<256> LongPackageName;
+		LongPackageName << MountDir->GetLongPackageName();
+		FPathViews::AppendPath(LongPackageName, ScanDir->GetMountRelPath());
+		FPathViews::AppendPath(LongPackageName, FileRelPathNoExt);
+		AddDiscovered(TConstArrayView<FDiscoveredPathData>(), TConstArrayView<FDiscoveredPathData> { FDiscoveredPathData(LocalAbsPath, LongPackageName, RelPathFromParentDir, StatData.ModificationTime) });
+		if (FPathViews::IsPathLeaf(FileRelPath))
+		{
+			ScanDir->MarkFileAlreadyScanned(FileRelPath);
+		}
+	}
+}
+
+FMountDir* FAssetDataDiscovery::FindContainingMountPoint(FStringView LocalAbsPath)
+{
+	CHECK_IS_LOCKED_CURRENT_THREAD(TreeLock);
+	int32 Index = FindLowerBoundMountPoint(LocalAbsPath);
+	// The LowerBound is >= LocalAbsPath, so it is a parentpath of LocalAbsPath only if it is equal to LocalAbsPath
+	if (Index < MountDirs.Num() && FPathViews::Equals(MountDirs[Index]->GetLocalAbsPath(), LocalAbsPath))
+	{
+		return MountDirs[Index].Get();
+	}
+
+	// The last element before the lower bound is either (1) an unrelated path and LocalAbsPath does not have a parent
+	// (2) a parent path of LocalAbsPath, (3) A sibling path that is a child of an earlier path that is a parent path of LocalAbsPath
+	// (4) An unrelated path that is a child of an earlier path, but none of its parents are a parent path of LocalAbsPath
+	// Distinguishing between cases (3) and (4) doesn't have a fast algorithm based on sorted paths alone, but we have recorded the parent
+	// so we can figure it out that way
+	if (Index > 0)
+	{
+		FMountDir* Previous = MountDirs[Index - 1].Get();
+		while (Previous)
+		{
+			if (FPathViews::IsParentPathOf(Previous->GetLocalAbsPath(), LocalAbsPath))
+			{
+				return Previous;
+			}
+			Previous = Previous->GetParentMount();
+		}
+	}
+	return nullptr;
+}
+
+const FMountDir* FAssetDataDiscovery::FindContainingMountPoint(FStringView LocalAbsPath) const
+{
+	CHECK_IS_LOCKED_CURRENT_THREAD(TreeLock);
+	return const_cast<FAssetDataDiscovery*>(this)->FindContainingMountPoint(LocalAbsPath);
+}
+
+FMountDir* FAssetDataDiscovery::FindMountPoint(FStringView LocalAbsPath)
+{
+	CHECK_IS_LOCKED_CURRENT_THREAD(TreeLock);
+	int32 Index = FindLowerBoundMountPoint(LocalAbsPath);
+	if (Index != MountDirs.Num() && FPathViews::Equals(MountDirs[Index]->GetLocalAbsPath(), LocalAbsPath))
+	{
+		return MountDirs[Index].Get();
+	}
+	return nullptr;
+}
+
+FMountDir& FAssetDataDiscovery::FindOrAddMountPoint(FStringView LocalAbsPath, FStringView LongPackageName)
+{
+	CHECK_IS_LOCKED_CURRENT_THREAD(TreeLock);
+	int32 Index = FindLowerBoundMountPoint(LocalAbsPath);
+	if (Index != MountDirs.Num() && FPathViews::Equals(MountDirs[Index]->GetLocalAbsPath(), LocalAbsPath))
+	{
+		// Already exists
+		return *MountDirs[Index];
+	}
+	return *MountDirs.EmplaceAt_GetRef(Index, new FMountDir(*this, LocalAbsPath, LongPackageName));
+}
+
+int32 FAssetDataDiscovery::FindLowerBoundMountPoint(FStringView LocalAbsPath) const
+{
+	CHECK_IS_LOCKED_CURRENT_THREAD(TreeLock);
+	return Algo::LowerBound(MountDirs, LocalAbsPath, [](const TUniquePtr<FMountDir>& MountDir, FStringView LocalAbsPath)
+		{
+			return FPathViews::Less(MountDir->GetLocalAbsPath(), LocalAbsPath);
+		}
+	);
+}
+
+void FAssetDataDiscovery::AddDiscovered(TConstArrayView<FDiscoveredPathData> SubDirs, TConstArrayView<FDiscoveredPathData> Files)
+{
+	// This function is inside the critical section so we have moved filtering results outside of it
+	// Caller is responsible for filtering SubDirs and Files by ShouldScan and packagename validity
+	FGathererScopeLock ResultsScopeLock(&ResultsLock);
+	for (const FDiscoveredPathData& SubDir : SubDirs)
+	{
+		DiscoveredDirectories.Add(FString(SubDir.LongPackageName));
+	}
+	for (const FDiscoveredPathData& DiscoveredFile : Files)
+	{
+		DiscoveredFiles.Emplace(DiscoveredFile); // Emplace passes the FDiscoveredPathData to the FGatheredPathData explicit constructor for it
+	}
+	NumDiscoveredFiles += Files.Num();
+}
+
+bool FAssetDataDiscovery::ShouldDirBeReported(FStringView LongPackageName) const
+{
+	return !DirLongPackageNamesToNotReport.ContainsByHash(GetTypeHash(LongPackageName), LongPackageName);
+}
+
+} // namespace Private
+} // namespace AssetDataGather
+} // namespace UE
+
+FAssetDataGatherer::FAssetDataGatherer(const TArray<FString>& InBlacklistLongPackageNames, const TArray<FString>& InBlacklistMountRelativePaths, bool bInIsSynchronous)
+	: Thread(nullptr)
+	, bIsSynchronous(bInIsSynchronous)
+	, IsStopped(0)
+	, IsPaused(0)
+	, bInitialPluginsLoaded(false)
+	, bSaveAsyncCacheTriggered(false)
+	, SearchStartTime(0)
+	, LastCacheWriteTime(0.0)
+	, bUseMonolithicCache(false)
+	, bUseTickManagedCache(false)
+	, bHasLoadedTickManagedCache(false)
+	, bDiscoveryIsComplete(false)
+	, bIsComplete(false)
+	, bIsIdle(false)
+	, bFirstTickAfterIdle(true)
+	, bFinishedInitialDiscovery(false)
+	, WaitBatchCount(0)
+	, NumCachedFiles(0)
+	, NumUncachedFiles(0)
+	, bIsSavingAsyncCache(false)
+{
+	using namespace UE::AssetDataGather::Private;
+
+	bGatherDependsData = (GIsEditor && !FParse::Param(FCommandLine::Get(), TEXT("NoDependsGathering"))) || FParse::Param(FCommandLine::Get(),TEXT("ForceDependsGathering"));
+
+	bCacheEnabled = !FParse::Param(FCommandLine::Get(), TEXT("NoAssetRegistryCache")) && !FParse::Param(FCommandLine::Get(), TEXT("multiprocess"));
+
+#if !UE_BUILD_SHIPPING
+	bool bCommandlineSynchronous;
+	if (FParse::Bool(FCommandLine::Get(), TEXT("AssetGatherSync="), bCommandlineSynchronous))
+	{
+		bIsSynchronous = bCommandlineSynchronous;
+	}
+#endif
+	if (!bIsSynchronous && !FPlatformProcess::SupportsMultithreading())
+	{
+		bIsSynchronous = true;
+		UE_LOG(LogAssetRegistry, Warning, TEXT("Requested asynchronous asset data gather, but threading support is disabled. Performing a synchronous gather instead!"));
+	}
+	bIsSynchronousTick = bIsSynchronous;
+
+	Discovery = MakeUnique<UE::AssetDataGather::Private::FAssetDataDiscovery>(InBlacklistLongPackageNames, InBlacklistMountRelativePaths, bInIsSynchronous);
+}
+
+FAssetDataGatherer::~FAssetDataGatherer()
+{
+	EnsureCompletion();
+	NewCachedAssetDataMap.Empty();
+	DiskCachedAssetDataMap.Empty();
+
+	for (FDiskCachedAssetData* AssetData : NewCachedAssetData)
+	{
+		delete AssetData;
+	}
+	NewCachedAssetData.Empty();
+	for (TPair<int32, FDiskCachedAssetData*> BlockData : DiskCachedAssetBlocks)
+	{
+		delete[] BlockData.Get<1>();
+	}
+	DiskCachedAssetBlocks.Empty();
+}
+
+void FAssetDataGatherer::SetUseMonolithicCache(bool bInUseMonolithicCache)
+{
+	FGathererScopeLock ResultsScopeLock(&ResultsLock);
+	if (!bCacheEnabled || bUseMonolithicCache == bInUseMonolithicCache)
+	{
+		return;
+	}
+
+	bUseMonolithicCache = bInUseMonolithicCache;
+	bHasLoadedTickManagedCache = false;
+	LastCacheWriteTime = FPlatformTime::Seconds();
+	if (bUseMonolithicCache)
+	{
+		bUseTickManagedCache = bUseMonolithicCache;
+		TickManagedCacheFilename = FPaths::ProjectIntermediateDir() / (bGatherDependsData ? TEXT("CachedAssetRegistry.bin") : TEXT("CachedAssetRegistryNoDeps.bin"));
+	}
+	else
+	{
+		bUseTickManagedCache = false;
+		TickManagedCacheFilename.Empty();
+	}
+}
+
+void FAssetDataGatherer::StartAsync()
+{
+	if (!bIsSynchronous && !Thread)
+	{
+		Thread = FRunnableThread::Create(this, TEXT("FAssetDataGatherer"), 0, TPri_BelowNormal);
+		checkf(Thread, TEXT("Failed to create asset data gatherer thread"));
+		Discovery->StartAsync();
+	}
+}
+
+bool FAssetDataGatherer::Init()
+{
+	return true;
+}
+
+uint32 FAssetDataGatherer::Run()
+{
+	constexpr double IdleSleepTime = 0.1;
+	while (!IsStopped)
+	{
+		InnerTickLoop(false /* bInIsSynchronousTick */, true /* bContributeToCacheSave */);
+
+		for (;;)
+		{
+			{
+				FGathererScopeLock ResultsScopeLock(&ResultsLock); // bIsIdle requires the lock
+				if (IsStopped || bSaveAsyncCacheTriggered || (!IsPaused && !bIsIdle))
+				{
+					break;
+				}
+			}
+			// No work to do. Sleep for a little and try again later.
+			// TODO: Need IsPaused to be a condition variable so we avoid sleeping while waiting for it and then taking a long time to wake after it is unset.
+			FPlatformProcess::Sleep(IdleSleepTime);
+		}
+	}
+	return 0;
+}
+
+void FAssetDataGatherer::InnerTickLoop(bool bInIsSynchronousTick, bool bContributeToCacheSave)
+{
+	// Synchronous ticks during Wait contribute to saving of the async cache only if there is no dedicated async thread to do it (bIsSynchronous is true)
+	// The dedicated async thread always contributes
+	bContributeToCacheSave = !bInIsSynchronousTick || (bIsSynchronous && bContributeToCacheSave);
+
+	FString CacheFilename;
+	TArray<TPair<FName, FDiskCachedAssetData*>> AssetsToSave;
+	{
+		CHECK_IS_NOT_LOCKED_CURRENT_THREAD(ResultsLock);
+		FGathererScopeLock RunScopeLock(&TickLock);
+		TGuardValue<bool> ScopeSynchronousTick(bIsSynchronousTick, bInIsSynchronousTick);
+		TRACE_CPUPROFILER_EVENT_SCOPE(FAssetDataGatherer::Tick);
+		bool bTickInterruptionEvent = false;
+		while (!IsStopped && (bInIsSynchronousTick || !IsPaused) && !bTickInterruptionEvent)
+		{
+			TickInternal(bTickInterruptionEvent);
+		}
+
+		if (bContributeToCacheSave)
+		{
+			TryReserveSaveAsyncCache(CacheFilename, AssetsToSave);
+		}
+	}
+	SaveCacheFileInternal(CacheFilename, AssetsToSave, true /* bIsAsyncCache */);
+}
+
+FAssetDataGatherer::FScopedPause::FScopedPause(const FAssetDataGatherer& InOwner)
+	:Owner(InOwner)
+{
+	if (!Owner.bIsSynchronous)
+	{
+		Owner.IsPaused++;
+	}
+}
+
+FAssetDataGatherer::FScopedPause::~FScopedPause()
+{
+	if (!Owner.bIsSynchronous)
+	{
+		check(Owner.IsPaused > 0)
+		Owner.IsPaused--;
+	}
+}
+
+void FAssetDataGatherer::Stop()
+{
+	Discovery->Stop();
+	IsStopped++;
+}
+
+void FAssetDataGatherer::Exit()
+{
+}
+
+bool FAssetDataGatherer::IsSynchronous() const
+{
+	return bIsSynchronous;
+}
+
+void FAssetDataGatherer::EnsureCompletion()
+{
+	Discovery->EnsureCompletion();
+
+	Stop();
+
+	if (Thread != nullptr)
+	{
+		Thread->WaitForCompletion();
+		delete Thread;
+		Thread = nullptr;
+	}
+}
+
+void FAssetDataGatherer::TickInternal(bool& bOutIsTickInterrupt)
+{
+	using namespace UE::AssetDataGather::Private;
+
+	const int32 BatchSize = FTaskGraphInterface::Get().GetNumWorkerThreads() * AssetDataGathererConstants::SingleThreadFilesPerBatch;
+	check(BatchSize > 0); // GetNumWorkerThreads should never return 0
+	typedef TInlineAllocator<AssetDataGathererConstants::ExpectedMaxBatchSize> FBatchInlineAllocator;
+
+	TArray<FGatheredPathData, FBatchInlineAllocator> LocalFilesToSearch;
+	TArray<FAssetData*, FBatchInlineAllocator> LocalAssetResults;
+	TArray<FPackageDependencyData, FBatchInlineAllocator> LocalDependencyResults;
+	TArray<FString, FBatchInlineAllocator> LocalCookedPackageNamesWithoutAssetDataResults;
+	FString LocalCacheFilename;
+	double LocalLastCacheWriteTime = 0.0;
+	bool bWaitBatchCountDecremented = false;
+	bOutIsTickInterrupt = false;
+
+	{
+		FGathererScopeLock ResultsScopeLock(&ResultsLock);
+
+		if (bFirstTickAfterIdle)
+		{
+			bFirstTickAfterIdle = false;
+			LastCacheWriteTime = FPlatformTime::Seconds();
+			SearchStartTime = LastCacheWriteTime;
+		}
+
+		IngestDiscoveryResults();
+
+		// Take a batch off of the work list. If we're waiting only on the first WaitBatchCount results don't take more than that
+		int32 NumToProcess = FMath::Min<int32>(BatchSize-LocalFilesToSearch.Num(), FilesToSearch.Num());
+		if (WaitBatchCount > 0)
+		{
+			bWaitBatchCountDecremented = true;
+			NumToProcess = FMath::Min(NumToProcess, WaitBatchCount);
+			WaitBatchCount -= NumToProcess;
+			if (WaitBatchCount == 0)
 			{
-				// We might be able to reduce load time by using AssetRegistry::SerializationOptions
-				// to save certain common tags as FName.
-				FAssetRegistryWriter Ar(FAssetRegistryWriterOptions(), *FileAr);
-				SerializeCache(Ar);
+				bOutIsTickInterrupt = true;
 			}
-#else		
-		checkf(false, TEXT("Cannot save asset registry cache in this configuration"));
-#endif
+		}
 
-			// Close file handle before moving temp file to target 
-			FileAr.Reset();
-			IFileManager::Get().Move(*CacheFilename, *TempFilename);
+		while (NumToProcess > 0)
+		{
+			LocalFilesToSearch.Add(FilesToSearch.PopFrontValue());
+			--NumToProcess;
 		}
-		else
+
+		// If all work is finished mark idle and exit
+		if (LocalFilesToSearch.Num() == 0 && bDiscoveryIsComplete)
 		{
-			UE_LOG(LogAssetRegistry, Error, TEXT("Failed to open file for write %s"), *TempFilename);
+			WaitBatchCount = 0; // Clear WaitBatchCount in case it was set higher than FilesToSearch.Num()
+			bOutIsTickInterrupt = true;
+
+			if (!bFinishedInitialDiscovery)
+			{
+				bSaveAsyncCacheTriggered = true;
+			}
+			SetIsIdle(true);
+			return;
 		}
-	};
+		if (bUseTickManagedCache && !bHasLoadedTickManagedCache)
+		{
+			LocalCacheFilename = TickManagedCacheFilename;
+		}
+		LocalLastCacheWriteTime = LastCacheWriteTime;
+	}
 
-	while ( true )
+	// Load the async cache if not yet loaded
+	if (!LocalCacheFilename.IsEmpty())
 	{
-		bool LocalIsDiscoveringFiles = false;
+		LoadCacheFileInternal(LocalCacheFilename);
+		FGathererScopeLock ResultsScopeLock(&ResultsLock);
+		bHasLoadedTickManagedCache = true;
+	}
 
+	struct FReadContext
+	{
+		FName PackageName;
+		FName Extension;
+		const FGatheredPathData& AssetFileData;
+		TArray<FAssetData*> AssetDataFromFile;
+		FPackageDependencyData DependencyData;
+		TArray<FString> CookedPackageNamesWithoutAssetData;
+		bool bCanAttemptAssetRetry = false;
+		bool bResult = false;
+		bool bCanceled = false;
+
+		FReadContext(FName InPackageName, FName InExtension, const FGatheredPathData& InAssetFileData)
+			: PackageName(InPackageName)
+			, Extension(InExtension)
+			, AssetFileData(InAssetFileData)
 		{
-			FScopeLock CritSectionLock(&WorkerThreadCriticalSection);
+		}
+	};
 
-			AssetResults.Append(MoveTemp(LocalAssetResults));
-			DependencyResults.Append(MoveTemp(LocalDependencyResults));
-			CookedPackageNamesWithoutAssetDataResults.Append(MoveTemp(LocalCookedPackageNamesWithoutAssetDataResults));
-			LocalAssetResults.Reset();
-			LocalDependencyResults.Reset();
-			LocalCookedPackageNamesWithoutAssetDataResults.Reset();
+	// Try to read each file in the batch out of the cache, and accumulate a list for more expensive reading of all of the files that are not in the cache 
+	TArray<FReadContext> ReadContexts;
+	for (const FGatheredPathData& AssetFileData : LocalFilesToSearch)
+	{
+		const FName PackageName = *AssetFileData.LongPackageName;
+		const FName Extension = FName(*FPaths::GetExtension(AssetFileData.LocalAbsPath));
 
-			if (StopTaskCounter.GetValue() != 0)
+		FDiskCachedAssetData** DiskCachedAssetDataPtr = DiskCachedAssetDataMap.Find(PackageName);
+		FDiskCachedAssetData* DiskCachedAssetData = DiskCachedAssetDataPtr ? *DiskCachedAssetDataPtr : nullptr;
+		if (DiskCachedAssetData)
+		{
+			// Check whether we need to invalidate the cached data
+			const FDateTime& CachedTimestamp = DiskCachedAssetData->Timestamp;
+			if (AssetFileData.PackageTimestamp != CachedTimestamp)
 			{
-				break;
+				DiskCachedAssetData = nullptr;
 			}
-
-			// Grab any new package files from the background directory scan
-			if (BackgroundPackageFileDiscovery.IsValid())
+			else if ((DiskCachedAssetData->DependencyData.PackageName != PackageName && DiskCachedAssetData->DependencyData.PackageName != NAME_None) ||
+				DiskCachedAssetData->Extension != Extension)
 			{
-				bIsDiscoveringFiles = BackgroundPackageFileDiscovery->GetAndTrimSearchResults(DiscoveredPaths, FilesToSearch, NumPathsToSearchAtLastSyncPoint);
-				LocalIsDiscoveringFiles = bIsDiscoveringFiles;
+				UE_LOG(LogAssetRegistry, Display, TEXT("Cached dependency data for package '%s' is invalid. Discarding cached data."), *PackageName.ToString());
+				DiskCachedAssetData = nullptr;
 			}
+		}
 
+		if (DiskCachedAssetData)
+		{
+			// Add the valid cached data to our results, and to the map of data we keep to write out the new version of the cache file
+			++NumCachedFiles;
 
-			if (FilesToSearch.Num() > 0)
+			LocalAssetResults.Reserve(LocalAssetResults.Num() + DiskCachedAssetData->AssetDataList.Num());
+			for (const FAssetData& AssetData : DiskCachedAssetData->AssetDataList)
 			{
-				if (SearchStartTime == 0)
-				{
-					SearchStartTime = FPlatformTime::Seconds();
-				}
-
-				const int32 NumFilesToProcess = FMath::Min<int32>(AssetDataGathererConstants::MaxFilesToGatherBeforeFlush, FilesToSearch.Num());
-				LocalFilesToSearch.Append(FilesToSearch.GetData(), NumFilesToProcess);
-				FilesToSearch.RemoveAt(0, NumFilesToProcess, false);
+				LocalAssetResults.Add(new FAssetData(AssetData));
 			}
-			else if (SearchStartTime != 0 && !LocalIsDiscoveringFiles)
+
+			if (bGatherDependsData)
 			{
-				SearchTimes.Add(FPlatformTime::Seconds() - SearchStartTime);
-				SearchStartTime = 0;
+				LocalDependencyResults.Add(DiskCachedAssetData->DependencyData);
 			}
-		}
 
-		TArray<FDiscoveredPackageFile> LocalFilesToRetry;
+			AddToCache(PackageName, DiskCachedAssetData);
+		}
+		else
+		{
+			ReadContexts.Emplace(PackageName, Extension, AssetFileData);
+		}
+	}
 
-		if (LocalFilesToSearch.Num() > 0)
+	// For all the files not found in the cache, read them from their package files on disk; the file reads are done in parallel
+	ParallelFor(ReadContexts.Num(),
+		[this, &ReadContexts](int32 Index)
+		{
+			FReadContext& ReadContext = ReadContexts[Index];
+			if (!bIsSynchronousTick && IsPaused)
+			{
+				ReadContext.bCanceled = true;
+				return;
+			}
+			ReadContext.bResult = ReadAssetFile(ReadContext.AssetFileData.LocalAbsPath, ReadContext.AssetDataFromFile, ReadContext.DependencyData, ReadContext.CookedPackageNamesWithoutAssetData, ReadContext.bCanAttemptAssetRetry);
+		},
+		EParallelForFlags::Unbalanced | EParallelForFlags::BackgroundPriority
+	);
+
+	// Accumulate the results
+	bool bHasCancelation = false;
+	TArray<FGatheredPathData> LocalFilesToRetry;
+	for (FReadContext& ReadContext : ReadContexts)
+	{
+		if (ReadContext.bCanceled)
+		{
+			bHasCancelation = true;
+		}
+		else if (ReadContext.bResult)
 		{
-			struct FReadContext
+			++NumUncachedFiles;
+
+			// Add the results from a cooked package into our results on cooked package
+			LocalCookedPackageNamesWithoutAssetDataResults.Append(MoveTemp(ReadContext.CookedPackageNamesWithoutAssetData));
+			// Do not add the results from a cooked package into the map of data we keep to write out the new version of the cache file
+			bool bCachePackage = bCacheEnabled && LocalCookedPackageNamesWithoutAssetDataResults.Num() == 0;
+			if (bCachePackage)
 			{
-				FName PackageName;
-				FName Extension;
-				const FDiscoveredPackageFile& AssetFileData;
-				TArray<FAssetData*> AssetDataFromFile;
-				FPackageDependencyData DependencyData;
-				TArray<FString> CookedPackageNamesWithoutAssetData;
-				bool bCanAttemptAssetRetry = false;
-				bool bResult = false;
-
-				FReadContext(FName InPackageName, FName InExtension, const FDiscoveredPackageFile& InAssetFileData)
-					: PackageName(InPackageName)
-					, Extension(InExtension)
-					, AssetFileData(InAssetFileData)
+				for (const FAssetData* AssetData : ReadContext.AssetDataFromFile)
 				{
+					if (!!(AssetData->PackageFlags & PKG_FilterEditorOnly))
+					{
+						bCachePackage = false;
+						break;
+					}
 				}
-			};
+			}
 
-			TArray<FReadContext> ReadContexts;
-			FString LongPackageNameString;
-			for (const FDiscoveredPackageFile& AssetFileData : LocalFilesToSearch)
+			// Add the results from non-cooked packages into the map of data we keep to write out the new version of the cache file 
+			if (bCachePackage)
 			{
-				if (StopTaskCounter.GetValue() != 0)
+				// Update the cache
+				FDiskCachedAssetData* NewData = new FDiskCachedAssetData(ReadContext.AssetFileData.PackageTimestamp, ReadContext.Extension);
+				NewData->AssetDataList.Reserve(ReadContext.AssetDataFromFile.Num());
+				for (const FAssetData* BackgroundAssetData : ReadContext.AssetDataFromFile)
 				{
-					// We have been asked to stop, so don't read any more files
-					break;
+					NewData->AssetDataList.Add(*BackgroundAssetData);
 				}
 
-				LongPackageNameString.Reset();
-				if (!FPackageName::TryConvertFilenameToLongPackageName(AssetFileData.PackageFilename, LongPackageNameString))
+				// MoveTemp only used if we don't need DependencyData anymore
+				if (bGatherDependsData)
+				{
+					NewData->DependencyData = ReadContext.DependencyData;
+				}
+				else
 				{
-					// Conversion is expected to fail when the path has recently been unmounted, fail silent instead of fatal crash
-					continue;
+					NewData->DependencyData = MoveTemp(ReadContext.DependencyData);
 				}
 
-				const FName PackageName = *LongPackageNameString;
-				const FName Extension = FName(*FPaths::GetExtension(AssetFileData.PackageFilename));
+				NewCachedAssetData.Add(NewData);
+				AddToCache(ReadContext.PackageName, NewData);
+			}
+
+			// Add the results from the package into our output results
+			LocalAssetResults.Append(MoveTemp(ReadContext.AssetDataFromFile));
+			if (bGatherDependsData)
+			{
+				LocalDependencyResults.Add(MoveTemp(ReadContext.DependencyData));
+			}
+		}
+		else if (ReadContext.bCanAttemptAssetRetry)
+		{
+			// If the read temporarily failed, return it to the worklist, pushed to the end
+			FGathererScopeLock ResultsScopeLock(&ResultsLock);
+			FilesToSearch.Add(ReadContext.AssetFileData);
+		}
+	}
+
+	{
+		FGathererScopeLock ResultsScopeLock(&ResultsLock);
+
+		// Submit the results into the thread-shared lists
+		AssetResults.Append(MoveTemp(LocalAssetResults));
+		DependencyResults.Append(MoveTemp(LocalDependencyResults));
+		CookedPackageNamesWithoutAssetDataResults.Append(MoveTemp(LocalCookedPackageNamesWithoutAssetDataResults));
 
-				bool bLoadedFromCache = false;
-				if (bLoadAndSaveCache)
+		if (bHasCancelation)
+		{
+			// If we skipped reading files due to a pause request, push the canceled files back onto the FilesToSearch
+			for (int Index = ReadContexts.Num() - 1; Index >= 0; --Index) // AddToFront in reverse order so that the elements are readded in the same order they were popped
+			{
+				FReadContext& ReadContext = ReadContexts[Index];
+				if (ReadContext.bCanceled)
 				{
-					FDiskCachedAssetData* DiskCachedAssetData = DiskCachedAssetDataMap.Find(PackageName);
-					if (DiskCachedAssetData)
+					FilesToSearch.AddFront(ReadContext.AssetFileData);
+					if (bWaitBatchCountDecremented)
 					{
-						const FDateTime& CachedTimestamp = DiskCachedAssetData->Timestamp;
-						if (AssetFileData.PackageTimestamp != CachedTimestamp)
-						{
-							DiskCachedAssetData = nullptr;
-						}
-						else if ((DiskCachedAssetData->DependencyData.PackageName != PackageName && DiskCachedAssetData->DependencyData.PackageName != NAME_None) ||
-								DiskCachedAssetData->Extension != Extension)
-						{
-							UE_LOG(LogAssetRegistry, Display, TEXT("Cached dependency data for package '%s' is invalid. Discarding cached data."), *PackageName.ToString());
-							DiskCachedAssetData = nullptr;
-						}
+						++WaitBatchCount;
 					}
+				}
+			}
+		}
 
-					if (DiskCachedAssetData)
-					{
-						bLoadedFromCache = true;
+		if (bUseTickManagedCache && !bIsSavingAsyncCache 
+			&& FPlatformTime::Seconds() - LocalLastCacheWriteTime >= AssetDataGathererConstants::MinSecondsToElapseBeforeCacheWrite)
+		{
+			bSaveAsyncCacheTriggered = true;
+			bOutIsTickInterrupt = true;
+		}
+	}
+}
 
-						++NumCachedFiles;
+void FAssetDataGatherer::IngestDiscoveryResults()
+{
+	CHECK_IS_LOCKED_CURRENT_THREAD(TickLock);
+	CHECK_IS_LOCKED_CURRENT_THREAD(ResultsLock);
+	Discovery->GetAndTrimSearchResults(bDiscoveryIsComplete, DiscoveredPaths, FilesToSearch, NumPathsToSearchAtLastSyncPoint);
+}
 
-						LocalAssetResults.Reserve(LocalAssetResults.Num() + DiskCachedAssetData->AssetDataList.Num());
-						for (const FAssetData& AssetData : DiskCachedAssetData->AssetDataList)
-						{
-							LocalAssetResults.Add(new FAssetData(AssetData));
-						}
+bool FAssetDataGatherer::ReadAssetFile(const FString& AssetFilename, TArray<FAssetData*>& AssetDataList, FPackageDependencyData& DependencyData, TArray<FString>& CookedPackageNamesWithoutAssetData, bool& OutCanRetry) const
+{
+	TRACE_CPUPROFILER_EVENT_SCOPE(FAssetDataGatherer::ReadAssetFile);
+	OutCanRetry = false;
+	AssetDataList.Reset();
 
-						if (bGatherDependsData)
-						{
-							LocalDependencyResults.Add(DiskCachedAssetData->DependencyData);
-						}
+	FPackageReader PackageReader;
 
-						AddToCache(PackageName, DiskCachedAssetData);
-					}
-				}
+	FPackageReader::EOpenPackageResult OpenPackageResult;
+	if (!PackageReader.OpenPackageFile(AssetFilename, &OpenPackageResult))
+	{
+		// If we're missing a custom version, we might be able to load this package later once the module containing that version is loaded...
+		//   -	We can only attempt a retry in editors (not commandlets) that haven't yet finished initializing (!GIsRunning), as we 
+		//		have no guarantee that a commandlet or an initialized editor is going to load any more modules/plugins
+		//   -	Likewise, we can only attempt a retry for asynchronous scans, as during a synchronous scan we won't be loading any 
+		//		modules/plugins so it would last forever
+		const bool bAllowRetry = GIsEditor && !bInitialPluginsLoaded && !bIsSynchronousTick;
+		OutCanRetry = bAllowRetry && OpenPackageResult == FPackageReader::EOpenPackageResult::CustomVersionMissing;
+		return false;
+	}
+	else
+	{
+		return ReadAssetFile(PackageReader, AssetDataList, (bGatherDependsData ? &DependencyData : nullptr), CookedPackageNamesWithoutAssetData);
+	}
+}
 
-				if (!bLoadedFromCache)
-				{
-					ReadContexts.Emplace(PackageName, Extension, AssetFileData);
-				}
-			}
+bool FAssetDataGatherer::ReadAssetFile(FPackageReader& PackageReader, TArray<FAssetData*>& AssetDataList, FPackageDependencyData* DependencyData, TArray<FString>& CookedPackageNamesWithoutAssetData)
+{
+	if (PackageReader.ReadAssetRegistryDataIfCookedPackage(AssetDataList, CookedPackageNamesWithoutAssetData))
+	{
+		// Cooked data is special. No further data is found in these packages
+		return true;
+	}
 
-			ParallelFor(ReadContexts.Num(),
-				[this, &ReadContexts](int32 Index)
-				{
-					FReadContext& ReadContext = ReadContexts[Index];
-					ReadContext.bResult = ReadAssetFile(ReadContext.AssetFileData.PackageFilename, ReadContext.AssetDataFromFile, ReadContext.DependencyData, ReadContext.CookedPackageNamesWithoutAssetData, ReadContext.bCanAttemptAssetRetry);
-				},
-				EParallelForFlags::Unbalanced | EParallelForFlags::BackgroundPriority
-			);
+	if (!PackageReader.ReadAssetRegistryData(AssetDataList))
+	{
+		if (!PackageReader.ReadAssetDataFromThumbnailCache(AssetDataList))
+		{
+			// It's ok to keep reading even if the asset registry data doesn't exist yet
+			//return false;
+		}
+	}
+
+	if ( DependencyData )
+	{
+		if ( !PackageReader.ReadDependencyData(*DependencyData) )
+		{
+			return false;
+		}
 
-			for (FReadContext& ReadContext : ReadContexts)
+		// DEPRECATION_TODO: Remove this fixup-on-load once we bump EUnrealEngineObjectUE4Version VER_UE4_ASSETREGISTRY_DEPENDENCYFLAGS and therefore all projects will resave each ObjectRedirector
+		// UObjectRedirectors were originally incorrectly marked as having editor-only imports, since UObjectRedirector is an editor-only class. But UObjectRedirectors are followed during cooking
+		// and so their imports should be considered used-in-game. Mark all dependencies in the package as used in game if the package has a UObjectRedirector object
+		FName RedirectorClassName = UObjectRedirector::StaticClass()->GetFName();
+		if (Algo::AnyOf(AssetDataList, [RedirectorClassName](FAssetData* AssetData) { return AssetData->AssetClass == RedirectorClassName; }))
+		{
+			TBitArray<>& ImportUsedInGame = DependencyData->ImportUsedInGame;
+			for (int32 ImportNum = ImportUsedInGame.Num(), Index = 0; Index < ImportNum; ++Index)
 			{
-				if (ReadContext.bResult)
-				{
-					++NumUncachedFiles;
+				ImportUsedInGame[Index] = true;
+			}
+		}
+		// END DEPRECATION_TODO
+	}
 
-					LocalCookedPackageNamesWithoutAssetDataResults.Append(MoveTemp(ReadContext.CookedPackageNamesWithoutAssetData));
+	return true;
+}
 
-					// Don't store info on cooked packages
-					bool bCachePackage = bLoadAndSaveCache && LocalCookedPackageNamesWithoutAssetDataResults.Num() == 0;
-					if (bCachePackage)
-					{
-						for (const FAssetData* AssetData : ReadContext.AssetDataFromFile)
-						{
-							if (!!(AssetData->PackageFlags & PKG_FilterEditorOnly))
-							{
-								bCachePackage = false;
-								break;
-							}
-						}
-					}
+void FAssetDataGatherer::AddToCache(FName PackageName, FDiskCachedAssetData* DiskCachedAssetData)
+{
+	CHECK_IS_LOCKED_CURRENT_THREAD(TickLock);
+	FDiskCachedAssetData*& ValueInMap = NewCachedAssetDataMap.FindOrAdd(PackageName, DiskCachedAssetData);
+	if (ValueInMap != DiskCachedAssetData)
+	{
+		// An updated DiskCachedAssetData for the same package; replace the existing DiskCachedAssetData with the new one.
+		// Note that memory management of the DiskCachedAssetData is handled in a separate structure; we do not need to delete the old value here.
+		if (DiskCachedAssetData->Extension != ValueInMap->Extension)
+		{
+			// Two files with the same package name but different extensions, e.g. basename.umap and basename.uasset
+			// This is invalid - some systems in the engine (Cooker's FPackageNameCache) assume that package : filename is 1 : 1 - so issue a warning
+			// Because it is invalid, we don't fully support it here (our map is keyed only by packagename), and will remove from cache all but the last filename we find with the same packagename
+			// TODO: Turn this into a warning once all sample projects have fixed it
+			UE_LOG(LogAssetRegistry, Display, TEXT("Multiple files exist with the same package name %s but different extensions (%s and %s). ")
+				TEXT("This is invalid and will cause errors; merge or rename or delete one of the files."),
+				*PackageName.ToString(), *ValueInMap->Extension.ToString(), *DiskCachedAssetData->Extension.ToString());
+		}
+		ValueInMap = DiskCachedAssetData;
+	}
+}
 
-					if (bCachePackage)
-					{
-						// Update the cache
-						FDiskCachedAssetData* NewData = new FDiskCachedAssetData(ReadContext.AssetFileData.PackageTimestamp, ReadContext.Extension);
-						NewData->AssetDataList.Reserve(ReadContext.AssetDataFromFile.Num());
-						for (const FAssetData* BackgroundAssetData : ReadContext.AssetDataFromFile)
-						{
-							NewData->AssetDataList.Add(*BackgroundAssetData);
-						}
+void FAssetDataGatherer::GetAndTrimSearchResults(bool& bOutIsSearching, TRingBuffer<FAssetData*>& OutAssetResults, TRingBuffer<FString>& OutPathResults, TRingBuffer<FPackageDependencyData>& OutDependencyResults,
+	TRingBuffer<FString>& OutCookedPackageNamesWithoutAssetDataResults, TArray<double>& OutSearchTimes, int32& OutNumFilesToSearch, int32& OutNumPathsToSearch, bool& OutIsDiscoveringFiles)
+{
+	FGathererScopeLock ResultsScopeLock(&ResultsLock);
 
-						// MoveTemp only used if we don't need DependencyData anymore
-						if (bGatherDependsData)
-						{
-							NewData->DependencyData = ReadContext.DependencyData;
-						}
-						else
-						{
-							NewData->DependencyData = MoveTemp(ReadContext.DependencyData);
-						}
+	auto MoveAppendRangeToRingBuffer = [](auto& OutRingBuffer, auto& InArray)
+	{
+		OutRingBuffer.Reserve(OutRingBuffer.Num() + InArray.Num());
+		for (auto& Element : InArray)
+		{
+			OutRingBuffer.Add(MoveTemp(Element));
+		}
+		InArray.Reset();
+	};
 
-						NewCachedAssetData.Add(NewData);
-						AddToCache(ReadContext.PackageName, NewData);
-					}
+	MoveAppendRangeToRingBuffer(OutAssetResults, AssetResults);
+	MoveAppendRangeToRingBuffer(OutPathResults, DiscoveredPaths);
+	MoveAppendRangeToRingBuffer(OutDependencyResults, DependencyResults);
+	MoveAppendRangeToRingBuffer(OutCookedPackageNamesWithoutAssetDataResults, CookedPackageNamesWithoutAssetDataResults);
 
-					LocalAssetResults.Append(MoveTemp(ReadContext.AssetDataFromFile));
-					if (bGatherDependsData)
-					{
-						LocalDependencyResults.Add(MoveTemp(ReadContext.DependencyData));
-					}
-				}
-				else if (ReadContext.bCanAttemptAssetRetry)
-				{
-					LocalFilesToRetry.Add(ReadContext.AssetFileData);
-				}
-			}
+	OutSearchTimes.Append(MoveTemp(SearchTimes));
+	SearchTimes.Reset();
 
-			LocalFilesToSearch.Reset();
-			LocalFilesToSearch.Append(LocalFilesToRetry);
-			LocalFilesToRetry.Reset();
+	OutNumFilesToSearch = FilesToSearch.Num();
+	OutNumPathsToSearch = NumPathsToSearchAtLastSyncPoint;
+	OutIsDiscoveringFiles = !bDiscoveryIsComplete;
 
-			if (bLoadAndSaveCache)
-			{
-				// Only write intermediate state cache file if we have spent a good amount of time working on it
-				if (FPlatformTime::Seconds() - LastCacheWriteTime >= AssetDataGathererConstants::MinSecondsToElapseBeforeCacheWrite)
-				{
-					WriteAssetCacheFile();
-					LastCacheWriteTime = FPlatformTime::Seconds();
-				}
-			}
+	if (bIsIdle && !bIsComplete)
+	{
+		bIsComplete = true;
+		Shrink();
+	}
+	bOutIsSearching = !bIsIdle;
+}
+
+void FAssetDataGatherer::WaitOnPath(FStringView InPath)
+{
+	{
+		FGathererScopeLock ResultsScopeLock(&ResultsLock);
+		if (bIsIdle)
+		{
+			return;
 		}
-		else
+	}
+	FString LocalAbsPath = NormalizeLocalPath(InPath);
+	Discovery->SetPropertiesAndWait(LocalAbsPath, false /* bAddToWhitelist */, false /* bForceRescan */, false /* bIgnoreBlacklistScanFilters */);
+	WaitOnPathsInternal(TConstArrayView<FString>(&LocalAbsPath, 1), FString(), TArray<FString>());
+}
+
+void FAssetDataGatherer::ScanPathsSynchronous(const TArray<FString>& InLocalPaths, bool bForceRescan, bool bIgnoreBlackListScanFilters, const FString& SaveCacheFilename, const TArray<FString>& SaveCacheLongPackageNameDirs)
+{
+	TArray<FString> LocalAbsPaths;
+	LocalAbsPaths.Reserve(InLocalPaths.Num());
+	for (const FString& LocalPath : InLocalPaths)
+	{
+		LocalAbsPaths.Add(NormalizeLocalPath(LocalPath));
+	}
+
+	for (const FString& LocalAbsPath : LocalAbsPaths)
+	{
+		Discovery->SetPropertiesAndWait(LocalAbsPath, true /* bAddToWhitelist */, bForceRescan, bIgnoreBlackListScanFilters);
+	}
+
+	{
+		FGathererScopeLock ResultsScopeLock(&ResultsLock);
+		SetIsIdle(false);
+	}
+
+	WaitOnPathsInternal(LocalAbsPaths, SaveCacheFilename, SaveCacheLongPackageNameDirs);
+}
+
+void FAssetDataGatherer::WaitOnPathsInternal(TConstArrayView<FString> LocalAbsPaths, const FString& SaveCacheFilename, const TArray<FString>& SaveCacheLongPackageNameDirs)
+{
+	// Request a halt to the async tick
+	FScopedPause ScopedPause(*this);
+	CHECK_IS_NOT_LOCKED_CURRENT_THREAD(ResultsLock);
+	{
+		FGathererScopeLock TickScopeLock(&TickLock);
+
+		// Read all results from Discovery into our worklist and then sort our worklist
 		{
-			if (bIsSynchronous)
+			FGathererScopeLock ResultsScopeLock(&ResultsLock);
+			IngestDiscoveryResults();
+
+			int32 NumDiscoveredPaths;
+			SortPathsByPriority(LocalAbsPaths, NumDiscoveredPaths);
+			if (NumDiscoveredPaths == 0)
 			{
-				// This is synchronous. Since our work is done, we should safely exit
-				Stop();
+				return;
 			}
-			else
-			{
-				if (!LocalIsDiscoveringFiles && !bFinishedInitialDiscovery)
-				{
-					bFinishedInitialDiscovery = true;
+			WaitBatchCount = NumDiscoveredPaths;
+		}
+	}
 
-					UE_LOG(LogAssetRegistry, Verbose, TEXT("Initial scan took %0.6f seconds (found %d cached assets, and loaded %d)"), FPlatformTime::Seconds() - InitialScanStartTime, NumCachedFiles, NumUncachedFiles);
+	// We do not contribute to the async cache save if we have been given a modular cache to save 
+	bool bContributeToCacheSave = SaveCacheFilename.IsEmpty();
+
+	// Tick until NumDiscoveredPaths have been read
+	TRACE_CPUPROFILER_EVENT_SCOPE(FAssetDataGatherer::Tick);
+	for (;;)
+	{
+		InnerTickLoop(true /* bInIsSynchronousTick */, bContributeToCacheSave);
+		FGathererScopeLock ResultsScopeLock(&ResultsLock); // WaitBatchCount requires the lock
+		if (WaitBatchCount == 0)
+		{
+			break;
+		}
+	}
+
+	if (!SaveCacheFilename.IsEmpty())
+	{
+		TArray<TPair<FName, FDiskCachedAssetData*>> AssetsToSave;
+		{
+			FGathererScopeLock TickScopeLock(&TickLock);
+			GetAssetsToSave(SaveCacheLongPackageNameDirs, AssetsToSave);
+		}
+		SaveCacheFileInternal(SaveCacheFilename, AssetsToSave, false /* bIsAsyncCacheSave */);
+	}
+}
+
+void FAssetDataGatherer::WaitForIdle()
+{
+	{
+		FGathererScopeLock ResultsScopeLock(&ResultsLock);
+		if (bIsIdle)
+		{
+			return;
+		}
+	}
+	Discovery->WaitForIdle();
+	CHECK_IS_NOT_LOCKED_CURRENT_THREAD(ResultsLock);
+
+	// Request a halt to the async tick
+	FScopedPause ScopedPause(*this);
+	// Tick until idle
+	for (;;)
+	{
+		InnerTickLoop(true /* bInIsSynchronousTick */, true /* bContributeToCacheSave */);
+		FGathererScopeLock ResultsScopeLock(&ResultsLock); // bIsIdle requires the lock
+		if (bIsIdle)
+		{
+			break;
+		}
+	}
+}
+
+bool FAssetDataGatherer::IsComplete() const
+{
+	FGathererScopeLock ResultsScopeLock(&ResultsLock);
+	return bIsComplete;
+}
+
+void FAssetDataGatherer::SetInitialPluginsLoaded()
+{
+	bInitialPluginsLoaded = true;
+}
+
+bool FAssetDataGatherer::IsGatheringDependencies() const
+{
+	return bGatherDependsData;
+}
 
-					// If we are caching discovered assets and this is the first time we had no work to do, save off the cache now in case the user terminates unexpectedly
-					if (bLoadAndSaveCache)
-					{
-						WriteAssetCacheFile();
-					}
-				}
+bool FAssetDataGatherer::IsCacheEnabled() const
+{
+	return bCacheEnabled;
+}
 
-				// No work to do. Sleep for a little and try again later.
-				FPlatformProcess::Sleep(0.1);
-			}
+FString FAssetDataGatherer::GetCacheFilename(TConstArrayView<FString> CacheFilePackagePaths)
+{
+	// Try and build a consistent hash for this input
+	// Normalize the paths; removing any trailing /
+	TArray<FString> SortedPaths(CacheFilePackagePaths);
+	for (FString& PackagePath : SortedPaths)
+	{
+		while (PackagePath.Len() > 1 && PackagePath.EndsWith(TEXT("/")))
+		{
+			PackagePath.LeftChopInline(1);
 		}
 	}
 
-	check(LocalAssetResults.Num() == 0); // All LocalAssetResults needed to be copied to AssetResults to avoid leaking memory
+	// Sort the paths
+	SortedPaths.StableSort();
 
-	if ( bLoadAndSaveCache )
+	// todo: handle hash collisions?
+	uint32 CacheHash = SortedPaths.Num() > 0 ? GetTypeHash(SortedPaths[0]) : 0;
+	for (int32 PathIndex = 1; PathIndex < SortedPaths.Num(); ++PathIndex)
 	{
-		WriteAssetCacheFile();
+		CacheHash = HashCombine(CacheHash, GetTypeHash(SortedPaths[PathIndex]));
 	}
 
-	return 0;
+	return FPaths::ProjectIntermediateDir() / TEXT("AssetRegistryCache") / FString::Printf(TEXT("%08x%s.bin"), CacheHash, bGatherDependsData ? TEXT("") : TEXT("NoDeps"));
 }
 
-void FAssetDataGatherer::AddToCache(FName PackageName, FDiskCachedAssetData* DiskCachedAssetData)
+void FAssetDataGatherer::LoadCacheFile(FStringView CacheFilename)
 {
-	FDiskCachedAssetData*& ValueInMap = NewCachedAssetDataMap.FindOrAdd(PackageName, DiskCachedAssetData);
-	if (ValueInMap != DiskCachedAssetData)
+	FScopedPause ScopedPause(*this);
+	CHECK_IS_NOT_LOCKED_CURRENT_THREAD(ResultsLock);
+	FGathererScopeLock TickScopeLock(&TickLock);
+	LoadCacheFileInternal(CacheFilename);
+}
+
+void FAssetDataGatherer::LoadCacheFileInternal(FStringView CacheFilename)
+{
+	CHECK_IS_LOCKED_CURRENT_THREAD(TickLock);
+	if (!bCacheEnabled)
 	{
-		// An updated DiskCachedAssetData for the same package; replace the existing DiskCachedAssetData with the new one.
-		// Note that memory management of the DiskCachedAssetData is handled in a separate structure; we do not need to delete the old value here.
-		if (DiskCachedAssetData->Extension != ValueInMap->Extension)
+		return;
+	}
+
+	TRACE_CPUPROFILER_EVENT_SCOPE(LoadCacheFile);
+	// load the cached data
+	TUniquePtr<FArchive> FileAr(IFileManager::Get().CreateFileReader(*FString(CacheFilename), FILEREAD_Silent));
+	if (FileAr && !FileAr->IsError() && FileAr->TotalSize() > 2 * sizeof(uint32))
+	{
+		uint32 MagicNumber = 0;
+		*FileAr << MagicNumber;
+		if (!FileAr->IsError() && MagicNumber == AssetDataGathererConstants::CacheSerializationMagic)
 		{
-			// Two files with the same package name but different extensions, e.g. basename.umap and basename.uasset
-			// This is invalid - some systems in the engine (Cooker's FPackageNameCache) assume that package : filename is 1 : 1 - so issue a warning
-			// Because it is invalid, we don't fully support it here (our map is keyed only by packagename), and will remove from cache all but the last filename we find with the same packagename
-			// TODO: Turn this into a warning once all sample projects have fixed it
-			UE_LOG(LogAssetRegistry, Display, TEXT("Multiple files exist with the same package name %s but different extensions (%s and %s). ")
-				TEXT("This is invalid and will cause errors; merge or rename or delete one of the files."),
-				*PackageName.ToString(), *ValueInMap->Extension.ToString(), *DiskCachedAssetData->Extension.ToString());
+			FAssetRegistryVersion::Type RegistryVersion;
+			if (FAssetRegistryVersion::SerializeVersion(*FileAr, RegistryVersion) && RegistryVersion == FAssetRegistryVersion::LatestVersion)
+			{
+				FAssetRegistryReader RegistryReader(*FileAr);
+				if (!RegistryReader.IsError())
+				{
+					SerializeCacheLoad(RegistryReader);
+
+					FGathererScopeLock ResultsScopeLock(&ResultsLock);
+					DependencyResults.Reserve(DiskCachedAssetDataMap.Num());
+					AssetResults.Reserve(DiskCachedAssetDataMap.Num());
+				}
+			}
 		}
-		ValueInMap = DiskCachedAssetData;
 	}
 }
 
-void FAssetDataGatherer::Stop()
+void FAssetDataGatherer::TryReserveSaveAsyncCache(FString& OutCacheFilename, TArray<TPair<FName,FDiskCachedAssetData*>>& AssetsToSave)
 {
-	if (BackgroundPackageFileDiscovery.IsValid())
+	OutCacheFilename.Reset();
+	if (IsStopped)
 	{
-		BackgroundPackageFileDiscovery->Stop();
+		return;
 	}
-
-	StopTaskCounter.Increment();
+	if (!bSaveAsyncCacheTriggered || bIsSavingAsyncCache)
+	{
+		return;
+	}
+	CHECK_IS_LOCKED_CURRENT_THREAD(TickLock);
+	{
+		FGathererScopeLock ResultsScopeLock(&ResultsLock);
+		if (bUseTickManagedCache && !TickManagedCacheFilename.IsEmpty())
+		{
+			OutCacheFilename = TickManagedCacheFilename;
+		}
+	}
+	if (!OutCacheFilename.IsEmpty())
+	{
+		GetAssetsToSave(TArrayView<const FString>(), AssetsToSave);
+		bIsSavingAsyncCache = true;
+	}
+	bSaveAsyncCacheTriggered = false;
 }
 
-void FAssetDataGatherer::Exit()
-{   
+void FAssetDataGatherer::GetAssetsToSave(TArrayView<const FString> SaveCacheLongPackageNameDirs, TArray<TPair<FName,FDiskCachedAssetData*>>& OutAssetsToSave)
+{
+	CHECK_IS_LOCKED_CURRENT_THREAD(TickLock);
+
+	OutAssetsToSave.Reset();
+	if (SaveCacheLongPackageNameDirs.IsEmpty())
+	{
+		OutAssetsToSave.Reserve(NewCachedAssetDataMap.Num());
+		for (const TPair<FName, FDiskCachedAssetData*>& Pair : NewCachedAssetDataMap)
+		{
+			OutAssetsToSave.Add(Pair);
+		}
+	}
+	else
+	{
+		for (const TPair<FName, FDiskCachedAssetData*>& Pair : NewCachedAssetDataMap)
+		{
+			TStringBuilder<128> PackageNameStr;
+			Pair.Key.ToString(PackageNameStr);
+			if (Algo::AnyOf(SaveCacheLongPackageNameDirs, [PackageNameSV = FStringView(PackageNameStr)](const FString& SaveCacheLongPackageNameDir)
+			{
+				return FPathViews::IsParentPathOf(SaveCacheLongPackageNameDir, PackageNameSV);
+			}))
+			{
+				OutAssetsToSave.Add(Pair);
+			}
+		}
+	}
 }
 
-void FAssetDataGatherer::EnsureCompletion()
+void FAssetDataGatherer::SaveCacheFileInternal(const FString& CacheFilename, const TArray<TPair<FName,FDiskCachedAssetData*>>& AssetsToSave, bool bIsAsyncCacheSave)
 {
-	if (BackgroundPackageFileDiscovery.IsValid())
+	if (CacheFilename.IsEmpty() || !bCacheEnabled)
 	{
-		BackgroundPackageFileDiscovery->EnsureCompletion();
+		return;
 	}
+	CHECK_IS_NOT_LOCKED_CURRENT_THREAD(TickLock);
+	CHECK_IS_NOT_LOCKED_CURRENT_THREAD(ResultsLock);
+	TRACE_CPUPROFILER_EVENT_SCOPE(SaveCacheFile);
+	// Save to a temp file first, then move to the destination to avoid corruption
+	FString CacheFilenameStr(CacheFilename);
+	FString TempFilename = CacheFilenameStr + TEXT(".tmp");
+	TUniquePtr<FArchive> FileAr(IFileManager::Get().CreateFileWriter(*TempFilename, 0));
+	if (FileAr)
+	{
+		int32 MagicNumber = AssetDataGathererConstants::CacheSerializationMagic;
+		*FileAr << MagicNumber;
 
+		FAssetRegistryVersion::Type RegistryVersion = FAssetRegistryVersion::LatestVersion;
+		FAssetRegistryVersion::SerializeVersion(*FileAr, RegistryVersion);
+#if ALLOW_NAME_BATCH_SAVING
+		{
+			// We might be able to reduce load time by using AssetRegistry::SerializationOptions
+			// to save certain common tags as FName.
+			FAssetRegistryWriter Ar(FAssetRegistryWriterOptions(), *FileAr);
+			SerializeCacheSave(Ar, AssetsToSave);
+		}
+#else		
+		checkf(false, TEXT("Cannot save asset registry cache in this configuration"));
+#endif
+		// Close file handle before moving temp file to target 
+		FileAr.Reset();
+		IFileManager::Get().Move(*CacheFilenameStr, *TempFilename);
+	}
+	else
 	{
-		FScopeLock CritSectionLock(&WorkerThreadCriticalSection);
-		FilesToSearch.Empty();
+		UE_LOG(LogAssetRegistry, Error, TEXT("Failed to open file for write %s"), *TempFilename);
 	}
 
-	Stop();
-
-	if (Thread != nullptr)
+	if (bIsAsyncCacheSave)
 	{
-		Thread->WaitForCompletion();
-		delete Thread;
-		Thread = nullptr;
+		FScopedPause ScopedPause(*this);
+		FGathererScopeLock TickScopeLock(&TickLock);
+		bIsSavingAsyncCache = false;
+		LastCacheWriteTime = FPlatformTime::Seconds();
 	}
 }
 
-bool FAssetDataGatherer::GetAndTrimSearchResults(TBackgroundGatherResults<FAssetData*>& OutAssetResults, TBackgroundGatherResults<FString>& OutPathResults, TBackgroundGatherResults<FPackageDependencyData>& OutDependencyResults, TBackgroundGatherResults<FString>& OutCookedPackageNamesWithoutAssetDataResults, TArray<double>& OutSearchTimes, int32& OutNumFilesToSearch, int32& OutNumPathsToSearch, bool& OutIsDiscoveringFiles)
+void FAssetDataGatherer::SerializeCacheSave(FAssetRegistryWriter& Ar, const TArray<TPair<FName,FDiskCachedAssetData*>>& AssetsToSave)
 {
-	FScopeLock CritSectionLock(&WorkerThreadCriticalSection);
-
-	OutAssetResults.Append(MoveTemp(AssetResults));
-	AssetResults.Reset();
-
-	OutPathResults.Append(MoveTemp(DiscoveredPaths));
-	DiscoveredPaths.Reset();
-
-	OutDependencyResults.Append(MoveTemp(DependencyResults));
-	DependencyResults.Reset();
+#if ALLOW_NAME_BATCH_SAVING
+	using namespace UE::AssetDataGather::Private;
 
-	OutCookedPackageNamesWithoutAssetDataResults.Append(MoveTemp(CookedPackageNamesWithoutAssetDataResults));
-	CookedPackageNamesWithoutAssetDataResults.Reset();
+	double SerializeStartTime = FPlatformTime::Seconds();
 
-	OutSearchTimes.Append(MoveTemp(SearchTimes));
-	SearchTimes.Reset();
+	// serialize number of objects
+	int32 LocalNumAssets = AssetsToSave.Num();
+	Ar << LocalNumAssets;
 
-	OutNumFilesToSearch = FilesToSearch.Num();
-	OutNumPathsToSearch = NumPathsToSearchAtLastSyncPoint;
-	OutIsDiscoveringFiles = bIsDiscoveringFiles;
+	for (const TPair<FName, FDiskCachedAssetData*>& Pair : AssetsToSave)
+	{
+		FName AssetName = Pair.Key;
+		Ar << AssetName;
+		Pair.Value->SerializeForCache(Ar);
+	}
 
-	return (SearchStartTime > 0 || bIsDiscoveringFiles);
+	UE_LOG(LogAssetRegistry, Verbose, TEXT("Asset data gatherer serialized in %0.6f seconds"), FPlatformTime::Seconds() - SerializeStartTime);
+#endif
 }
 
-void FAssetDataGatherer::AddPathToSearch(const FString& Path)
+void FAssetDataGatherer::SerializeCacheLoad(FAssetRegistryReader& Ar)
 {
-	if (BackgroundPackageFileDiscovery.IsValid())
+	CHECK_IS_LOCKED_CURRENT_THREAD(TickLock);
+	using namespace UE::AssetDataGather::Private;
+
+	double SerializeStartTime = FPlatformTime::Seconds();
+	// serialize number of objects
+	int32 LocalNumAssets = 0;
+	Ar << LocalNumAssets;
+
+	const int32 MinAssetEntrySize = sizeof(int32);
+	int32 MaxPossibleNumAssets = (Ar.TotalSize() - Ar.Tell()) / MinAssetEntrySize;
+	if (Ar.IsError() || LocalNumAssets < 0 || MaxPossibleNumAssets < LocalNumAssets)
+	{
+		Ar.SetError();
+	}
+	else
 	{
-		BackgroundPackageFileDiscovery->AddPathToSearch(Path);
+		FSoftObjectPathSerializationScope SerializationScope(NAME_None, NAME_None, ESoftObjectPathCollectType::NeverCollect, ESoftObjectPathSerializeType::AlwaysSerialize);
+
+		// allocate one single block for all asset data structs (to reduce tens of thousands of heap allocations)
+		FName* PackageNameBlock = new FName[LocalNumAssets];
+		FDiskCachedAssetData* AssetDataBlock = new FDiskCachedAssetData[LocalNumAssets];
+		for (int32 AssetIndex = 0; AssetIndex < LocalNumAssets; ++AssetIndex)
+		{
+			// Load the name first to add the entry to the tmap below
+			Ar << PackageNameBlock[AssetIndex];
+			AssetDataBlock[AssetIndex].SerializeForCache(Ar);
+			if (Ar.IsError())
+			{
+				// There was an error reading the cache. Bail out.
+				break;
+			}
+		}
+
+		if (Ar.IsError())
+		{
+			delete[] AssetDataBlock;
+			UE_LOG(LogAssetRegistry, Error, TEXT("There was an error loading the asset registry cache."));
+		}
+		else
+		{
+			DiskCachedAssetDataMap.Reserve(DiskCachedAssetDataMap.Num() + LocalNumAssets);
+			for (int32 AssetIndex = 0; AssetIndex < LocalNumAssets; ++AssetIndex)
+			{
+				DiskCachedAssetDataMap.Add(PackageNameBlock[AssetIndex], &AssetDataBlock[AssetIndex]);
+			}
+			DiskCachedAssetBlocks.Emplace(LocalNumAssets, AssetDataBlock);
+		}
+		delete[] PackageNameBlock;
 	}
+
+	UE_LOG(LogAssetRegistry, Verbose, TEXT("Asset data gatherer serialized in %0.6f seconds"), FPlatformTime::Seconds() - SerializeStartTime);
 }
 
-void FAssetDataGatherer::AddFilesToSearch(const TArray<FString>& Files)
+uint32 FAssetDataGatherer::GetAllocatedSize() const
 {
-	TArray<FString> FilesToAdd;
-	FilesToAdd.Reserve(Files.Num());
-	FString LongPackageName;
-	for (const FString& Filename : Files)
+	using namespace UE::AssetDataGather::Private;
+	auto GetArrayRecursiveAllocatedSize = [](auto Container)
 	{
-		LongPackageName.Reset();
-		if (ConvertToValidLongPackageName(Filename, LongPackageName)
-			&& AssetDataDiscoveryUtil::PassesScanFilters(BlacklistScanFilters, LongPackageName))
+		int32 Result = Container.GetAllocatedSize();
+		for (const auto& Value : Container)
 		{
-			// Add the path to this asset into the list of discovered paths
-			FilesToAdd.Add(Filename);
+			Result += Value.GetAllocatedSize();
 		}
-	}
+		return Result;
+	};
 
-	if (FilesToAdd.Num() > 0)
+	uint32 Result = 0;
+	if (Thread)
 	{
-		FScopeLock CritSectionLock(&WorkerThreadCriticalSection);
-		FilesToSearch.Append(MoveTemp(FilesToAdd));
+		// TODO: Add size of Thread->GetAllocatedSize()
+		Result += sizeof(*Thread);
 	}
-}
 
-void FAssetDataGatherer::PrioritizeSearchPath(const FString& PathToPrioritize)
-{
-	if (BackgroundPackageFileDiscovery.IsValid())
+	Result += sizeof(*Discovery) + Discovery->GetAllocatedSize();
+
+	FScopedPause ScopedPause(*this);
+	CHECK_IS_NOT_LOCKED_CURRENT_THREAD(ResultsLock);
+	FGathererScopeLock TickScopeLock(&TickLock);
+	FGathererScopeLock ResultsScopeLock(&ResultsLock);
+
+	Result += GetArrayRecursiveAllocatedSize(FilesToSearch);
+
+	Result += AssetResults.GetAllocatedSize();
+	FAssetDataTagMapSharedView::FMemoryCounter TagMemoryUsage;
+	for (FAssetData* Value : AssetResults)
 	{
-		BackgroundPackageFileDiscovery->PrioritizeSearchPath(PathToPrioritize);
+		Result += sizeof(*Value);
+		Result += Value->ChunkIDs.GetAllocatedSize();
+		TagMemoryUsage.Include(Value->TagsAndValues);
 	}
+	Result += TagMemoryUsage.GetFixedSize() + TagMemoryUsage.GetLooseSize();
 
-	FString LocalFilenamePathToPrioritize;
-	if (FPackageName::TryConvertLongPackageNameToFilename(PathToPrioritize / TEXT(""), LocalFilenamePathToPrioritize))
-	{
-		FScopeLock CritSectionLock(&WorkerThreadCriticalSection);
+	Result += GetArrayRecursiveAllocatedSize(DependencyResults);
+	Result += GetArrayRecursiveAllocatedSize(CookedPackageNamesWithoutAssetDataResults);
+	Result += SearchTimes.GetAllocatedSize();
+	Result += GetArrayRecursiveAllocatedSize(DiscoveredPaths);
+	Result += TickManagedCacheFilename.GetAllocatedSize();
 
-		FilenamePathToPrioritize = LocalFilenamePathToPrioritize;
-		SortPathsByPriority(INDEX_NONE);
+	Result += NewCachedAssetData.GetAllocatedSize();
+	for (const FDiskCachedAssetData* Value : NewCachedAssetData)
+	{
+		Result += sizeof(*Value);
+		Result += Value->GetAllocatedSize();
 	}
-}
-
-void FAssetDataGatherer::SetBlacklistScanFilters(const TArray<FString>& InBlacklistScanFilters)
-{
+	Result += DiskCachedAssetBlocks.GetAllocatedSize();
+	for (const TPair<int32, FDiskCachedAssetData*>& ArrayData : DiskCachedAssetBlocks)
 	{
-		FScopeLock CritSectionLock(&WorkerThreadCriticalSection);
-		BlacklistScanFilters = InBlacklistScanFilters;
+		Result += ArrayData.Get<0>() * sizeof(FDiskCachedAssetData);
 	}
+	Result += DiskCachedAssetDataMap.GetAllocatedSize();
+	Result += NewCachedAssetDataMap.GetAllocatedSize();
+
+	return Result;
+}
 
-	if (BackgroundPackageFileDiscovery.IsValid())
+void FAssetDataGatherer::Shrink()
+{
+	CHECK_IS_LOCKED_CURRENT_THREAD(ResultsLock);
+	// TODO: Make TRingBuffer::Shrink
+	TRingBuffer<UE::AssetDataGather::Private::FGatheredPathData> Buffer;
+	Buffer.Reserve(FilesToSearch.Num());
+	for (UE::AssetDataGather::Private::FGatheredPathData& File : FilesToSearch)
 	{
-		BackgroundPackageFileDiscovery->SetBlacklistScanFilters(InBlacklistScanFilters);
+		Buffer.Add(MoveTemp(File));
 	}
+	Swap(Buffer, FilesToSearch);
+	AssetResults.Shrink();
+	DependencyResults.Shrink();
+	CookedPackageNamesWithoutAssetDataResults.Shrink();
+	SearchTimes.Shrink();
+	DiscoveredPaths.Shrink();
 }
 
-bool FAssetDataGatherer::IsGatheringDependencies() const
+void FAssetDataGatherer::AddMountPoint(FStringView LocalPath, FStringView LongPackageName)
 {
-	return bGatherDependsData;
+	Discovery->AddMountPoint(NormalizeLocalPath(LocalPath), NormalizeLongPackageName(LongPackageName));
+	{
+		FGathererScopeLock ResultsScopeLock(&ResultsLock);
+		SetIsIdle(false);
+	}
 }
 
-void FAssetDataGatherer::SetInitialPluginsLoaded()
+void FAssetDataGatherer::RemoveMountPoint(FStringView LocalPath)
 {
-	bInitialPluginsLoaded.store(true);
+	Discovery->RemoveMountPoint(NormalizeLocalPath(LocalPath));
 }
 
-void FAssetDataGatherer::SortPathsByPriority(const int32 MaxNumToSort)
+void FAssetDataGatherer::AddRequiredMountPoints(TArrayView<FString> LocalPaths)
 {
-	FScopeLock CritSectionLock(&WorkerThreadCriticalSection);
-
-	// Critical section. This code needs to be as fast as possible since it is in a critical section!
-	// Swap all priority files to the top of the list
-	if (FilenamePathToPrioritize.Len() > 0)
+	TStringBuilder<128> MountPackageName;
+	TStringBuilder<128> MountFilePath;
+	TStringBuilder<128> RelPath;
+	for (const FString& LocalPath : LocalPaths)
 	{
-		int32 LowestNonPriorityFileIdx = 0;
-		int32 NumSorted = 0;
-		const int32 NumToSort = (MaxNumToSort == INDEX_NONE) ? FilesToSearch.Num() : FMath::Min(FilesToSearch.Num(), MaxNumToSort);
-		for (int32 FilenameIdx = 0; FilenameIdx < FilesToSearch.Num(); ++FilenameIdx)
+		if (FPackageName::TryGetMountPointForPath(LocalPath, MountPackageName, MountFilePath, RelPath))
 		{
-			if (FilesToSearch[FilenameIdx].PackageFilename.StartsWith(FilenamePathToPrioritize))
-			{
-				FilesToSearch.Swap(FilenameIdx, LowestNonPriorityFileIdx);
-				LowestNonPriorityFileIdx++;
-
-				if (++NumSorted >= NumToSort)
-				{
-					break;
-				}
-			}
+			Discovery->AddMountPoint(NormalizeLocalPath(MountFilePath), NormalizeLongPackageName(MountPackageName));
 		}
 	}
+	{
+		FGathererScopeLock ResultsScopeLock(&ResultsLock);
+		SetIsIdle(false);
+	}
 }
 
-bool FAssetDataGatherer::ReadAssetFile(const FString& AssetFilename, TArray<FAssetData*>& AssetDataList, FPackageDependencyData& DependencyData, TArray<FString>& CookedPackageNamesWithoutAssetData, bool& OutCanRetry) const
+void FAssetDataGatherer::OnDirectoryCreated(FStringView LocalPath)
 {
-	TRACE_CPUPROFILER_EVENT_SCOPE(FAssetDataGatherer::ReadAssetFile)
-	OutCanRetry = false;
-	AssetDataList.Reset();
-
-	FPackageReader PackageReader;
+	Discovery->OnDirectoryCreated(NormalizeLocalPath(LocalPath));
+	{
+		FGathererScopeLock ResultsScopeLock(&ResultsLock);
+		SetIsIdle(false);
+	}
+}
 
-	FPackageReader::EOpenPackageResult OpenPackageResult;
-	if ( !PackageReader.OpenPackageFile(AssetFilename, &OpenPackageResult) )
+void FAssetDataGatherer::OnFilesCreated(TConstArrayView<FString> LocalPaths)
+{
+	TArray<FString> LocalAbsPaths;
+	LocalAbsPaths.Reserve(LocalPaths.Num());
+	for (const FString& LocalPath : LocalPaths)
 	{
-		// If we're missing a custom version, we might be able to load this package later once the module containing that version is loaded...
-		//   -	We can only attempt a retry in editors (not commandlets) that haven't yet finished initializing (!GIsRunning), as we 
-		//		have no guarantee that a commandlet or an initialized editor is going to load any more modules/plugins
-		//   -	Likewise, we can only attempt a retry for asynchronous scans, as during a synchronous scan we won't be loading any 
-		//		modules/plugins so it would last forever
-		const bool bAllowRetry = GIsEditor && !IsRunningCommandlet() && !bInitialPluginsLoaded.load() && !bIsSynchronous;
-		OutCanRetry = bAllowRetry && OpenPackageResult == FPackageReader::EOpenPackageResult::CustomVersionMissing;
-		return false;
+		LocalAbsPaths.Add(NormalizeLocalPath(LocalPath));
 	}
-	else
+	Discovery->OnFilesCreated(LocalAbsPaths);
 	{
-		return ReadAssetFile(PackageReader, AssetDataList, (bGatherDependsData ? &DependencyData : nullptr), CookedPackageNamesWithoutAssetData);
+		FGathererScopeLock ResultsScopeLock(&ResultsLock);
+		SetIsIdle(false);
 	}
 }
 
-bool FAssetDataGatherer::ReadAssetFile(FPackageReader& PackageReader, TArray<FAssetData*>& AssetDataList, FPackageDependencyData* DependencyData, TArray<FString>& CookedPackageNamesWithoutAssetData)
+void FAssetDataGatherer::PrioritizeSearchPath(const FString& PathToPrioritize)
 {
-	if ( PackageReader.ReadAssetRegistryDataIfCookedPackage(AssetDataList, CookedPackageNamesWithoutAssetData) )
+	using namespace UE::AssetDataGather::Private;
+
+	FString LocalFilenamePathToPrioritize;
+	if (FPackageName::TryConvertLongPackageNameToFilename(PathToPrioritize, LocalFilenamePathToPrioritize))
 	{
-		// Cooked data is special. No further data is found in these packages
-		return true;
+		FSetPathProperties Properties;
+		Properties.Priority = EPriority::High;
+		SetDirectoryProperties(LocalFilenamePathToPrioritize, Properties);
 	}
+}
 
-	if ( !PackageReader.ReadAssetRegistryData(AssetDataList) )
+void FAssetDataGatherer::SetDirectoryProperties(FStringView LocalPath, const UE::AssetDataGather::Private::FSetPathProperties& InProperties)
+{
+	FString LocalAbsPath = NormalizeLocalPath(LocalPath);
+	if (LocalAbsPath.Len() == 0)
 	{
-		if ( !PackageReader.ReadAssetDataFromThumbnailCache(AssetDataList) )
-		{
-			// It's ok to keep reading even if the asset registry data doesn't exist yet
-			//return false;
-		}
+		return;
 	}
 
-	if ( DependencyData )
+	if (!Discovery->TrySetDirectoryProperties(LocalAbsPath, InProperties, false))
 	{
-		if ( !PackageReader.ReadDependencyData(*DependencyData) )
-		{
-			return false;
-		}
+		return;
+	}
 
-		// DEPRECATION_TODO: Remove this fixup-on-load once we bump EUnrealEngineObjectUE4Version VER_UE4_ASSETREGISTRY_DEPENDENCYFLAGS and therefore all projects will resave each ObjectRedirector
-		// UObjectRedirectors were originally incorrectly marked as having editor-only imports, since UObjectRedirector is an editor-only class. But UObjectRedirectors are followed during cooking
-		// and so their imports should be considered used-in-game. Mark all dependencies in the package as used in game if the package has a UObjectRedirector object
-		FName RedirectorClassName = UObjectRedirector::StaticClass()->GetFName();
-		if (Algo::AnyOf(AssetDataList, [RedirectorClassName](FAssetData* AssetData) { return AssetData->AssetClass == RedirectorClassName; }))
+	{
+		FGathererScopeLock ResultsScopeLock(&ResultsLock);
+		SetIsIdle(false);
+		if (InProperties.Priority.IsSet())
 		{
-			TBitArray<>& ImportUsedInGame = DependencyData->ImportUsedInGame;
-			for (int32 ImportNum = ImportUsedInGame.Num(), Index = 0; Index < ImportNum; ++Index)
-			{
-				ImportUsedInGame[Index] = true;
-			}
+			int32 NumPrioritizedPaths;
+			SortPathsByPriority(TConstArrayView<FString>(&LocalAbsPath, 1), NumPrioritizedPaths);
 		}
-		// END DEPRECATION_TODO
 	}
-
-	return true;
 }
 
-template<class Archive>
-void FAssetDataGatherer::SerializeCache(Archive&& Ar)
+void FAssetDataGatherer::SortPathsByPriority(TConstArrayView<FString> LocalAbsPathsToPrioritize, int32& OutNumPaths)
 {
-	double SerializeStartTime = FPlatformTime::Seconds();
+	CHECK_IS_LOCKED_CURRENT_THREAD(ResultsLock);
+	using namespace UE::AssetDataGather::Private;
 
-	// serialize number of objects
-	int32 LocalNumAssets = NewCachedAssetDataMap.Num();
-	Ar << LocalNumAssets;
+	// This code needs to be as fast as possible since it is in a critical section!
 
-	if (Ar.IsSaving())
+	// Swap all priority files to the top of the list
+	int32 LowestNonPriorityFileIdx = 0;
+	int32 NumFilesToSearch = FilesToSearch.Num();
+	for (int32 FilenameIdx = 0; FilenameIdx < NumFilesToSearch; ++FilenameIdx)
 	{
-		// save out by walking the TMap
-		for (TPair<FName, FDiskCachedAssetData*>& Pair : NewCachedAssetDataMap)
+		for (const FString& LocalAbsPathToPrioritize : LocalAbsPathsToPrioritize)
 		{
-			Ar << Pair.Key;
-			Pair.Value->SerializeForCache(Ar);
+			if (FPathViews::IsParentPathOf(LocalAbsPathToPrioritize, FilesToSearch[FilenameIdx].LocalAbsPath))
+			{
+				Swap(FilesToSearch[FilenameIdx], FilesToSearch[LowestNonPriorityFileIdx++]);
+				break;
+			}
 		}
 	}
-	else
-	{
-		// allocate one single block for all asset data structs (to reduce tens of thousands of heap allocations)
-		if (Ar.IsError() || LocalNumAssets < 0)
-		{
-			Ar.SetError();
-		}
-		else
-		{
-			FSoftObjectPathSerializationScope SerializationScope(NAME_None, NAME_None, ESoftObjectPathCollectType::NeverCollect, ESoftObjectPathSerializeType::AlwaysSerialize);
+	OutNumPaths = LowestNonPriorityFileIdx;
+}
 
-			const int32 MinAssetEntrySize = sizeof(int32);
-			int32 MaxReservation = (Ar.TotalSize() - Ar.Tell()) / MinAssetEntrySize;
-			DiskCachedAssetDataMap.Empty(FMath::Min(LocalNumAssets, MaxReservation));
+void FAssetDataGatherer::SetIsWhitelisted(FStringView LocalPath, bool bIsWhitelisted)
+{
+	using namespace UE::AssetDataGather::Private;
 
-			for (int32 AssetIndex = 0; AssetIndex < LocalNumAssets; ++AssetIndex)
-			{
-				// Load the name first to add the entry to the tmap below
-				FName PackageName;
-				Ar << PackageName;
-				if (Ar.IsError())
-				{
-					// There was an error reading the cache. Bail out.
-					break;
-				}
+	FSetPathProperties Properties;
+	Properties.IsWhitelisted = bIsWhitelisted;
+	SetDirectoryProperties(LocalPath, Properties);
+}
 
-				// Add to the cached map
-				FDiskCachedAssetData& CachedAssetData = DiskCachedAssetDataMap.Add(PackageName);
+bool FAssetDataGatherer::IsWhitelisted(FStringView LocalPath) const
+{
+	return Discovery->IsWhitelisted(NormalizeLocalPath(LocalPath));
+}
 
-				// Now load the data
-				CachedAssetData.SerializeForCache(Ar);
+bool FAssetDataGatherer::IsBlacklisted(FStringView LocalPath) const
+{
+	return Discovery->IsBlacklisted(NormalizeLocalPath(LocalPath));
+}
 
-				if (Ar.IsError())
-				{
-					// There was an error reading the cache. Bail out.
-					break;
-				}
-			}
-		}
+bool FAssetDataGatherer::IsMonitored(FStringView LocalPath) const
+{
+	return Discovery->IsMonitored(NormalizeLocalPath(LocalPath));
+}
 
-		// If there was an error loading the cache, abandon all data loaded from it so we can build a clean one.
-		if (Ar.IsError())
+void FAssetDataGatherer::SetIsIdle(bool bInIsIdle)
+{
+	CHECK_IS_LOCKED_CURRENT_THREAD(ResultsLock);
+	if (bInIsIdle == bIsIdle)
+	{
+		return;
+	}
+
+	bIsIdle = bInIsIdle;
+	if (bIsIdle)
+	{
+		// bIsComplete will be set in GetAndTrimSearchResults
+		double SearchTime = FPlatformTime::Seconds() - SearchStartTime;
+		if (!bFinishedInitialDiscovery)
 		{
-			UE_LOG(LogAssetRegistry, Error, TEXT("There was an error loading the asset registry cache. Generating a new one."));
-			DiskCachedAssetDataMap.Empty();
+			bFinishedInitialDiscovery = true;
+
+			UE_LOG(LogAssetRegistry, Verbose, TEXT("Initial scan took %0.6f seconds (found %d cached assets, and loaded %d)"), (float)SearchTime, NumCachedFiles, NumUncachedFiles);
 		}
+		SearchTimes.Add(SearchTime);
 	}
+	else
+	{
+		bIsComplete = false;
+		bDiscoveryIsComplete = false;
+		bFirstTickAfterIdle = true;
+	}
+}
 
-	UE_LOG(LogAssetRegistry, Verbose, TEXT("Asset data gatherer serialized in %0.6f seconds"), FPlatformTime::Seconds() - SerializeStartTime);
+FString FAssetDataGatherer::NormalizeLocalPath(FStringView LocalPath)
+{
+	FString LocalAbsPath(LocalPath);
+	LocalAbsPath = FPaths::ConvertRelativePathToFull(MoveTemp(LocalAbsPath));
+	return LocalAbsPath;
+}
+
+FStringView FAssetDataGatherer::NormalizeLongPackageName(FStringView LongPackageName)
+{
+	// Conform LongPackageName to our internal format, which does not have a terminating redundant /
+	if (LongPackageName.EndsWith(TEXT('/')))
+	{
+		LongPackageName = LongPackageName.LeftChop(1);
+	}
+	return LongPackageName;
 }
+
diff --git a/Engine/Source/Runtime/AssetRegistry/Private/AssetDataGatherer.h b/Engine/Source/Runtime/AssetRegistry/Private/AssetDataGatherer.h
--- a/Engine/Source/Runtime/AssetRegistry/Private/AssetDataGatherer.h
+++ b/Engine/Source/Runtime/AssetRegistry/Private/AssetDataGatherer.h
@@ -2,175 +2,77 @@
 
 #pragma once
 
-#include "CoreMinimal.h"
-#include "HAL/ThreadSafeCounter.h"
-#include "PackageDependencyData.h"
+#include <atomic>
+#include "Containers/Array.h"
+#include "Containers/ArrayView.h"
+#include "Containers/Map.h"
+#include "Containers/RingBuffer.h"
+#include "Containers/StringFwd.h"
+#include "Containers/StringView.h"
+#include "Containers/UnrealString.h"
+#include "CriticalSectionQueryable.h"
+#include "DiskCachedAssetData.h"
+#include "HAL/CriticalSection.h"
 #include "HAL/FileManager.h"
+#include "HAL/Platform.h"
 #include "HAL/Runnable.h"
-#include "DiskCachedAssetData.h"
-#include "BackgroundGatherResults.h"
-#include <atomic>
-
+#include "Misc/DateTime.h"
+#include "Misc/Optional.h"
+#include "PackageDependencyData.h"
+#include "UObject/NameTypes.h"
+#include "Templates/UniquePtr.h"
+
+class FArchive;
+struct FAssetData;
+class FDiskCachedAssetData;
+class FAssetRegistryReader;
+class FAssetRegistryWriter; // Not defined if !ALLOW_NAME_BATCH_SAVING
 class FPackageReader;
-
-/**
- * Minimal amount of information needed about a discovered asset file
- */
-struct FDiscoveredPackageFile
+namespace UE
 {
-	explicit FDiscoveredPackageFile(FString InPackageFilename)
-		: PackageFilename(MoveTemp(InPackageFilename))
-		, PackageTimestamp(IFileManager::Get().GetTimeStamp(*PackageFilename))
-	{
-	}
-
-	FDiscoveredPackageFile(FString InPackageFilename, FDateTime InPackageTimestamp)
-		: PackageFilename(MoveTemp(InPackageFilename))
-		, PackageTimestamp(MoveTemp(InPackageTimestamp))
-	{
-	}
-
-	FDiscoveredPackageFile(const FDiscoveredPackageFile& Other)
-		: PackageFilename(Other.PackageFilename)
-		, PackageTimestamp(Other.PackageTimestamp)
-	{
-	}
-
-	FDiscoveredPackageFile(FDiscoveredPackageFile&& Other)
-		: PackageFilename(MoveTemp(Other.PackageFilename))
-		, PackageTimestamp(MoveTemp(Other.PackageTimestamp))
-	{
-	}
-
-	FDiscoveredPackageFile& operator=(const FDiscoveredPackageFile& Other)
-	{
-		if (this != &Other)
-		{
-			PackageFilename = Other.PackageFilename;
-			PackageTimestamp = Other.PackageTimestamp;
-		}
-		return *this;
-	}
-
-	FDiscoveredPackageFile& operator=(FDiscoveredPackageFile&& Other)
-	{
-		if (this != &Other)
-		{
-			PackageFilename = MoveTemp(Other.PackageFilename);
-			PackageTimestamp = MoveTemp(Other.PackageTimestamp);
-		}
-		return *this;
-	}
-
-	/** The name of the package file on disk */
-	FString PackageFilename;
-
-	/** The modification timestamp of the package file (when it was discovered) */
-	FDateTime PackageTimestamp;
-};
-
-namespace AssetDataDiscoveryUtil
+namespace AssetDataGather
+{
+namespace Private
 {
-	bool PassesScanFilters(const TArray<FString>& InBlacklistFilters, const FString& InPath);
+class FAssetDataDiscovery;
+struct FGatheredPathData;
+struct FSetPathProperties;
+}
 }
+}
+
+#if DO_CHECK
+typedef FCriticalSectionQueryable FGathererCriticalSection;
+typedef FScopeLockQueryable FGathererScopeLock;
+#define CHECK_IS_LOCKED_CURRENT_THREAD(CritSec) check(CritSec.IsLockedOnCurrentThread())
+#define CHECK_IS_NOT_LOCKED_CURRENT_THREAD(CritSec) check(!CritSec.IsLockedOnCurrentThread())
+#else
+typedef FCriticalSection FGathererCriticalSection;
+typedef FScopeLock FGathererScopeLock;
+#define CHECK_IS_LOCKED_CURRENT_THREAD(CritSec) do {} while (false)
+#define CHECK_IS_NOT_LOCKED_CURRENT_THREAD(CritSec) do {} while (false)
+#endif
 
 /**
- * Async task for discovering files that FAssetDataGatherer should search
+ * Async task for gathering asset data from from the file list in FAssetRegistry
  */
-class FAssetDataDiscovery : public FRunnable
+class FAssetDataGatherer : public FRunnable
 {
 public:
-	/** Constructor */
-	FAssetDataDiscovery(const TArray<FString>& InPaths, const TArray<FString>& InBlacklistScanFilters, bool bInIsSynchronous);
-	virtual ~FAssetDataDiscovery();
-
-	// FRunnable implementation
-	virtual bool Init() override;
-	virtual uint32 Run() override;
-	virtual void Stop() override;
-	virtual void Exit() override;
-
-	/** Signals to end the thread and waits for it to close before returning */
-	void EnsureCompletion();
-
-	/** Gets search results from the file discovery */
-	bool GetAndTrimSearchResults(TArray<FString>& OutDiscoveredPaths, TArray<FDiscoveredPackageFile>& OutDiscoveredFiles, int32& OutNumPathsToSearch);
-
-	/** Adds a root path to the search queue. Only works when searching asynchronously */
-	void AddPathToSearch(const FString& Path);
-
-	/** If assets are currently being asynchronously scanned in the specified path, this will cause them to be scanned before other assets. */
-	void PrioritizeSearchPath(const FString& PathToPrioritize);
-
-	/** Set the blacklist filters to use during scanning. */
-	void SetBlacklistScanFilters(const TArray<FString>& InBlacklistScanFilters);
-
-private:
-	/** Sort the paths so that items belonging to the current priority path is processed first */
-	void SortPathsByPriority(const int32 MaxNumToSort);
-
-	/** A critical section to protect data transfer to other threads */
-	FCriticalSection WorkerThreadCriticalSection;
-
-	/** The current path to prioritize. Files and folders under this location will be bumped to the top of the processing list as they are discovered */
-	FString FilenamePathToPrioritize;
-
-	/** True if this gather request is synchronous (i.e, IsRunningCommandlet()) */
-	bool bIsSynchronous;
-
-	/** True if in the process of discovering files in PathsToSearch */
-	bool bIsDiscoveringFiles;
-
-	/** The directories in which to discover assets and paths */
-	// IMPORTANT: This variable may be modified by from a different thread via a call to AddPathToSearch(), so access 
-	//            to this array should be handled very carefully.
-	TArray<FString> DirectoriesToSearch;
-
-	/** The paths found during the search. It is not threadsafe to directly access this array */
-	TArray<FString> DiscoveredPaths;
-
-	/** List of filters that will be run on discovered paths to see if they should be scanned. */
-	TArray<FString> BlacklistScanFilters;
-
-	/** List of priority files that need to be processed by the gatherer. It is not threadsafe to directly access this array */
-	TArray<FDiscoveredPackageFile> PriorityDiscoveredFiles;
-
-	/** List of non-priority files that need to be processed by the gatherer. It is not threadsafe to directly access this array */
-	TArray<FDiscoveredPackageFile> NonPriorityDiscoveredFiles;
-
-	/** > 0 if we've been asked to abort work in progress at the next opportunity */
-	FThreadSafeCounter StopTaskCounter;
-
-	/** Thread to run the discovery FRunnable on */
-	FRunnableThread* Thread;
-};
+	FAssetDataGatherer(const TArray<FString>& InBlacklistLongPackageNames, const TArray<FString>& InBlacklistMountRelativePaths, bool bInIsSynchronous);
+	virtual ~FAssetDataGatherer();
 
 
-/**
- * Used to control the cache location and behavior of an FAssetDataGatherer cache
- */
-enum class EAssetDataCacheMode : uint8
-{
-	/** Do not cache */
-	NoCache,
+	// Extra at-construction configuration 
 
-	/** Use the monolithic CachedAssetRegistry.bin cache (should only be used for the main asset registry scan) */
-	UseMonolithicCache,
+	/** Configure the gatherer to use a single monolithic cache for all files, and read/write this cache during ticks. */
+	void SetUseMonolithicCache(bool bInUseMonolithicCache);
 
-	/** Use a modular cache file based on the given paths (should be used for synchronous scans, eg, from object libraries) */
-	UseModularCache,
-};
 
+	// Controlling Async behavior
 
-/**
- * Async task for gathering asset data from from the file list in FAssetRegistry
- */
-class FAssetDataGatherer : public FRunnable
-{
-public:
-	/** Constructor */
-	FAssetDataGatherer(const TArray<FString>& Paths, const TArray<FString>& SpecificFiles, const TArray<FString>& InBlacklistScanFilters, bool bInIsSynchronous, EAssetDataCacheMode AssetDataCacheMode);
-	virtual ~FAssetDataGatherer();
+	/** Start the async thread, if this Gatherer was created async. Does nothing if not async or already started. */
+	void StartAsync();
 
 	// FRunnable implementation
 	virtual bool Init() override;
@@ -178,33 +80,91 @@ public:
 	virtual void Stop() override;
 	virtual void Exit() override;
 
+	bool IsSynchronous() const;
 	/** Signals to end the thread and waits for it to close before returning */
 	void EnsureCompletion();
 
-	/** Gets search results from the data gatherer */
-	bool GetAndTrimSearchResults(TBackgroundGatherResults<FAssetData*>& OutAssetResults, TBackgroundGatherResults<FString>& OutPathResults, TBackgroundGatherResults<FPackageDependencyData>& OutDependencyResults, TBackgroundGatherResults<FString>& OutCookedPackageNamesWithoutAssetDataResults, TArray<double>& OutSearchTimes, int32& OutNumFilesToSearch, int32& OutNumPathsToSearch, bool& OutIsDiscoveringFiles);
 
-	/** Adds a root path to the search queue. Only works when searching asynchronously */
-	void AddPathToSearch(const FString& Path);
+	// Receiving Results (possibly while tick is running)
 
-	/** Adds specific files to the search queue. Only works when searching asynchronously */
-	void AddFilesToSearch(const TArray<FString>& Files);
+	/** Gets search results from the data gatherer. */
+	void GetAndTrimSearchResults(bool& bOutIsSearching, TRingBuffer<FAssetData*>& OutAssetResults, TRingBuffer<FString>& OutPathResults,
+		TRingBuffer<FPackageDependencyData>& OutDependencyResults, TRingBuffer<FString>& OutCookedPackageNamesWithoutAssetDataResults,
+		TArray<double>& OutSearchTimes, int32& OutNumFilesToSearch, int32& OutNumPathsToSearch, bool& OutIsDiscoveringFiles);
+	/** Wait for all monitored assets under the given path to be added to search results. Returns immediately if the given path is not monitored. */
+	void WaitOnPath(FStringView LocalPath);
+	/**
+	 * Mark a set of paths as whitelisted, optionally force rescanning and ignore blacklist on them,
+	 * and wait for all assets in the paths to be added to search results.
+	 * Wait time is minimized by prioritizing the paths and transferring async scanning to the current thread.
+	 * If SaveCacheFilename is non-empty, save a cachefile to it with all discovered paths that are in one of
+	 * the SaveCacheLongPackageNameDirs paths.
+	 */
+	void ScanPathsSynchronous(const TArray<FString>& InPaths, bool bForceRescan, bool bIgnoreBlackListScanFilters,
+		const FString& SaveCacheFilename, const TArray<FString>& SaveCacheLongPackageNameDirs);
+	/** Wait for all monitored assets to be added to search results. */
+	void WaitForIdle();
+	/**
+	 * Report whether all monitored assets have been added to search results, AND these results have been gathered through a GetAndTrimSearchResults call.
+	 * This function can be used to check whether there is any work to be done on the gatherer.
+	 */
+	bool IsComplete() const;
 
-	/** If assets are currently being asynchronously scanned in the specified path, this will cause them to be scanned before other assets. */
-	void PrioritizeSearchPath(const FString& PathToPrioritize);
 
-	/** Set the blacklist filters to use during scanning. */
-	void SetBlacklistScanFilters(const TArray<FString>& InBlacklistScanFilters);
+	// Reading/writing/triggering depot-wide properties and events (possibly while tick is running)
 
+	/** Inform the gatherer that initial plugins have finished loading and it should not retry failed loads in hope of a missing custom version */
+	void SetInitialPluginsLoaded();
 	/** Report whether the gatherer is configured to load depends data in addition to asset data. */
 	bool IsGatheringDependencies() const;
+	/** Return whether the current process enables reading and writing AssetDataGatherer cache files. */
+	bool IsCacheEnabled() const;
+	/** Calculate the cache filename that should be used for the given list of package paths. */
+	FString GetCacheFilename(TConstArrayView<FString> CacheFilePackagePaths);
+	/** Attempt to read the cache file at the given LocalPath, and store all of its results in the in-memory cache. */
+	void LoadCacheFile(FStringView CacheFilename);
+	/** Return the total amount of heap memory used by the gatherer (including not-yet-claimed search results). Used for performance metrics. */
+	uint32 GetAllocatedSize() const;
 
-	/** Inform the gatherer that initial plugins have finished loading and it should not retry failed loads in hope of a missing custom version */
-	void SetInitialPluginsLoaded();
+
+	// Configuring mount points (possibly while tick is running)
+
+	/** Add a mountpoint to the gatherer. This should be the localpath to the packagepath/localpath pair that has been registered with FPackageName .*/
+	void AddMountPoint(FStringView LocalPath, FStringView LongPackageName);
+	/** Remove a previously added mountpoint. */
+	void RemoveMountPoint(FStringView LocalPath);
+	/** For each LocalPath in LocalPaths, query FPackageName for its mount point, and add this mountpoint to the gatherer. */
+	void AddRequiredMountPoints(TArrayView<FString> LocalPaths);
+
+
+	// Reading/Writing properties of files and directories (possibly while tick is running)
+
+	/** Called from DirectoryWatcher. Update the directory in internal structures and report it as necessary in a future search result. */
+	void OnDirectoryCreated(FStringView LocalPath);
+	/** Called from DirectoryWatcher. Update the files in internal structures and report them as necessary in a future search result. */
+	void OnFilesCreated(TConstArrayView<FString> LocalPaths);
+	/** Mark a file or directory to be scanned before unprioritized assets. */
+	void PrioritizeSearchPath(const FString& PathToPrioritize);
+	/**
+	 * Mark whether a given path is in the scanning whitelist.
+	 *
+	 * By default no paths are scanned; adding a path to the whitelist causes it and its subdirectories to be scanned.
+	 * Note that the blacklist (InBlacklistLongPackageName) overrides the whitelist.
+	 * Whitelist settings are recursive. Attempting to mark a path as whitelisted if a parent path is whitelisted
+	 * will have no effect. This means the scenario (1) add whitelist A (2) add whitelist A/Child (3) remove whitelist A will
+	 * therefore not result in A/Child being whitelisted.
+	 */
+	void SetIsWhitelisted(FStringView LocalPath, bool bIsWhitelisted);
+	/** Report whether the path is in the whitelist. If not it will not be scanned. If so it will be scanned unless it is in the blacklist. */
+	bool IsWhitelisted(FStringView LocalPath) const;
+	/** Report whether the path is in the blacklist. If so it will not be scanned regardless of whether it is in the whitelist. */
+	bool IsBlacklisted(FStringView LocalPath) const;
+	/** Report whether the path is both in the whitelist and not in the blacklist. */
+	bool IsMonitored(FStringView LocalPath) const;
 
 	/**
 	 * Reads FAssetData information out of a previously initialized package reader
-	 * 
+	 *
 	 * @param PackageReader the previously opened package reader
 	 * @param AssetDataList the FAssetData for every asset found in the file
 	 * @param DependencyData the FPackageDependencyData for every asset found in the file, can be null
@@ -213,9 +173,35 @@ public:
 	static bool ReadAssetFile(FPackageReader& PackageReader, TArray<FAssetData*>& AssetDataList, FPackageDependencyData* DependencyData, TArray<FString>& CookedPackagesToLoadUponDiscovery);
 
 private:
-	/** Sort the paths so that items belonging to the current priority path is processed first */
-	void SortPathsByPriority(const int32 MaxNumToSort);
 
+	/**
+	 * Helper function to run the tick in a loop-within-a-loop to minimize critical section entry, and to move expensive
+	 * operations out of the critical section
+	 */
+	void InnerTickLoop(bool bInIsSynchronousTick, bool bContributeToCacheSave);
+	/**
+	 * Tick function to pump scanning and push results into the search results structure. May be called from devoted thread
+	 * or inline from synchronous functions on other threads.
+	 */
+	void TickInternal(bool& bOutIsIdle);
+	/** Add any new package files from the background directory scan to our work list **/
+	void IngestDiscoveryResults();
+
+	/** Helper for OnFilesCreated. Update the file in internal structures and report it as necessary in a future search result. */
+	void OnFileCreated(FStringView LocalPath);
+
+	/**
+	 * Set a selection of directory-scanning properties on a given LocalPath.
+	 * This function is used when multiple properties need to be set on the path and we want to avoid redundant
+	 * tree-traversal costs.
+	 */
+	void SetDirectoryProperties(FStringView LocalPath, const UE::AssetDataGather::Private::FSetPathProperties& Properties);
+
+	/** Wait for all monitored assets under the given path to be added to search results. Returns immediately if the given path are not monitored. */
+	void WaitOnPathsInternal(TConstArrayView<FString> LocalAbsPaths, const FString& SaveCacheFilename, const TArray<FString>& SaveCacheFilterDirs);
+
+	/** Sort the pending list of discovered asset filepaths so that assets under the given directory/filename are processed first. */
+	void SortPathsByPriority(TConstArrayView<FString> LocalAbsPaths, int32& OutNumPaths);
 	/**
 	 * Reads FAssetData information out of a file
 	 *
@@ -227,93 +213,201 @@ private:
 	 *
 	 * @return true if the file was successfully read
 	 */
-	bool ReadAssetFile(const FString& AssetFilename, TArray<FAssetData*>& AssetDataList, FPackageDependencyData& DependencyData, TArray<FString>& CookedPackagesToLoadUponDiscovery, bool& OutCanRetry) const;
+	bool ReadAssetFile(const FString& AssetFilename, TArray<FAssetData*>& AssetDataList, FPackageDependencyData& DependencyData,
+		TArray<FString>& CookedPackagesToLoadUponDiscovery, bool& OutCanRetry) const;
 
-	/** Serializes the timestamped cache of discovered assets. Used for quick loading of data for assets that have not changed on disk */
-	template<class Archive>
-	void SerializeCache(Archive&& Ar);
+	/** Helper for LoadCacheFile and internal callsites; occurs inside the proper critical section. */
+	void LoadCacheFileInternal(FStringView CacheFilename);
+	/**
+	 * If a save of the async cache has been triggered, get the cache filename and pointers to all elements that should be saved,
+	 * for later saving outside of the critical section.
+	 */
+	void TryReserveSaveAsyncCache(FString& OutCacheFilename, TArray<TPair<FName,FDiskCachedAssetData*>>& AssetsToSave);
+	/**
+	 * If the CacheFilename/AssetsToSave are non empty, save the cache file. 
+	 * This function reads the read-only-after-creation data from each FDiskCachedAssetData*, but otherwise does not use
+	 * data from this Gatherer and so can be run outside any critical section.
+	 */
+	void SaveCacheFileInternal(const FString& CacheFilename, const TArray<TPair<FName,FDiskCachedAssetData*>>& AssetsToSave, bool bIsAsyncCacheSave);
+	/**
+	 * Get the list of FDiskCachedAssetData* that have been loaded in the gatherer, for saving into a cachefile.
+	 * Filters the list of assets by child paths of the elements in SaveCacheLongPackageNameDirs, if it is non-empty.
+	 */
+	void GetAssetsToSave(TArrayView<const FString> SaveCacheLongPackageNameDirs, TArray<TPair<FName,FDiskCachedAssetData*>>& OutAssetsToSave);
+	/** Writes to file the timestamped cache of discovered assets. */
+	void SerializeCacheSave(FAssetRegistryWriter& Ar, const TArray<TPair<FName,FDiskCachedAssetData*>>& AssetsToSave);
+	/** Reads from file the timestamped cache of discovered assets, for quick loading of data for assets that have not changed on disk. */
+	void SerializeCacheLoad(FAssetRegistryReader& Ar);
 
 	/* Adds the given PackageName,DiskCachedAssetData pair into NewCachedAssetDataMap, and detects collisions for multiple files with the same PackageName */
 	void AddToCache(FName PackageName, FDiskCachedAssetData* DiskCachedAssetData);
 
+	/**
+	 * Mark that the gatherer has become idle or has become active. Called from tick function and configuration functions when they note a possible state change.
+	 * Caller is responsible for holding the ResultsLock.
+	 */
+	void SetIsIdle(bool IsIdle);
+
+	/** Minimize memory usage in the buffers used during gathering. */
+	void Shrink();
+
+	/** Scoped guard for pausing the asynchronous tick. */
+	struct FScopedPause
+	{
+		FScopedPause(const FAssetDataGatherer& InOwner);
+		~FScopedPause();
+		const FAssetDataGatherer& Owner;
+	};
+
+	/** Convert the LocalPath into our normalized version. */
+	static FString NormalizeLocalPath(FStringView LocalPath);
+	/** Convert the LongPackageName into our normalized version. */
+	static FStringView NormalizeLongPackageName(FStringView LongPackageName);
+
 private:
-	/** A critical section to protect data transfer to the main thread */
-	FCriticalSection WorkerThreadCriticalSection;
 
-	/** The current path to prioritize. Files and folders under this location will be bumped to the top of the processing list as they are discovered */
-	FString FilenamePathToPrioritize;
+	/**
+	 * Critical section to allow Tick to be called from worker thread or from synchronous functions on other threads.
+	 * To prevent DeadLocks, TickLock can not be entered from within any of the other locks on this class.
+	 */
+	mutable FGathererCriticalSection TickLock;
+	/**
+	 * A critical section to protect data transfer to GetAndTrimSearchResults.
+	 * ResultsLock can be entered while holding TickLock.
+	 */
+	mutable FGathererCriticalSection ResultsLock;
 
-	/** List of files that need to be processed by the search. It is not threadsafe to directly access this array */
-	TArray<FDiscoveredPackageFile> FilesToSearch;
 
-	/** > 0 if we've been asked to abort work in progress at the next opportunity */
-	FThreadSafeCounter StopTaskCounter;
+	// Variable section for variables that are constant during threading.
 
-	/** True if this gather request is synchronous (i.e, IsRunningCommandlet()) */
+	/** Thread to run async Ticks on. Constant during threading. */
+	FRunnableThread* Thread;
+	/**
+	 * True if this Gatherer is synchronous, false if it has a worker thread. If synchronous, results will only be
+	 * added when Wait functions are called. Constant during threading.
+	 */
 	bool bIsSynchronous;
+	/** True if dependency data should be gathered. Constant during threading. */
+	bool bGatherDependsData;
+	/** True if the current process allows reading/writing of AssetDataGatherer cache files. Constant during threading. */
+	bool bCacheEnabled;
 
-	/** True if in the process of discovering files in PathsToSearch */
-	bool bIsDiscoveringFiles;
 
-	/** The current search start time */
-	double SearchStartTime;
+	// Variable section for variables that are atomics read/writable from outside critical sections.
 
-	/** The asset data gathered from the searched files */
-	TArray<FAssetData*> AssetResults;
+	/** > 0 if we've been asked to abort work in progress at the next opportunity. Read/writable anywhere. */
+	std::atomic<uint32> IsStopped;
+	/** > 0 if we've been asked to pause the worker thread so a synchronous function can take over the tick. Read/writable anywhere */
+	mutable std::atomic<uint32> IsPaused;
+	/**
+	 * Discovery subsystem; this subsystem handles deciding which paths to search and running the FileManager calls to search directories.
+	 * Pointer itself is constant during threading. Object pointed to is read/writable anywhere; internally provides threadsafety.
+	 */
+	TUniquePtr<UE::AssetDataGather::Private::FAssetDataDiscovery> Discovery;
+	/** Async only. Set to true once initial plugins have been loaded. Read/writable anywhere. */
+	std::atomic<bool> bInitialPluginsLoaded;
+	/**
+	 * Set to true when periodic or final save of the async cache was detected during TickInternal and should be processed outside of the Tick.
+	 * Read/writable anywhere.
+	 */
+	std::atomic<bool> bSaveAsyncCacheTriggered;
 
-	/** Dependency data for scanned packages */
-	TArray<FPackageDependencyData> DependencyResults;
+	// Variable section for variables that are read/writable only within ResultsLock.
 
-	/** A list of cooked packages that did not have asset data in them. These assets may still contain assets (if they were older for example) */
+	/** List of files that need to be processed by the search. Read/writable only within ResultsLock. */
+	TRingBuffer<UE::AssetDataGather::Private::FGatheredPathData> FilesToSearch;
+	/** The asset data gathered from the searched files. Read/writable only within ResultsLock. */
+	TArray<FAssetData*> AssetResults;
+	/** Dependency data gathered from the searched files packages. Read/writable only within ResultsLock. */
+	TArray<FPackageDependencyData> DependencyResults;
+	/**
+	 * A list of cooked packages that did not have asset data in them.
+	 * These assets may still contain assets (if they were older for example). Read/writable only within ResultsLock.
+	 */
 	TArray<FString> CookedPackageNamesWithoutAssetDataResults;
 
-	/** All the search times since the last main thread tick */
+	/** All the search times since the last call to GetAndTrimSearchResults. Read/writable only within ResultsLock. */
 	TArray<double> SearchTimes;
-
-	/** The paths found during the search */
+	/**
+	 * The directories found during the search, unless they are hidden by DirLongPackageNamesToNotReport.
+	 * Read/writable only within ResultsLock.
+	 */
 	TArray<FString> DiscoveredPaths;
+	/**
+	 * If using a TickManaged cache, this is the name of the cache file.
+	 * The cachefile contains timestamp,guid,and asset data for each package file in the the previous process's gather.
+	 * Read/writable only within ResultsLock.
+	 */
+	FString TickManagedCacheFilename;
 
-	/** List of filters that will be run on discovered paths to see if they should be scanned. */
-	TArray<FString> BlacklistScanFilters;
-
-	/** True if dependency data should be gathered */
-	bool bGatherDependsData;
-
-	/** Async only. Set to true once initial plugins have been loaded. */
-	std::atomic<bool> bInitialPluginsLoaded;
-
-	/** The cached value of the NumPathsToSearch returned by FAssetDataDiscovery the last time we synchronized with it */
+	/** The current search start time, set in the first tick after idle and used for performance metrics when reporting results. Read/writable only within ResultsLock. */
+	double SearchStartTime;
+	/** The last time at which the cache file was written, used to periodically update the cache with the results of ticking. Read/writable only within ResultsLock. */
+	double LastCacheWriteTime;
+	/**
+	 * The cached value of the NumPathsToSearch returned by the Discovery subsystem the last time we synchronized with it.
+	 * Read/writable only within ResultsLock.
+	 */
 	int32 NumPathsToSearchAtLastSyncPoint;
-
-	/** Background package file discovery (when running async) */
-	TSharedPtr<FAssetDataDiscovery> BackgroundPackageFileDiscovery;
-
-	///////////////////////////////////////////////////////////////
-	// Asset discovery caching
-	///////////////////////////////////////////////////////////////
-
-	/** True if this gather request should both load and save the asset cache. Only one gatherer should do this at a time! */
-	bool bLoadAndSaveCache;
-
-	/** True if we have finished discovering our first wave of files */
+	/**
+	 * Track whether we are using a monolithic cache. The monolithic cache is only supported as a TickManaged cache.
+	 * Setting it to true sets up the TickManaged cache. Read/writable only within ResultsLock.
+	 */
+	bool bUseMonolithicCache;
+	/**
+	 * Track whether a cache file should be loaded and saved by the Tick function
+	 * (as opposed to being externally managed by calls to LoadCacheFile/SaveCacheFile).
+	 * Read/writable only within ResultsLock.
+	 */
+	bool bUseTickManagedCache;
+	/** If bUseTickManagedCache is true, track whether the cache has been loaded. Read/writable only within ResultsLock. */
+	bool bHasLoadedTickManagedCache;
+	/** Track whether the Discovery subsystem has gone idle and we have read all filenames from it. Read/writable only within ResultsLock. */
+	bool bDiscoveryIsComplete;
+	/** Track whether this Gather has gone idle and a caller has read all search data from it. Read/writable only within ResultsLock. */
+	bool bIsComplete;
+	/** Track whether this Gatherer has gone idle. Read/writable only within ResultsLock. */
+	bool bIsIdle;
+	/** Track the first tick after idle to set up e.g. timing data. Read/writable only within ResultsLock. */
+	bool bFirstTickAfterIdle;
+	/** True if we have finished discovering our first wave of files, to report metrics for that most-important wave. Read/writable only within ResultsLock. */
 	bool bFinishedInitialDiscovery;
 
-	/** The name of the file that contains the timestamped cache of discovered assets */
-	FString CacheFilename;
 
-	/** An array of all cached data that was newly discovered this run. This array is just used to make sure they are all deleted at shutdown */
-	TArray<FDiskCachedAssetData*> NewCachedAssetData;
+	// Variable section for variables that are read/writable only within TickLock.
 
-	/** 
-	 * Map of PackageName to cached discovered assets that were loaded from disk. 
-	 * This should only be modified in the loading section of SerializeCache and remain const after
-	 * since NewCachedAssetDataMap keeps pointers to the values in this map.
+	/**
+	 * An array of all cached data that was newly discovered this run. This array is just used to make sure they are all deleted at shutdown.
+	 * Read/writable only within TickLock.
 	 */
-	TMap<FName, FDiskCachedAssetData> DiskCachedAssetDataMap;
-
-	/** Map of PackageName to cached discovered assets that will be written to disk at shutdown */
+	TArray<FDiskCachedAssetData*> NewCachedAssetData;
+	TArray<TPair<int32, FDiskCachedAssetData*>> DiskCachedAssetBlocks;
+	/**
+	 * Map of PackageName to cached discovered assets that were loaded from disk.
+	 * This should only be modified in the loading section of SerializeCache calls.
+	 * Read/Writable only within TickLock.
+	 */
+	TMap<FName, FDiskCachedAssetData*> DiskCachedAssetDataMap;
+	/** Map of PackageName to cached discovered assets that will be written to disk at shutdown. Read/writable only within TickLock. */
 	TMap<FName, FDiskCachedAssetData*> NewCachedAssetDataMap;
+	/**
+	 * Used when we are blocking on gather results. If non-zero, we should end the tick when WaitBatchCount files have been processed.
+	 * Read/writable only within TickLock.
+	 */
+	int32 WaitBatchCount;
+	/** How many files in the search results were read from the cache. Read/writable only within TickLock. */
+	int32 NumCachedFiles;
+	/** How many files in the search results were not in the cache and were read by parsing the file. Read/writable only within TickLock. */
+	int32 NumUncachedFiles;
+	/**
+	 * True if the current TickInternal is synchronous, which may be because bIsSynchronous or because the game thread has taken over the tick for a synchronous function.
+	 * Read/writable only within TickLock.
+	 */
+	bool bIsSynchronousTick;
+	/**
+	 * Set to true when a thread is in the middle of saving an async cache and so another save of the cache should not be triggered.
+	 * Read/writable only within TickLock.
+	 */
+	bool bIsSavingAsyncCache;
 
-	/** Thread to run the cleanup FRunnable on */
-	FRunnableThread* Thread;
 };
diff --git a/Engine/Source/Runtime/AssetRegistry/Private/AssetDataGathererPrivate.h b/Engine/Source/Runtime/AssetRegistry/Private/AssetDataGathererPrivate.h
new file mode 100644
--- /dev/null
+++ b/Engine/Source/Runtime/AssetRegistry/Private/AssetDataGathererPrivate.h
@@ -0,0 +1,578 @@
+// Copyright Epic Games, Inc. All Rights Reserved.
+
+#pragma once
+
+#include "AssetDataGatherer.h"
+
+#include <atomic>
+#include "Containers/Set.h"
+#include "Misc/StringBuilder.h"
+#include "Templates/RefCounting.h"
+#include "Templates/UniquePtr.h"
+
+namespace UE
+{
+namespace AssetDataGather
+{
+namespace Private
+{
+
+class FMountDir;
+
+/** Enum to specify files and directories that should be completed earlier than others */
+enum class EPriority : uint8
+{
+	/** Game thread is blocked until the file/directory is completed */
+	Blocking,
+	/** Optional information (e.g. use of the ReferenceViewer) is unavailable until the file/directory is completed */
+	High,
+	/** Nothing has requested the file/directory yet */
+	Normal
+};
+constexpr uint32 CountEPriority = static_cast<uint32>(EPriority::Normal) + 1;
+
+/** Fields being set in a call to SetProperty */
+struct FSetPathProperties
+{
+	/** The path (usually a plugin's root content path) has been requested for scanning through e.g. ScanPathsSynchronous */
+	TOptional<bool> IsWhitelisted;
+	/** The directory's list of direct file/subdirectory children has been scanned through a call to IFileManager::IterateDirectoryStat after process start or the last request to rescan it */
+	TOptional<bool> HasScanned;
+	/**
+	 * Paths under the directory should be used (if whitelisted) even if they match blacklist entries. Used only on the MountDir.
+	 * Used e.g. when a level is opened in the editor and we want to add to the asset registries all of its one-file-per-actor files that are normally blacklisted.
+	 */
+	TOptional<bool> IgnoreBlacklist;
+	/** Priority that determines which directories are scanned first, for when a directory is being waited on or are not blocking but are wanted for optional features */
+	TOptional<EPriority> Priority;
+
+	/** Used to early-exit from tree traversal when all properties have finished being handled */
+	bool IsSet() const
+	{
+		return IsWhitelisted.IsSet() | HasScanned.IsSet() | IgnoreBlacklist.IsSet() | Priority.IsSet();
+	}
+};
+
+/** Information needed about a discovered asset file or path that is needed by the Discoverer */
+struct FDiscoveredPathData
+{
+	/** The absolute path to the file on disk, relative to the directory that issued the scan */
+	FString LocalAbsPath;
+	/** The LongPackageName of the path (inherited from the MountDir) */
+	FString LongPackageName;
+	/** The relative path from the path's parent directory */
+	FString RelPath;
+	/** If the path is a file, the modification timestamp of the package file (that it had when it was discovered) */
+	FDateTime PackageTimestamp;
+
+	FDiscoveredPathData() = default;
+	FDiscoveredPathData(FStringView InLocalAbsPath, FStringView InLongPackageName, FStringView InRelPath, const FDateTime& InPackageTimestamp);
+	FDiscoveredPathData(FStringView InLocalAbsPath, FStringView InLongPackageName, FStringView InRelPath);
+	void Assign(FStringView InLocalAbsPath, FStringView InLongPackageName, FStringView InRelPath, const FDateTime& InPackageTimestamp);
+	void Assign(FStringView InLocalAbsPath, FStringView InLongPackageName, FStringView InRelPath);
+
+	/** Return the total amount of heap memory used by the gatherer (including not-yet-claimed search results). Used for performance metrics. */
+	uint32 GetAllocatedSize() const;
+};
+
+/** Information needed about a discovered asset file or path that is needed by the Gatherer */
+struct FGatheredPathData
+{
+	/** The absolute path to the file on disk, relative to the directory that issued the scan */
+	FString LocalAbsPath;
+	/** The LongPackageName of the path (inherited from the MountDir) */
+	FString LongPackageName;
+	/** The modification timestamp of the package file (that it had when it was discovered) */
+	FDateTime PackageTimestamp;
+
+	FGatheredPathData() = default;
+	FGatheredPathData(FStringView InLocalAbsPath, FStringView InLongPackageName, const FDateTime& InPackageTimestamp);
+	explicit FGatheredPathData(const FDiscoveredPathData& DiscoveredData);
+	void Assign(FStringView InLocalAbsPath, FStringView InLongPackageName, const FDateTime& InPackageTimestamp);
+	void Assign(const FDiscoveredPathData& DiscoveredData);
+
+	/**
+	 * Return the total amount of heap memory used by the gatherer (including not-yet-claimed search results).
+	 * Used for performance metrics.
+	 */
+	uint32 GetAllocatedSize() const;
+};
+
+/**
+ * Tree data node representing a directory in the scan; direct subdirectories are stored as other FScanDir instances
+ * referenced from the SubDirs array. Directories are removed from the tree once their scans are finished to save memory.
+ * Queries take into account that deleted nodes have been completed.
+ *
+ * This class is not ThreadSafe; The FAssetDataDiscovery reads/writes its data only while holding the TreeLock critical section.
+ */
+class FScanDir : public FRefCountBase
+{
+public:
+	FScanDir(FMountDir& InMountDir, FScanDir* InParent, FStringView RelPath);
+	~FScanDir();
+
+	/**
+	 * Marks that this ScanDir is no longer in use and clear its data. The ScanDir will remain allocated until all threads
+	 * have dropped their reference to it.
+	 */
+	void Shutdown();
+
+	/** Check whether this ScanDir is still alive or if it has been marked for destruction and cleared on another thread. */
+	bool IsValid() const;
+
+	FMountDir* GetMountDir() const;
+
+	/** Get this ScanDir's RelPath from its Parent */
+	FStringView GetRelPath() const;
+
+	/** Return this ScanDir's accumulated Priority: the maximum of its directly-set priority and its children's priority. */
+	EPriority GetPriority() const;
+
+	/** Calculate this ScanDir's full absolute path by accumulating RelPaths from parents and append it. */
+	void AppendLocalAbsPath(FStringBuilderBase& OutFullPath) const;
+	/** Calculate this ScanDir's full absolute path by accumulating RelPaths from parents and return it as a string. */
+	FString GetLocalAbsPath() const;
+	/** Calculate the relative path from the MountPoint to this ScanDir and append it to OutRelPath.  */
+	void AppendMountRelPath(FStringBuilderBase& OutRelPath) const;
+	/** Calculate the relative path from the MountPoint to this ScanDir and return it as a string.  */
+	FString GetMountRelPath() const;
+
+	/** Report whether this ScanDir is blacklisted */
+	bool IsBlacklisted() const;
+	/**
+	 * Return whether this ScanDir has been directly marked as whitelisted. A ScanDir may be indirectly whitelisted
+	 * because it has a whitelisted parent directory; that is not reported by this function.
+	 */
+	bool IsDirectWhitelisted() const;
+	/**
+	 * Return whether this scandir is whitelisted and not blacklisted and hence is in the set of directories that
+	 * need to be scanned.
+	 * Considers whitelisting/blacklisting inherited from parents in addition to the direct settings on this scandir.
+	 */
+	bool IsRecursiveMonitored(bool bParentIsWhitelisted) const;
+	/** Report whether the given RelPath is whitelisted, because either it or a parent is direct whitelisted. */
+	bool IsPathWhitelisted(FStringView InRelPath, bool bParentIsWhitelisted) const;
+
+	/** Report whether this ScanDir will be scanned in the current or future Tick. */
+	bool ShouldScan(bool bParentIsWhitelisted) const;
+	/** Report whether this ScanDir has been scanned. */
+	bool HasScanned() const;
+	/** Report whether this ScanDir is complete: has scanned or should not scan, and all subdirs have completed. */
+	bool IsComplete() const;
+
+	/** Return the memory used by the tree under this ScanDir; excludes sizeof(*this). */
+	uint32 GetAllocatedSize() const;
+	/**
+	 * Find the Direct parent of InRelPath, or a fallback. Will return null only if !bIsDirectory and InRelPath is empty.
+	 * The fallback is returned if InRelPath has already completed and been deleted, or if InRelPath is not recursively whitelisted or is blacklisted.
+	 * The fallback is the lowest existing parent directory of InRelPath.
+	 */
+	FScanDir* GetControllingDir(FStringView InRelPath, bool bIsDirectory, bool bParentIsWhitelisted, bool& bOutIsWhitelisted, FString& OutRelPath);
+
+	/**
+	 * Set values of fields on the given directory indicated by InRelPath for all of the properties existing on InProperties.
+	 * Returns whether the directory was found and its property was changed; returns false if InRelPath was not a
+	 * directory or the property did not need to be changed.
+	 */
+	bool TrySetDirectoryProperties(FStringView InRelPath, const FSetPathProperties& InProperties, bool bConfirmedExists);
+	/** Mark that the given file has already been scanned, so that it will not be double reporting in the upcoming directory scan, if one is upcoming. */
+	void MarkFileAlreadyScanned(FStringView BaseName);
+	/** Set the direct priority of this ScanDir and update the accumulated priority to match it. */
+	void SetDirectPriority(EPriority InPriority);
+
+	/**
+	 * Called from the Tick; handle the list of subdirs and files that were found from IterateDirectoryStat called on this ScanDir,
+	 * reporting discovered files and updating status variables.
+	 */
+	void SetScanResults(FStringView LocalAbsPath, TArrayView<FDiscoveredPathData>& InOutSubDirs, TArrayView<FDiscoveredPathData>& InOutFiles);
+	/**
+	 * Update the completion state of this ScanDir based on its scan status and its child dirs completion state.
+	 * Set OutCursor to the highest priority child dir that needs to be updated if any children need to be updated and/or scanned first.
+	 * Set OutCursor to the parent ScanDir (or null if there is no parent) if this ScanDir is now complete.
+	 */
+	void Update(FScanDir*& OutCursor, bool& bInOutParentIsWhitelisted);
+
+	/** Thread-synchronization helper - return true if the Tick thread is in the middle of scanning this directory. */
+	bool IsScanInFlight() const;
+	/** Thread-synchronization helper - set that the Tick thread is starting/done with the scan of this directory. */
+	void SetScanInFlight(bool bInScanInFlight);
+	/** Thread-synchronization helper - report whether a non tick thread has marked that this directory is changed or invalidated and the scan should be thrown out. */
+	bool IsScanInFlightInvalidated() const;
+	/** Thread-synchronization helper - set that the current ongoing scan is invalidated, or clear the marker from the tick thread once it has been consumed. */
+	void SetScanInFlightInvalidated(bool bInValidated);
+
+	/** Set completion flags on this and its parents (and optionally its descendents) so that it will be updated again. */
+	void MarkDirty(bool bMarkDescendents);
+
+	/** Minimize data in internal buffers. */
+	void Shrink();
+
+protected:
+
+	/**
+	 * Setting to complete marks that this directory has been scanned, and all of its subdirectories have scanned as well, so it can skipped when it or its parent
+	 * is waited on. Setting back to incomplete can occur when a AssetDataGatherer client wants to rescan the directory.
+	 */
+	virtual void SetComplete(bool bInIsComplete);
+
+	/** Find the ScanDir subdirectory for the given basename, or return nullptr. */
+	FScanDir* FindSubDir(FStringView SubDirBaseName);
+	const FScanDir* FindSubDir(FStringView SubDirBaseName) const;
+	/** Find the ScanDir subdirectory for the given basename, and add it if it does not exist. */
+	FScanDir& FindOrAddSubDir(FStringView SubDirBaseName);
+	/** Find the ScanDir subdirectory for the given basename, and if it exists, Shutdown and remove it from SubDirs, which will eventually delete it. */
+	void RemoveSubDir(FStringView SubDirBaseName);
+
+	/** Find the index of the subdir with the given Relative path. */
+	int32 FindLowerBoundSubDir(FStringView SubDirBaseName);
+
+	/** Return the the highest-priority not-yet-complete SubDir, or null if all are complete. */
+	FScanDir* FindHighestPrioritySubDir();
+
+	/** Call the given lambda void(FScanDir&) on each existing SubDir. */
+	template <typename CallbackType> void ForEachSubDir(const CallbackType& Callback);
+	/**
+	 * Call the given lambda void(FScanDir&) on each present-in-memory descedent ScanDir of this.
+	 * Does not look for directories on disk, only the ones that have already been created in memory.
+	 * Depth-first-search traversal, called on parents before children.
+	 */
+	template <typename CallbackType> void ForEachDescendent(const CallbackType& Callback);
+
+	/** Set the AccumulatedPriority of this ScanDir to the maximum of the DirectPriorities of it and all its child dirs. */
+	void UpdateAccumulatedPriority();
+	/** Record a changed priority of a direct child ScanDir and update the accumulated priority of this ScanDir. */
+	void OnChildPriorityChanged(EPriority InPriority, int32 Delta);
+
+	TArray<TRefCountPtr<FScanDir>> SubDirs; // Sorted
+	TArray<FString> AlreadyScannedFiles; // Unsorted
+	FMountDir* MountDir = nullptr;
+	FScanDir* Parent = nullptr;
+	FString RelPath;
+	uint8 PriorityRefCounts[CountEPriority] = {}; // Initialize all elements to 0
+	EPriority DirectPriority = EPriority::Normal;
+	EPriority AccumulatedPriority = EPriority::Normal;
+	bool bHasScanned = false;
+	bool bIsDirectWhitelisted = false;
+	bool bScanInFlight = false;
+	bool bScanInFlightInvalidated = false;
+	bool bIsComplete = false;
+};
+
+/**
+ * Gather data about a MountPoint that has been registered with FPackageName
+ * The FMountDir holds a FScanTree with information about each directory (that is pruned when not in use).
+ * It also holds some data that is needed only per MountPoint, such as the packagename.
+ * It also holds some data per subdirectory that is more performant to hold in a map rather than to require the FScanTrees to be kept.
+ *
+ * This class is not ThreadSafe; The FAssetDataDiscovery reads/writes its data only while holding the TreeLock critical section.
+ */
+class FMountDir
+{
+public:
+	FMountDir(FAssetDataDiscovery& InDiscovery, FStringView LocalAbsPath, FStringView PackagePath);
+	~FMountDir();
+
+	/** The local filesystem path from FPackageName, as an absolute path d:\root\Engine\Content rather than a relative path ../../../Engine/Content. */
+	FStringView GetLocalAbsPath() const;
+	/** The package path from FPackageName. */
+	FStringView GetLongPackageName() const;
+
+	/**
+	 * Get the set of blacklist paths that are child paths of this MountDir.
+	 * Returned paths are relative paths from the MountDir's root.
+	 */
+	const TSet<FString>& GetBlacklist() const;
+	/** Return the FAssetDataDiscovery that owns this FMountDir. */
+	FAssetDataDiscovery& GetDiscovery() const;
+
+	/** Find the direct parent of InRelPath, or the lowest fallback. See FScanDir::GetControllingDir. */
+	FScanDir* GetControllingDir(FStringView LocalAbsPath, bool bIsDirectory, bool& bOutIsWhitelisted, FString& OutRelPath);
+	/** Return the memory used by the tree under this MountDir, except that sizeof(*this) is excluded. */
+	uint32 GetAllocatedSize() const;
+
+	/** Report whether this MountDir is complete: all ScanDirs under it either have scanned or should not scan. */
+	bool IsComplete() const;
+	/** Return this MountDir's Priority, which is the maximum of any of its ScanDirs' priorities. */
+	EPriority GetPriority() const;
+
+	/** Return whether the given path is a child path of *this and is whitelisted. */
+	bool IsPathWhitelisted(FStringView LocalAbsPath) const;
+	/** Return whether the given path is a child path of *this and matches the blacklist filters. */
+	bool IsBlacklisted(FStringView LocalAbsPath) const;
+	/**
+	 * Return whether the given path is a child path of *this and is whitelisted and is not blacklisted, which means
+	 * it will be or has been scanned.
+	 */
+	bool IsMonitored(FStringView LocalAbsPath) const;
+
+	/*
+	 * Set values of fields on the given directory for all of the properties existing on InProperties.
+	 * Returns whether the directory was foundand its property was changed; returns false if LocalAbsPath was not a
+	 * directory under this MountDir or the property did not need to be changed.
+	 */
+	bool TrySetDirectoryProperties(FStringView LocalAbsPath, const FSetPathProperties& InProperties, bool bConfirmedExists);
+	/** 
+	 * Update all ScanDirs under this MountDir in priority order until one is found that needs to be scanned.
+	 * Set OutCursor to that ScanDir, or to nullptr if all are complete.
+	 * Set bOutCursorParentIsWhitelisted to whether the parent of the OutCursor is whitelisted,
+	 * or to false if there is no OutCursor.
+	 */
+	void Update(FScanDir*& OutCursor, bool& bOutCursorParentIsWhitelisted);
+
+	/** Record that a directory at or under the MountDir's root has been scanned, used to detect if configuration occurs after scanning starts. */
+	void SetHasStartedScanning();
+	/** Minimize data in internal buffers. */
+	void Shrink();
+
+	/** Record that the MountDir is rooted at a childpath of this. The childpath will not be scanned by this. */
+	void AddChildMount(FMountDir* ChildMount);
+	/** Mark that a childpath MountDir is being deleted and the childpath should be scanned again by this. */
+	void RemoveChildMount(FMountDir* ChildMount);
+	/** Remove all childmounts. Does not handle properly updating the MountDir to reown those paths; this is used during destruction all MountDirs. */
+	void OnDestroyClearChildMounts();
+	/** Record the backpointer to the parent mountdir that this mountdir's path is a child path of, or null if the parent no longer exists. */
+	void SetParentMount(FMountDir* ParentMount);
+	/** Return the parent MountDir. */
+	FMountDir* GetParentMount() const;
+	/** Return the MountDirs that have been recorded as ChildMounts. */
+	TArray<FMountDir*> GetChildMounts() const;
+
+protected:
+	/** Inspect the Discovery's blacklists and add the ones applicable to this MountDir into this MountDir's set of Blacklists. */
+	void UpdateBlacklist();
+	/** Mark that given path needs to be reconsidered by Update. */
+	void MarkDirty(FStringView MountRelPath);
+
+	/** Holds data about child paths directly in the MountDir rather than requiring FScanDir children. This is more memory efficient and is performant for rarely accessed data. */
+	struct FPathData
+	{
+		FPathData(FStringView MountRelPath);
+		bool IsEmpty() const;
+		uint32 GetAllocatedSize() const;
+
+		/** Relative path from the mountdir. */
+		FString RelPath;
+		/** The path has been marked as should no longer be blacklisted by e.g. ScanPathsSynchronous. */
+		bool bIgnoreBlacklist = false;
+		/** The path is a childpath; another MountDir exists to scan it and it should not be scanned by this MountDir. */
+		bool bIsChildPath = false;
+	};
+	/** Return the PathData with the given MountRelPath, or nullptr if none exists. */
+	FPathData* FindPathData(FStringView MountRelPath);
+	/** Return the PathData with the given MountRelPath, creating it if none exists. */
+	FPathData& FindOrAddPathData(FStringView MountRelPath);
+	/** Remove the PathData with the given MountRelPath if one exists. */
+	void RemovePathData(FStringView MountRelPath);
+
+	/** Array of data keyed by relative path from the MountDir for information about child paths that should not be stored in ScanDirs. */
+	TArray<FPathData> PathDatas;
+	/** Set of relative path from the MountDir paths that should not be scanned, because they were requested blacklisted by clients or because a childmount owns them. */
+	TSet<FString> BlacklistedRelPaths;
+	/** Absolute path to the root of the MountDir in the local file system. */
+	FString LocalAbsPath;
+	/** LongPackageName that was assigned to the MountDir in FPackageName. */
+	FString LongPackageName;
+	/** ScanDir for the root directory of this MountDir; child paths to scan will be created (and destroyed after use) as children of the ScanDir (with the exception of childmounts). */
+	TRefCountPtr<FScanDir> Root;
+	/** Backpointer to the Discovery that owns this MountDir. */
+	FAssetDataDiscovery& Discovery;
+	/** If this is a nested MountDir (a nested path was registered with FPackageName) ParentMount is a pointer to the FMountDir that corresponds to the registered parent directory. */
+	FMountDir* ParentMount = nullptr;
+	/** Records whether any directory at or under the MountDir's root has been scanned. */
+	bool bHasStartedScanning = false;
+};
+
+/** Subsystem that discovers the files that FAssetDataGatherer should process.
+ */
+class FAssetDataDiscovery : public FRunnable
+{
+public:
+	FAssetDataDiscovery(const TArray<FString>& InBlacklistLongPackageNames, const TArray<FString>& InBlacklistMountRelativePaths, bool bInIsSynchronous);
+	virtual ~FAssetDataDiscovery();
+
+
+	// Controlling Async behavior
+
+	/** Start the async thread, if this Gatherer was created async. Does nothing if not async or already started. */
+	void StartAsync();
+
+	// FRunnable implementation
+	virtual bool Init() override;
+	virtual uint32 Run() override;
+	virtual void Stop() override;
+	virtual void Exit() override;
+
+	/** Signals to end the thread and waits for it to close before returning */
+	void EnsureCompletion();
+
+	// Receiving Results and reading properties (possibly while tick is running)
+	/** Gets search results from the file discovery. */
+	void GetAndTrimSearchResults(bool& bOutIsComplete, TArray<FString>& OutDiscoveredPaths, TRingBuffer<FGatheredPathData>& OutDiscoveredFiles, int32& OutNumPathsToSearch);
+	/** Wait (joining in on the tick) until all currently monitored paths have been scanned. */
+	void WaitForIdle();
+	/** Optionally set some scan properties for the given path and then wait for the scan of it to finish. */
+	void SetPropertiesAndWait(const FString& LocalAbsPath, bool bAddToWhitelist, bool bForceRescan, bool bIgnoreBlackListScanFilters);
+	/** Return whether the given path is whitelisted due to e.g. TrySetDirectoryProperties with IsWhitelisted. */
+	bool IsWhitelisted(FStringView LocalAbsPath) const;
+	/** Return whether the given path is blacklisted due to matching blacklist settings and not having been marked IgnoreBlacklist. */
+	bool IsBlacklisted(FStringView LocalAbsPath) const;
+	/** Return whether the given path should be scanned (possibly already has been) because it is whitelisted and not blacklisted. */
+	bool IsMonitored(FStringView LocalAbsPath) const;
+	/** Return the memory used by *this. sizeof(*this) is not included. */
+	uint32 GetAllocatedSize() const;
+
+	// Events and setting of properties (possibly while tick is running)
+	/** Register the given LocalAbsPath/LongPackageName pair that came from FPackageName's list of mount points as a mountpoint to track. Will not be scanned until whitelisted. */
+	void AddMountPoint(const FString& LocalAbsPath, FStringView LongPackageName);
+	/** Remove the mountpoint because FPackageName has removed it. */
+	void RemoveMountPoint(const FString& LocalAbsPath);
+	/** Set properties on the directory, called from AssetDataGatherer when files are requested whitelisted, blacklisted, or rescanned. */
+	bool TrySetDirectoryProperties(const FString& LocalAbsPath, const UE::AssetDataGather::Private::FSetPathProperties& Properties, bool bConfirmedExists);
+	/** Event called from the directory watcher when a directory is created under a mount point. It will be added to the scan if IsMonitored. */
+	void OnDirectoryCreated(FStringView LocalAbsPath);
+	/** Event called from the directory watcher when files are created under a mount point. Each will be reported to the Gatherer if IsMonitored. */
+	void OnFilesCreated(TConstArrayView<FString> LocalAbsPaths);
+
+private:
+	/** Find the MountDir with a root that contains LocalAbsPath, finding the lowest child if there are multiple. */
+	FMountDir* FindContainingMountPoint(FStringView LocalAbsPath);
+	const FMountDir* FindContainingMountPoint(FStringView LocalAbsPath) const;
+	/** Find the mountpoint with the given root. */
+	FMountDir* FindMountPoint(FStringView LocalAbsPath);
+	/** Find the mountpoint with the given root, creating it with root and packagename if it does not already exist. Handle registration of child mountdirs. */
+	FMountDir& FindOrAddMountPoint(FStringView LocalAbsPath, FStringView LongPackageName);
+	/** Execute the search on the sorted-by-path mountdirs, returning the first mountdir with root >= the given path. */
+	int32 FindLowerBoundMountPoint(FStringView LocalAbsPath) const;
+
+	/** Run the tick, either called from the async Run or called on thread from a thread executing a synchronous wait. */
+	void TickInternal();
+	/** Search all of the MountDirs for a highest priority ScanDir that needs to be Updated. */
+	void FindFirstCursor(FScanDir*& OutCursor, bool& bOutCursorParentIsWhitelisted);
+	/** Invalidate the current cursor set from the tick, called when properties are set and we need to reconsider whether directories ShouldBeScanned. */
+	void InvalidateCursor();
+	/** Mark that this discoverer has either finished and is idle, or has received new data and is no longer idle. Update properties dependent upon the idle state. */
+	void SetIsIdle(bool bInIdle);
+
+	/** Store the given discovered files and directories in the results. */
+	void AddDiscovered(TConstArrayView<FDiscoveredPathData> SubDirs, TConstArrayView<FDiscoveredPathData> Files);
+	/**
+	 * Return whether a directory with the given LongPackageName should be reported to the AssetRegistry
+	 * We do not report some directories because they are paths that should not enter the AssetRegistry list of paths if empty,
+	 * and reporting a path to the AssetRegistry adds it unconditionally to the list of paths.
+	 * If ShouldDirBeReported returns false, the directory will still be added to the catalog if non-empty, because the AssetRegistry adds the path of every added file.
+	 */
+	bool ShouldDirBeReported(FStringView LongPackageName) const;
+
+	/** Handle the actions necessary for a single created file. */
+	void OnFileCreated(const FString& LocalPaths);
+	/** Helper for TrySetDirectoryProperties and other internals that need to set it, called from within the critical section. */
+	bool TrySetDirectoryPropertiesInternal(const FString& LocalAbsPath, const UE::AssetDataGather::Private::FSetPathProperties& Properties, bool bConfirmedExists);
+
+	/** Add the given path as a MountPoint and update child registrations. */
+	void AddMountPointInternal(const FString& LocalAbsPath, FStringView LongPackageName);
+	/** Remove the given path as a MountPoint and update child registrations. */
+	void RemoveMountPointInternal(const FString& LocalAbsPath);
+
+	/** Minimize memory usage in the buffers used during gathering. */
+	void Shrink();
+
+	/** Scoped guard for pausing the asynchronous tick. */
+	struct FScopedPause
+	{
+		FScopedPause(const FAssetDataDiscovery& InOwner);
+		~FScopedPause();
+		const FAssetDataDiscovery& Owner;
+	};
+
+private:
+	/**
+	 * Prevent simultaneous ticks from two different threads and protect access to Tick-specific data.
+	 * To prevent DeadLocks, TickLock can not be entered from within any of the other locks on this class.
+	 */
+	mutable FGathererCriticalSection TickLock;
+	/**
+	 * Protect access to data in the ScanDir tree which can be read/write from the tick or from SetProperties.
+	 * To prevent DeadLocks, TreeLock can be entered while holding TickLock, but can not be entered from
+	 * within any of the other locks on this class.
+	 */
+	mutable FGathererCriticalSection TreeLock;
+	/**
+	 * Protect access to the data written from tick and read/written from GetAndTrimSearchResults.
+	 * ResultsLock can be entered while holding TickLock or TreeLock or both.
+	 */
+	mutable FGathererCriticalSection ResultsLock;
+
+
+	// Variable section for variables that are constant during threading.
+
+	/** Blacklisted paths that are full absolute paths. Child paths will not be scanned unless requested to ignore blacklists. Constant during threading. */
+	TArray<FString> BlacklistLongPackageNames;
+	/** Blacklisted paths that are relative paths in each mount. Child paths will not be scanned unless requested to ignore blacklists. Constant during threading. */
+	TArray<FString> BlacklistMountRelativePaths;
+	/** LongPackageNames for directories that should not be reported, see ShouldDirBeReported. Constant during threading. */
+	TSet<FString> DirLongPackageNamesToNotReport;
+	/** Thread to run the discovery FRunnable on. Read-only while threading is possible. Constant during threading. */
+	FRunnableThread* Thread;
+	/** True if this gather request is synchronous (i.e, IsRunningCommandlet()). Constant during threading. */
+	bool bIsSynchronous;
+
+
+	// Variable section for variables that are atomics read/writable from outside critical sections.
+
+	/**
+	 * Whether this Discoverer has finished all work (may be still present in the results.)
+	 * Readable anywhere. Writable only within TreeLock.
+	 */
+	std::atomic<bool> bIsIdle;
+	/** > 0 if we've been asked to abort work in progress at the next opportunity. Read/writable anywhere. */
+	std::atomic<uint32> IsStopped;
+	/** > 0 if we've been asked to pause the worker thread so a synchronous function can take over the tick. Read/writable anywhere. */
+	mutable std::atomic<uint32> IsPaused;
+	/**
+	 * Number of directories that have been discovered and IsMonitored but have not yet been scanned.
+	 * Used for progress tracking. Read/writable anywhere.
+	 */
+	FThreadSafeCounter NumDirectoriesToScan;
+
+
+	// Variable section for variables that are read/writable only within ResultsLock.
+
+	/** Directories found in the scan; may be empty. Read/writable only within ResultsLock. */
+	TArray<FString> DiscoveredDirectories;
+	/** Files found found in the scan. Read/writable only within ResultsLock. */
+	TArray<FGatheredPathData> DiscoveredFiles;
+	/** Time at which the scan was started or last resumed from idle. Used for logging. Read/writable only within ResultsLock. */
+	double DiscoverStartTime = 0.;
+	/** Number of files discovered during scanning since start or resumed from idle. Used for logging. Read/writable only within ResultsLock. */
+	int32 NumDiscoveredFiles = 0;
+
+
+	// Variable section for variables that are read/writable only within TreeLock.
+
+	/**
+	 * Sorted list of MountDirs, sorted by FPackagePath::Less on the absolute paths. Read/writable only within TreeLock.
+	 * Each MountDir contains a ScanDir tree and other data that configures the scanning within that MountPoint.
+	 * Read/writable only with TreeLock, both the list and all data owned by each MountDir.
+	 */
+	TArray<TUniquePtr<FMountDir>> MountDirs;
+	/** The next ScanDir to update in Tick. Read/writable only within TreeLock. */
+	TRefCountPtr<FScanDir> Cursor = nullptr;
+	/** The value of ParentIsWhitelisted for Cursor (or false if Cursor is null). Read/writable only within TreeLock. */
+	bool bCursorParentIsWhitelisted = false;
+
+
+	// Variable section for variables that are read/writable only within TickLock.
+
+	/** Scratch space to store discovered subdirs during the tick, to avoid allocations. Read/writable only within TickLock. */
+	TArray<FDiscoveredPathData> IteratedSubDirs;
+	/** Scratch space to store discovered files during the tick, to avoid allocations. Read/writable only within TickLock. */
+	TArray<FDiscoveredPathData> IteratedFiles;
+
+
+	friend class FMountDir;
+	friend class FScanDir;
+};
+
+} // namespace Private
+} // namespace AssetDataGather
+} // namespace UE
\ No newline at end of file
diff --git a/Engine/Source/Runtime/AssetRegistry/Private/AssetRegistry.cpp b/Engine/Source/Runtime/AssetRegistry/Private/AssetRegistry.cpp
--- a/Engine/Source/Runtime/AssetRegistry/Private/AssetRegistry.cpp
+++ b/Engine/Source/Runtime/AssetRegistry/Private/AssetRegistry.cpp
@@ -1,6 +1,7 @@
 // Copyright Epic Games, Inc. All Rights Reserved.
 
 #include "AssetRegistry.h"
+
 #include "Misc/CommandLine.h"
 #include "Misc/FileHelper.h"
 #include "Misc/Paths.h"
@@ -11,6 +12,7 @@
 #include "UObject/MetaData.h"
 #include "UObject/CoreRedirects.h"
 #include "Blueprint/BlueprintSupport.h"
+#include "AssetDataGatherer.h"
 #include "AssetRegistryPrivate.h"
 #include "AssetRegistry/ARFilter.h"
 #include "DependsNode.h"
@@ -242,7 +244,10 @@ UAssetRegistryImpl::UAssetRegistryImpl(const FObjectInitializer& ObjectInitializ
 
 	const double StartupStartTime = FPlatformTime::Seconds();
 
+	bInitialSearchStarted = false;
 	bInitialSearchCompleted = true;
+	bGatherIdle = false;
+	bSearchAllAssets = false;
 	AmortizeStartTime = 0;
 	TotalAmortizeTime = 0;
 
@@ -276,16 +281,23 @@ UAssetRegistryImpl::UAssetRegistryImpl(const FObjectInitializer& ObjectInitializ
 	// Read default serialization options
 	InitializeSerializationOptionsFromIni(SerializationOptions, FString());
 
-	// Read the scan filters for global asset scanning
-	InitializeBlacklistScanFiltersFromIni();
-
-	// If in the editor, we scan all content right now
-	// If in the game, we expect user to make explicit sync queries using ScanPathsSynchronous
-	// If in a commandlet, we expect the commandlet to decide when to perform a synchronous scan
-	if (GIsEditor && !IsRunningCommandlet())
+	// If in the editor or cookcommandlet, we start the GlobalGatherer now
+	// In the game or other commandlets, we do not construct it until project or commandlet code calls SearchAllAssets or ScanPathsSynchronous
+	bool bSearchAllAssetsAtStart = GIsEditor && (!IsRunningCommandlet() || IsRunningCookCommandlet());
+#if !UE_BUILD_SHIPPING
+	bool bCommandlineAllAssetsAtStart;
+	if (FParse::Bool(FCommandLine::Get(), TEXT("AssetGatherAll="), bCommandlineAllAssetsAtStart))
+	{
+		bSearchAllAssetsAtStart = bCommandlineAllAssetsAtStart;
+	}
+#endif
+	if (bSearchAllAssetsAtStart)
 	{
-		bInitialSearchCompleted = false;
-		SearchAllAssets(false);
+		ConstructGatherer();
+		if (!GlobalGatherer->IsSynchronous())
+		{
+			SearchAllAssetsInitialAsync();
+		}
 	}
 #if ASSETREGISTRY_ENABLE_PREMADE_REGISTRY_IN_EDITOR 
 	if (GIsEditor && !!LoadPremadeAssetRegistryInEditor)
@@ -406,6 +418,10 @@ UAssetRegistryImpl::UAssetRegistryImpl(const FObjectInitializer& ObjectInitializ
 	}
 }
 
+UAssetRegistryImpl::UAssetRegistryImpl(FVTableHelper& Helper)
+{
+}
+
 bool UAssetRegistryImpl::ResolveRedirect(const FString& InPackageName, FString& OutPackageName)
 {
 	int32 DotIndex = InPackageName.Find(TEXT("."), ESearchCase::CaseSensitive);
@@ -463,15 +479,13 @@ void UAssetRegistryImpl::InitRedirectors()
 			FString RootPackageName = FString::Printf(TEXT("/%s/"), *Plugin->GetName());
 			PathsToSearch.Add(RootPackageName);
 
-			const bool bForceRescan = false;
-			ScanPathsAndFilesSynchronous(PathsToSearch, TArray<FString>(), TArray<FString>(), bForceRescan, EAssetDataCacheMode::UseModularCache);
+			ScanPathsSynchronous(PathsToSearch, false /* bForceRescan */, false /* bIgnoreBlacklistScanFilters */);
 		}
 		
 		FName PluginPackageName = FName(*FString::Printf(TEXT("/%s/"), *Plugin->GetName()));
 		TArray<FAssetData> AssetList;
 		GetAssetsByPath(PluginPackageName, AssetList, true, false);
 
-#if 1 // new way
 		for (const FAssetData& Asset : AssetList)
 		{
 			FString NewPackageNameString = Asset.PackageName.ToString();
@@ -488,30 +502,12 @@ void UAssetRegistryImpl::InitRedirectors()
 		FCoreDelegates::FResolvePackageNameDelegate PackageResolveDelegate;
 		PackageResolveDelegate.BindUObject(this, &UAssetRegistryImpl::ResolveRedirect);
 		FCoreDelegates::PackageNameResolvers.Add(PackageResolveDelegate);
-#else
-		TArray<FCoreRedirect> PackageRedirects;
-
-		for ( const FAssetData& Asset : AssetList )
-		{
-			FString NewPackageNameString = Asset.PackageName.ToString();
-			FString RootPackageName = FString::Printf(TEXT("/%s/"), *Plugin->GetName());
-
-			FString OriginalPackageNameString = NewPackageNameString.Replace(*RootPackageName, TEXT("/Game/"));
-
-
-			PackageRedirects.Add( FCoreRedirect(ECoreRedirectFlags::Type_Package, OriginalPackageNameString, NewPackageNameString) );
-
-		}
-
-
-		FCoreRedirects::AddRedirectList(PackageRedirects, Plugin->GetName() );
-#endif
 	}
 }
 
 void UAssetRegistryImpl::RegisterReadOfScriptPackages()
 {
-	if (!BackgroundAssetSearch.IsValid())
+	if (!GlobalGatherer.IsValid())
 	{
 		return;
 	}
@@ -538,16 +534,15 @@ void UAssetRegistryImpl::OnPluginLoadingPhaseComplete(ELoadingPhase::Type Loadin
 
 void UAssetRegistryImpl::ReadScriptPackages()
 {
-	if (BackgroundAssetSearch.IsValid())
+	if (GlobalGatherer.IsValid())
 	{
-		BackgroundAssetSearch->SetInitialPluginsLoaded();
-		if (BackgroundAssetSearch->IsGatheringDependencies())
+		GlobalGatherer->SetInitialPluginsLoaded();
+		if (GlobalGatherer->IsGatheringDependencies())
 		{
 			// Now that all scripts have been loaded, we need to create AssetPackageDatas for every script
 			// This is also done whenever scripts are referenced in our gather of existing packages,
 			// but we need to complete it for all scripts that were referenced but not yet loaded for packages
 			// that we already gathered
-			// CODEREVIEWTODO: Is there a faster way to get the list of loaded script packages?
 			for (TObjectIterator<UPackage> It; It; ++It)
 			{
 				UPackage* Package = *It;
@@ -701,25 +696,6 @@ static void InitializeSerializationOptionsFromIni(FAssetRegistrySerializationOpt
 	}
 }
 
-void UAssetRegistryImpl::InitializeBlacklistScanFiltersFromIni()
-{
-	FConfigFile* EngineIni = GConfig->FindConfigFile(GEngineIni);
-	if (EngineIni)
-	{
-		TArray<FString> IniFilters;
-		EngineIni->GetArray(TEXT("AssetRegistry"), TEXT("BlacklistPackagePathScanFilters"), IniFilters);
-		BlacklistScanFilters.Append(MoveTemp(IniFilters));
-		EngineIni->GetArray(TEXT("AssetRegistry"), TEXT("BlacklistContentSubPathScanFilters"), BlacklistContentSubPaths);
-
-		TArray<FString> ContentRoots;
-		FPackageName::QueryRootContentPaths(ContentRoots);		
-		for (const FString& ContentRoot : ContentRoots)
-		{
-			AddSubContentBlacklist(ContentRoot);
-		}
-	}
-}
-
 void UAssetRegistryImpl::CollectCodeGeneratorClasses() const
 {
 	// Only refresh the list if our registered classes have changed
@@ -760,13 +736,6 @@ UAssetRegistryImpl::~UAssetRegistryImpl()
 		IAssetRegistryInterface::Default = nullptr;
 	}
 
-	// Make sure the asset search thread is closed
-	if (BackgroundAssetSearch.IsValid())
-	{
-		BackgroundAssetSearch->EnsureCompletion();
-		BackgroundAssetSearch.Reset();
-	}
-
 	// Stop listening for content mount point events
 	FPackageName::OnContentPathMounted().RemoveAll(this);
 	FPackageName::OnContentPathDismounted().RemoveAll(this);
@@ -818,49 +787,74 @@ UAssetRegistryImpl& UAssetRegistryImpl::Get()
 	return static_cast<UAssetRegistryImpl&>(*UE::AssetRegistry::Private::IAssetRegistrySingleton::Singleton);
 }
 
-void UAssetRegistryImpl::SearchAllAssets(bool bSynchronousSearch)
+void UAssetRegistryImpl::ConstructGatherer()
 {
-	// Mark the time before the first search started
-	FullSearchStartTime = FPlatformTime::Seconds();
+	if (GlobalGatherer.IsValid())
+	{
+		return;
+	}
 
-	// Figure out what all of the root asset directories are.  This will include Engine content, Game content, but also may include
-	// mounted content directories for one or more plugins.	Also keep in mind that plugins may become loaded later on.  We'll listen
-	// for that via a delegate, and add those directories to scan later as them come in.
-	TArray<FString> PathsToSearch;
-	FPackageName::QueryRootContentPaths(PathsToSearch);
+	TArray<FString> BlacklistPaths;
+	TArray<FString> BlacklistContentSubPaths;
+	if (FConfigFile* EngineIni = GConfig->FindConfigFile(GEngineIni))
+	{
+		EngineIni->GetArray(TEXT("AssetRegistry"), TEXT("BlacklistPackagePathScanFilters"), BlacklistPaths);
+		EngineIni->GetArray(TEXT("AssetRegistry"), TEXT("BlacklistContentSubPathScanFilters"), BlacklistContentSubPaths);
+	}
 
-	// Start the asset search (synchronous in commandlets)
-	if (bSynchronousSearch)
+	bool bIsSynchronous = IsRunningGame();
+	GlobalGatherer = MakeUnique<FAssetDataGatherer>(BlacklistPaths, BlacklistContentSubPaths, bIsSynchronous);
+	RegisterReadOfScriptPackages();
+}
+
+void UAssetRegistryImpl::SearchAllAssetsInitialAsync()
+{
+	bInitialSearchStarted = true;
+	bInitialSearchCompleted = false;
+	FullSearchStartTime = FPlatformTime::Seconds();
+	SearchAllAssets(false);
+}
+
+void UAssetRegistryImpl::SearchAllAssets(bool bSynchronousSearch)
+{
+	ConstructGatherer();
+	FAssetDataGatherer& Gatherer = *GlobalGatherer;
+	if (Gatherer.IsSynchronous())
 	{
-#if WITH_EDITOR
-		if (IsLoadingAssets())
-		{
-			// Force a flush of the current gatherer instead
-			UE_LOG(LogAssetRegistry, Log, TEXT("Flushing asset discovery search because of synchronous request, this can take several seconds..."));
+		UE_CLOG(!bSynchronousSearch, LogAssetRegistry, Warning, TEXT("SearchAllAssets: Gatherer is in synchronous mode; forcing bSynchronousSearch=true."));
+		bSynchronousSearch = true;
+	}
 
-			WaitForCompletion();
-		}
-		else
-#endif
-		{
-			const bool bForceRescan = false;
-			ScanPathsAndFilesSynchronous(PathsToSearch, TArray<FString>(), BlacklistScanFilters, bForceRescan, EAssetDataCacheMode::UseMonolithicCache);
-		}
+	Gatherer.SetUseMonolithicCache(true);
 
+	// Add all existing mountpoints to the GlobalGatherer
+	// This will include Engine content, Game content, but also may include mounted content directories for one or more plugins.
+	TArray<FString> PackagePathsToSearch;
+	FPackageName::QueryRootContentPaths(PackagePathsToSearch);
+	for (const FString& PackagePath : PackagePathsToSearch)
+	{
+		const FString& MountLocalPath = FPackageName::LongPackageNameToFilename(PackagePath);
+		Gatherer.AddMountPoint(MountLocalPath, PackagePath);
+		Gatherer.SetIsWhitelisted(MountLocalPath, true);
+	}
+	bSearchAllAssets = true; // Mark that future mounts and directories should be scanned
 
+	if (bSynchronousSearch)
+	{
+		Gatherer.WaitForIdle();
+		TickGatherer(-1.f);
 #if WITH_EDITOR
-		if (IsRunningCommandlet())
+		if (!bInitialSearchStarted)
 		{
-			// Update redirectors
+			// We have a contract that we call UpdateRedirectCollector after the call to SearchAllAssets completes.
+			// If we ran the initial async call asynchronously it is done in TickGatherer; for later synchronous calls it is done here
 			UpdateRedirectCollector();
 		}
 #endif
 	}
-	else if (!BackgroundAssetSearch.IsValid())
+	else
 	{
-		// if the BackgroundAssetSearch is already valid then we have already called it before
-		BackgroundAssetSearch = MakeShared<FAssetDataGatherer>(PathsToSearch, TArray<FString>(), BlacklistScanFilters, bSynchronousSearch, EAssetDataCacheMode::UseMonolithicCache);
-		RegisterReadOfScriptPackages();
+		Gatherer.StartAsync();
 	}
 }
 
@@ -1568,19 +1562,6 @@ void UAssetRegistryImpl::RunAssetsThroughFilterImpl(TArray<FAssetData>& AssetDat
 	}
 }
 
-void UAssetRegistryImpl::AddSubContentBlacklist(const FString& InMount)
-{
-	for (const FString& SubContentPath : BlacklistContentSubPaths)
-	{
-		BlacklistScanFilters.AddUnique(InMount / SubContentPath);
-	}
-
-	if (BackgroundAssetSearch.IsValid())
-	{
-		BackgroundAssetSearch->SetBlacklistScanFilters(BlacklistScanFilters);
-	}
-}
-
 void UAssetRegistryImpl::ExpandRecursiveFilter(const FARFilter& InFilter, FARFilter& ExpandedFilter) const
 {
 	FARCompiledFilter CompiledFilter;
@@ -1736,11 +1717,17 @@ void UAssetRegistryImpl::PrioritizeAssetInstall(const FAssetData& AssetData) con
 
 bool UAssetRegistryImpl::AddPath(const FString& PathToAdd)
 {
-	if (AssetDataDiscoveryUtil::PassesScanFilters(BlacklistScanFilters, PathToAdd))
+	bool bBlacklisted = false;
+	// If no GlobalGatherer, then we are in the game or non-cook commandlet and we do not implement blacklisting
+	if (GlobalGatherer.IsValid())
 	{
-		return AddAssetPath(FName(*PathToAdd));
+		bBlacklisted = GlobalGatherer->IsBlacklisted(PathToAdd);
 	}
-	return false;
+	if (bBlacklisted)
+	{
+		return false;
+	}
+	return AddAssetPath(FName(*PathToAdd));
 }
 
 bool UAssetRegistryImpl::RemovePath(const FString& PathToRemove)
@@ -1760,52 +1747,41 @@ bool UAssetRegistryImpl::PathExists(const FName PathToTest) const
 
 void UAssetRegistryImpl::ScanPathsSynchronous(const TArray<FString>& InPaths, bool bForceRescan, bool bIgnoreBlackListScanFilters)
 {
-	// When ignoring blacklist scan filters, and forcing a rescan of one of the blacklisted folders, remove that folder from the blacklist.
-	// That way, we can now receive assets events such as added / removed / updated.
-	if (bForceRescan && bIgnoreBlackListScanFilters)
-	{
-		for (const FString& Path : InPaths)
-		{
-			for (int i = 0; i < BlacklistScanFilters.Num(); ++i)
-			{
-				if (Path.StartsWith(BlacklistScanFilters[i]))
-				{
-					BlacklistScanFilters.RemoveAt(i--);
-				}
-			}
-		}
-
-		if (BackgroundAssetSearch.IsValid())
-		{
-			BackgroundAssetSearch->SetBlacklistScanFilters(BlacklistScanFilters);
-		}
-	}
-
-	ScanPathsAndFilesSynchronous(InPaths, TArray<FString>(), BlacklistScanFilters, bForceRescan, EAssetDataCacheMode::UseModularCache);
+	ScanPathsSynchronousInternal(InPaths, TArray<FString>(), bForceRescan, bIgnoreBlackListScanFilters, nullptr /* OutFindAssets */);
 }
 
 void UAssetRegistryImpl::ScanFilesSynchronous(const TArray<FString>& InFilePaths, bool bForceRescan)
 {
-	ScanPathsAndFilesSynchronous(TArray<FString>(), InFilePaths, BlacklistScanFilters, bForceRescan, EAssetDataCacheMode::UseModularCache);
+	return ScanPathsSynchronousInternal(TArray<FString>(), InFilePaths, bForceRescan, false /* bIgnoreBlackListScanFilters */, nullptr /* OutFindAssets */);
 }
 
 void UAssetRegistryImpl::PrioritizeSearchPath(const FString& PathToPrioritize)
 {
-	// Prioritize the background search
-	if (BackgroundAssetSearch.IsValid())
+	if (!GlobalGatherer.IsValid())
 	{
-		BackgroundAssetSearch->PrioritizeSearchPath(PathToPrioritize);
+		return;
 	}
+	GlobalGatherer->PrioritizeSearchPath(PathToPrioritize);
 
 	// Also prioritize the queue of background search results
-	BackgroundAssetResults.Prioritize([&PathToPrioritize](const FAssetData* BackgroundAssetResult)
+	int32 FirstNonPriorityIndex = 0;
+	for (int Index = 0; Index < BackgroundAssetResults.Num(); ++Index)
 	{
-		return BackgroundAssetResult && BackgroundAssetResult->PackagePath.ToString().StartsWith(PathToPrioritize);
-	});
-	BackgroundPathResults.Prioritize([&PathToPrioritize](const FString& BackgroundPathResult)
+		FAssetData* PriorityElement = BackgroundAssetResults[Index];
+		if (PriorityElement && PriorityElement->PackagePath.ToString().StartsWith(PathToPrioritize))
+		{
+			Swap(BackgroundAssetResults[FirstNonPriorityIndex++], BackgroundAssetResults[Index]);
+		}
+	}
+	FirstNonPriorityIndex = 0;
+	for (int Index = 0; Index < BackgroundPathResults.Num(); ++Index)
 	{
-		return BackgroundPathResult.StartsWith(PathToPrioritize);
-	});
+		FString& PriorityElement = BackgroundPathResults[Index];
+		if (PriorityElement.StartsWith(PathToPrioritize))
+		{
+			Swap(BackgroundPathResults[FirstNonPriorityIndex++], BackgroundPathResults[Index]);
+		}
+	}
 }
 
 void UAssetRegistryImpl::AssetCreated(UObject* NewAsset)
@@ -1931,6 +1907,16 @@ bool UAssetRegistryImpl::IsLoadingAssets() const
 
 void UAssetRegistryImpl::Tick(float DeltaTime)
 {
+	TickGatherer(DeltaTime);
+}
+
+void UAssetRegistryImpl::TickGatherer(float DeltaTime, TArray<FName>* OutFoundAssets)
+{
+	if (!GlobalGatherer.IsValid())
+	{
+		return;
+	}
+
 	double TickStartTime = FPlatformTime::Seconds();
 
 	if (DeltaTime < 0)
@@ -1938,10 +1924,6 @@ void UAssetRegistryImpl::Tick(float DeltaTime)
 		// Force a full flush
 		TickStartTime = -1;
 	}
-	
-	bool bOldTemporaryCachingMode = GetTemporaryCachingMode();
-	SetTemporaryCachingMode(true);
-	ON_SCOPE_EXIT { SetTemporaryCachingMode(bOldTemporaryCachingMode); };
 
 	// Gather results from the background search
 	bool bIsSearching = false;
@@ -1949,17 +1931,18 @@ void UAssetRegistryImpl::Tick(float DeltaTime)
 	int32 NumFilesToSearch = 0;
 	int32 NumPathsToSearch = 0;
 	bool bIsDiscoveringFiles = false;
-	if (BackgroundAssetSearch.IsValid())
-	{
-		bIsSearching = BackgroundAssetSearch->GetAndTrimSearchResults(BackgroundAssetResults, BackgroundPathResults, BackgroundDependencyResults, BackgroundCookedPackageNamesWithoutAssetDataResults, SearchTimes, NumFilesToSearch, NumPathsToSearch, bIsDiscoveringFiles);
-	}
-
+	GlobalGatherer->GetAndTrimSearchResults(bIsSearching, BackgroundAssetResults, BackgroundPathResults, BackgroundDependencyResults, BackgroundCookedPackageNamesWithoutAssetDataResults, SearchTimes, NumFilesToSearch, NumPathsToSearch, bIsDiscoveringFiles);
 	// Report the search times
 	for (int32 SearchTimeIdx = 0; SearchTimeIdx < SearchTimes.Num(); ++SearchTimeIdx)
 	{
 		UE_LOG(LogAssetRegistry, Verbose, TEXT("### Background search completed in %0.4f seconds"), SearchTimes[SearchTimeIdx]);
 	}
 
+	// Go into temporary caching mode to speed up class queries while adding assets
+	bool bOldTemporaryCachingMode = GetTemporaryCachingMode();
+	SetTemporaryCachingMode(true);
+	ON_SCOPE_EXIT{ SetTemporaryCachingMode(bOldTemporaryCachingMode); };
+
 	// Add discovered paths
 	if (BackgroundPathResults.Num())
 	{
@@ -1975,6 +1958,15 @@ void UAssetRegistryImpl::Tick(float DeltaTime)
 		{
 			AmortizeStartTime = FPlatformTime::Seconds();
 		}
+		if (OutFoundAssets)
+		{
+			OutFoundAssets->Reset();
+			OutFoundAssets->Reserve(BackgroundAssetResults.Num());
+			for (FAssetData* AssetData : BackgroundAssetResults)
+			{
+				OutFoundAssets->Add(AssetData->ObjectPath);
+			}
+		}
 
 		AssetSearchDataGathered(TickStartTime, BackgroundAssetResults);
 
@@ -2016,7 +2008,10 @@ void UAssetRegistryImpl::Tick(float DeltaTime)
 	}
 
 	// If completing an initial search, refresh the content browser
-	if (!bIsSearching && NumPending == 0)
+	bool bIsIdle = !bIsSearching && NumPending == 0;
+	bool bBecameIdle = !bGatherIdle && bIsIdle;
+	bGatherIdle = bIsIdle;
+	if (bIsIdle)
 	{
 		HighestPending = 0;
 
@@ -2033,13 +2028,13 @@ void UAssetRegistryImpl::Tick(float DeltaTime)
 
 			FileLoadedEvent.Broadcast();
 		}
+	}
 #if WITH_EDITOR
-		else if (bUpdateDiskCacheAfterLoad)
-		{
-			ProcessLoadedAssetsToUpdateCache(TickStartTime);
-		}
-#endif
+	if (bUpdateDiskCacheAfterLoad && (bIsIdle || TickStartTime < 0))
+	{
+		ProcessLoadedAssetsToUpdateCache(TickStartTime, bBecameIdle);
 	}
+#endif
 }
 
 void UAssetRegistryImpl::Serialize(FArchive& Ar)
@@ -2108,7 +2103,7 @@ uint32 UAssetRegistryImpl::GetAllocatedSize(bool bLogDetailed) const
 	uint32 StateSize = State.GetAllocatedSize(bLogDetailed);
 
 	uint32 StaticSize = sizeof(UAssetRegistryImpl) + CachedEmptyPackages.GetAllocatedSize() + CachedBPInheritanceMap.GetAllocatedSize()  + ClassGeneratorNames.GetAllocatedSize() + OnDirectoryChangedDelegateHandles.GetAllocatedSize();
-	uint32 SearchSize = BackgroundAssetResults.GetAllocatedSize() + BackgroundPathResults.GetAllocatedSize() + BackgroundDependencyResults.GetAllocatedSize() + BackgroundCookedPackageNamesWithoutAssetDataResults.GetAllocatedSize() + SynchronouslyScannedPathsAndFiles.GetAllocatedSize() + CachedPathTree.GetAllocatedSize();
+	uint32 SearchSize = BackgroundAssetResults.GetAllocatedSize() + BackgroundPathResults.GetAllocatedSize() + BackgroundDependencyResults.GetAllocatedSize() + BackgroundCookedPackageNamesWithoutAssetDataResults.GetAllocatedSize() + CachedPathTree.GetAllocatedSize();
 
 	if (bIsTempCachingEnabled && !bIsTempCachingAlwaysEnabled)
 	{
@@ -2117,16 +2112,10 @@ uint32 UAssetRegistryImpl::GetAllocatedSize(bool bLogDetailed) const
 		UE_LOG(LogAssetRegistry, Warning, TEXT("Asset Registry Temp caching enabled, wasting memory: %dk"), TempCacheMem / 1024);
 	}
 
-	StaticSize += BlacklistScanFilters.GetAllocatedSize();
-	for (const FString& ScanFilter : BlacklistScanFilters)
-	{
-		StaticSize += ScanFilter.GetAllocatedSize();
-	}
-
-	StaticSize += BlacklistContentSubPaths.GetAllocatedSize();
-	for (const FString& ScanFilter : BlacklistContentSubPaths)
+	if (GlobalGatherer.IsValid())
 	{
-		StaticSize += ScanFilter.GetAllocatedSize();
+		SearchSize += sizeof(*GlobalGatherer);
+		SearchSize += GlobalGatherer->GetAllocatedSize();
 	}
 
 	StaticSize += SerializationOptions.CookFilterlistTagsByClass.GetAllocatedSize();
@@ -2187,136 +2176,156 @@ const TSet<FName>& UAssetRegistryImpl::GetCachedEmptyPackages() const
 	return CachedEmptyPackages;
 }
 
-void UAssetRegistryImpl::ScanPathsAndFilesSynchronous(const TArray<FString>& InPaths, const TArray<FString>& InSpecificFiles, const TArray<FString>& InBlacklistScanFilters, bool bForceRescan, EAssetDataCacheMode AssetDataCacheMode)
-{
-	ScanPathsAndFilesSynchronous(InPaths, InSpecificFiles, InBlacklistScanFilters, bForceRescan, AssetDataCacheMode, nullptr, nullptr);
-}
-
-void UAssetRegistryImpl::ScanPathsAndFilesSynchronous(const TArray<FString>& InPaths, const TArray<FString>& InSpecificFiles, const TArray<FString>& InBlacklistScanFilters, bool bForceRescan, EAssetDataCacheMode AssetDataCacheMode, TArray<FName>* OutFoundAssets, TArray<FName>* OutFoundPaths)
+void UAssetRegistryImpl::ScanPathsSynchronousInternal(const TArray<FString>& InDirs, const TArray<FString>& InFiles, bool bForceRescan, bool bIgnoreBlackListScanFilters, TArray<FName>* OutFoundAssets)
 {
 	LLM_SCOPE(ELLMTag::AssetRegistry);
 	const double SearchStartTime = FPlatformTime::Seconds();
+	if (OutFoundAssets)
+	{
+		OutFoundAssets->Empty();
+	}
+	if (bIgnoreBlackListScanFilters && !bForceRescan)
+	{
+		// This restriction is necessary because we have not yet implemented some of the required behavior to handle bIgnoreBlackListScanFilters without bForceRescan;
+		// For skipping of directories that we have already scanned, we would have to check whether the directory has been set to be monitored with the proper flag (ignore blacklist or not)
+		// rather than just checking whether it has been set to be monitored at all
+		UE_LOG(LogAssetRegistry, Warning, TEXT("ScanPathsSynchronous: bIgnoreBlacklistScanFilters==true is only valid when bForceRescan==true. Setting bForceRescan=true."));
+		bForceRescan = true;
+	}
 
-	// Only scan paths that were not previously synchronously scanned, unless we were asked to force rescan.
-	TArray<FString> PathsToScan;
-	TArray<FString> FilesToScan;
-	bool bPathsRemoved = false;
-
-	for (const FString& Path : InPaths)
+	TArray<FString> PackageDirs;
+	TArray<FString> LocalDirs;
+	TArray<FString> PackageFiles;
+	TArray<FString> LocalFiles;
+	FString LocalPath;
+	FString PackageName;
+	FString Extension;
+	FPackageName::EFlexNameType FlexNameType;
+	LocalFiles.Reserve(InFiles.Num());
+	PackageFiles.Reserve(InFiles.Num());
+	for (const FString& InFile : InFiles)
 	{
-		bool bAlreadyScanned = false;
-		FString PathWithSlash = Path;
-		if (!PathWithSlash.EndsWith(TEXT("/"), ESearchCase::CaseSensitive))
+		if (!FPackageName::TryConvertToMountedPath(InFile, &LocalPath, &PackageName, nullptr, nullptr, &Extension, &FlexNameType))
 		{
-			// Add / if it's missing so the substring check is safe
-			PathWithSlash += TEXT("/");
+			UE_LOG(LogAssetRegistry, Warning, TEXT("ScanPathsSynchronous: %s is not in a mounted path, will not scan."), *InFile);
+			continue;
 		}
-
-		// Check that it starts with /
-		for (const FString& ScannedPath : SynchronouslyScannedPathsAndFiles)
+		if (Extension.IsEmpty())
 		{
-			if (PathWithSlash.StartsWith(ScannedPath))
+			// The empty extension is not a valid Package extension; it might exist, but we will pay the price to check it
+			if (!IFileManager::Get().FileExists(*LocalPath))
 			{
-				bAlreadyScanned = true;
-				break;
+				// Find the extension
+				FPackagePath PackagePath = FPackagePath::FromLocalPath(LocalPath);
+				if (!FPackageName::DoesPackageExist(PackagePath, &PackagePath))
+				{
+					UE_LOG(LogAssetRegistry, Warning, TEXT("ScanPathsSynchronous: Package %s does not exist, will not scan."), *InFile);
+					continue;
+				}
+				Extension = LexToString(PackagePath.GetHeaderExtension());
 			}
 		}
-
-		if (bForceRescan || !bAlreadyScanned)
-		{
-			PathsToScan.Add(Path);
-			SynchronouslyScannedPathsAndFiles.Add(PathWithSlash);
-		}
-		else
-		{
-			bPathsRemoved = true;
-		}
+		LocalFiles.Add(LocalPath + Extension);
+		PackageFiles.Add(PackageName);
 	}
-
-	for (const FString& SpecificFile : InSpecificFiles)
+	LocalDirs.Reserve(InDirs.Num());
+	PackageDirs.Reserve(InDirs.Num());
+	for (const FString& InDir : InDirs)
 	{
-		if (bForceRescan || !SynchronouslyScannedPathsAndFiles.Contains(SpecificFile))
-		{
-			FilesToScan.Add(SpecificFile);
-			SynchronouslyScannedPathsAndFiles.Add(SpecificFile);
-		}
-		else
+		if (!FPackageName::TryConvertToMountedPath(InDir, &LocalPath, &PackageName, nullptr, nullptr, &Extension, &FlexNameType))
 		{
-			bPathsRemoved = true;
+			UE_LOG(LogAssetRegistry, Warning, TEXT("ScanPathsSynchronous: %s is not in a mounted path, will not scan."), *InDir);
+			continue;
 		}
+		LocalDirs.Add(LocalPath + Extension);
+		PackageDirs.Add(PackageName + Extension);
 	}
 
-	// If we removed paths, we can't use the monolithic cache as this will replace it with invalid data
-	if (AssetDataCacheMode == EAssetDataCacheMode::UseMonolithicCache && bPathsRemoved)
-	{
-		AssetDataCacheMode = EAssetDataCacheMode::UseModularCache;
-	}
+	ConstructGatherer();
+	FAssetDataGatherer& Gatherer = *GlobalGatherer;
 
-	if (PathsToScan.Num() > 0 || FilesToScan.Num() > 0)
+	// Add a cache file for any not-yet-scanned dirs
+	TArray<FString> CacheFilePackagePaths;
+	if (!bForceRescan && Gatherer.IsCacheEnabled())
 	{
-		// Start the sync asset search
-		FAssetDataGatherer AssetSearch(PathsToScan, FilesToScan, InBlacklistScanFilters, /*bSynchronous=*/true, AssetDataCacheMode);
-
-		// Get the search results
-		TBackgroundGatherResults<FAssetData*> AssetResults;
-		TBackgroundGatherResults<FString> PathResults;
-		TBackgroundGatherResults<FPackageDependencyData> DependencyResults;
-		TBackgroundGatherResults<FString> CookedPackageNamesWithoutAssetDataResults;
-		TArray<double> SearchTimes;
-		int32 NumFilesToSearch = 0;
-		int32 NumPathsToSearch = 0;
-		bool bIsDiscoveringFiles = false;
-		AssetSearch.GetAndTrimSearchResults(AssetResults, PathResults, DependencyResults, CookedPackageNamesWithoutAssetDataResults, SearchTimes, NumFilesToSearch, NumPathsToSearch, bIsDiscoveringFiles);
-
-		if (OutFoundAssets)
+		for (int n = 0; n < LocalDirs.Num(); ++n)
 		{
-			OutFoundAssets->Reserve(OutFoundAssets->Num() + AssetResults.Num());
-			for (int32 i = 0; i < AssetResults.Num(); ++i)
+			if (!Gatherer.IsWhitelisted(LocalDirs[n]))
 			{
-				OutFoundAssets->Add(AssetResults[i]->ObjectPath);
+				CacheFilePackagePaths.Add(PackageDirs[n]);
 			}
 		}
+	}
+	TArray<FString> LocalPaths;
+	LocalPaths.Reserve(LocalFiles.Num() + LocalDirs.Num());
+	LocalPaths.Append(MoveTemp(LocalDirs));
+	LocalPaths.Append(MoveTemp(LocalFiles));
+	if (LocalPaths.IsEmpty())
+	{
+		return;
+	}
+	Gatherer.AddRequiredMountPoints(LocalPaths);
 
-		if (OutFoundPaths)
-		{
-			OutFoundPaths->Reserve(OutFoundPaths->Num() + PathResults.Num());
-			for (int32 i = 0; i < PathResults.Num(); ++i)
-			{
-				OutFoundPaths->Add(*PathResults[i]);
-			}
-		}
-
-		// Cache the search results
-		const int32 NumResults = AssetResults.Num();
-		AssetSearchDataGathered(-1, AssetResults);
-		PathDataGathered(-1, PathResults);
-		DependencyDataGathered(-1, DependencyResults);
-		CookedPackageNamesWithoutAssetDataGathered(-1, CookedPackageNamesWithoutAssetDataResults);
-
-#if WITH_EDITOR
-		if (bUpdateDiskCacheAfterLoad && bInitialSearchCompleted)
-		{
-			ProcessLoadedAssetsToUpdateCache(-1);
-		}
-#endif
-
-		// Log stats
-		TArray<FString> LogPathsAndFilenames = PathsToScan;
-		LogPathsAndFilenames.Append(FilesToScan);
+	FString CacheFilename;
+	if (!CacheFilePackagePaths.IsEmpty())
+	{
+		CacheFilename = Gatherer.GetCacheFilename(CacheFilePackagePaths);
+		Gatherer.LoadCacheFile(CacheFilename);
+	}
 
-		const FString& Path = LogPathsAndFilenames[0];
-		FString PathsString;
-		if (LogPathsAndFilenames.Num() > 1)
-		{
-			PathsString = FString::Printf(TEXT("'%s' and %d other paths/filenames"), *Path, LogPathsAndFilenames.Num() - 1);
-		}
-		else
-		{
-			PathsString = FString::Printf(TEXT("'%s'"), *Path);
-		}
+	TArray<FName> FoundAssetsBuffer;
+	TArray<FName>* FoundAssets = OutFoundAssets;
+	if (!FoundAssets)
+	{
+		FoundAssets = &FoundAssetsBuffer;
+	}
+	Gatherer.ScanPathsSynchronous(LocalPaths, bForceRescan, bIgnoreBlackListScanFilters, CacheFilename, PackageDirs);
+	TickGatherer(-1.f, FoundAssets);
+	// The gatherer may have added other assets that were scanned as part of the ongoing background scan; remove any assets that were not in the requested paths
+	if (OutFoundAssets)
+	{
+		FoundAssets->RemoveAll([&PackageFiles, &PackageDirs](FName ObjectPath)
+			{
+				TStringBuilder<128> ObjectPathStr;
+				ObjectPath.ToString(ObjectPathStr);
+				FStringView UnusedClassName;
+				FStringView PackageName;
+				FStringView UnusedObjectName;
+				FStringView UnusedSubObjectName;
+				FPackageName::SplitFullObjectPath(ObjectPathStr, UnusedClassName, PackageName, UnusedObjectName, UnusedSubObjectName);
+				bool bIsInRequestedPaths = false;
+				for (const FString& RequestedPackageDir : PackageDirs)
+				{
+					if (FPathViews::IsParentPathOf(RequestedPackageDir, PackageName))
+					{
+						bIsInRequestedPaths = true;
+						break;
+					}
+				}
+				for (const FString& RequestedPackageFile : PackageFiles)
+				{
+					if (PackageName.Equals(RequestedPackageFile, ESearchCase::IgnoreCase))
+					{
+						bIsInRequestedPaths = true;
+						break;
+					}
+				}
+				return !bIsInRequestedPaths;
+			});
+	}
 
-		UE_LOG(LogAssetRegistry, Verbose, TEXT("ScanPathsSynchronous completed scanning %s to find %d assets in %0.4f seconds"), *PathsString, NumResults, FPlatformTime::Seconds() - SearchStartTime);
+	// Log stats
+	FString PathsString;
+	if (LocalPaths.Num() > 1)
+	{
+		PathsString = FString::Printf(TEXT("'%s' and %d other paths"), *LocalPaths[0], LocalPaths.Num() - 1);
+	}
+	else
+	{
+		PathsString = FString::Printf(TEXT("'%s'"), *LocalPaths[0]);
 	}
+
+	UE_LOG(LogAssetRegistry, Verbose, TEXT("ScanPathsSynchronous completed scanning %s to find %d assets in %0.4f seconds"), *PathsString, FoundAssets->Num(), FPlatformTime::Seconds() - SearchStartTime);
 }
 
 bool UAssetRegistryImpl::IsPathMounted(const FString& Path, const TSet<FString>& MountPointsNoTrailingSlashes, FString& StringBuffer) const
@@ -2342,7 +2351,7 @@ bool UAssetRegistryImpl::IsPathMounted(const FString& Path, const TSet<FString>&
 	return false;
 }
 
-void UAssetRegistryImpl::AssetSearchDataGathered(const double TickStartTime, TBackgroundGatherResults<FAssetData*>& AssetResults)
+void UAssetRegistryImpl::AssetSearchDataGathered(const double TickStartTime, TRingBuffer<FAssetData*>& AssetResults)
 {
 	const bool bFlushFullBuffer = TickStartTime < 0;
 
@@ -2362,8 +2371,7 @@ void UAssetRegistryImpl::AssetSearchDataGathered(const double TickStartTime, TBa
 	// Add the found assets
 	while (AssetResults.Num() > 0)
 	{
-		FAssetData*& BackgroundResult = AssetResults.Pop();
-
+		FAssetData* BackgroundResult = AssetResults.PopFrontValue();
 		CA_ASSUME(BackgroundResult);
 
 		// Try to update any asset data that may already exist
@@ -2387,8 +2395,14 @@ void UAssetRegistryImpl::AssetSearchDataGathered(const double TickStartTime, TBa
 			// If this ensure fires then we've somehow processed the same result more than once, and that should never happen
 			if (ensure(AssetData != BackgroundResult))
 			{
-				// The asset exists in the cache, update it
-				UpdateAssetData(AssetData, *BackgroundResult);
+				// If the current AssetData came from a loaded asset, don't overwrite it with the new one from disk; loaded asset is more authoritative because it has run the postload steps
+#if WITH_EDITOR
+				if (!AssetDataObjectPathsUpdatedOnLoad.Contains(BackgroundResult->ObjectPath))
+#endif
+				{
+					// The asset exists in the cache from disk and has not yet been loaded into memory, update it with the new background data
+					UpdateAssetData(AssetData, *BackgroundResult);
+				}
 
 				// Delete the result that was originally created by an FPackageReader
 				delete BackgroundResult;
@@ -2416,12 +2430,9 @@ void UAssetRegistryImpl::AssetSearchDataGathered(const double TickStartTime, TBa
 			return;
 		}
 	}
-
-	// Trim the results array
-	AssetResults.Trim();
 }
 
-void UAssetRegistryImpl::PathDataGathered(const double TickStartTime, TBackgroundGatherResults<FString>& PathResults)
+void UAssetRegistryImpl::PathDataGathered(const double TickStartTime, TRingBuffer<FString>& PathResults)
 {
 	const bool bFlushFullBuffer = TickStartTime < 0;
 
@@ -2436,7 +2447,7 @@ void UAssetRegistryImpl::PathDataGathered(const double TickStartTime, TBackgroun
 
 	while (PathResults.Num() > 0)
 	{
-		const FString& Path = PathResults.Pop();
+		FString Path = PathResults.PopFrontValue();
 
 		// Skip stale results caused by mount then unmount of a path within short period.
 		if (!bVerifyMountPointAfterGather || IsPathMounted(Path, MountPoints, PackageRoot))
@@ -2450,19 +2461,16 @@ void UAssetRegistryImpl::PathDataGathered(const double TickStartTime, TBackgroun
 			return;
 		}
 	}
-
-	// Trim the results array
-	PathResults.Trim();
 }
 
-void UAssetRegistryImpl::DependencyDataGathered(const double TickStartTime, TBackgroundGatherResults<FPackageDependencyData>& DependsResults)
+void UAssetRegistryImpl::DependencyDataGathered(const double TickStartTime, TRingBuffer<FPackageDependencyData>& DependsResults)
 {
 	using namespace UE::AssetRegistry;
 	const bool bFlushFullBuffer = TickStartTime < 0;
 
 	while (DependsResults.Num() > 0)
 	{
-		FPackageDependencyData& Result = DependsResults.Pop();
+		FPackageDependencyData Result = DependsResults.PopFrontValue();
 
 		// Update package data
 		FAssetPackageData* PackageData = State.CreateOrGetAssetPackageData(Result.PackageName);
@@ -2596,12 +2604,9 @@ void UAssetRegistryImpl::DependencyDataGathered(const double TickStartTime, TBac
 			return;
 		}
 	}
-
-	// Trim the results array
-	DependsResults.Trim();
 }
 
-void UAssetRegistryImpl::CookedPackageNamesWithoutAssetDataGathered(const double TickStartTime, TBackgroundGatherResults<FString>& CookedPackageNamesWithoutAssetDataResults)
+void UAssetRegistryImpl::CookedPackageNamesWithoutAssetDataGathered(const double TickStartTime, TRingBuffer<FString>& CookedPackageNamesWithoutAssetDataResults)
 {
 	const bool bFlushFullBuffer = TickStartTime < 0;
 
@@ -2626,7 +2631,7 @@ void UAssetRegistryImpl::CookedPackageNamesWithoutAssetDataGathered(const double
 		{
 			// If this data is cooked and it we couldn't find any asset in its export table then try to load the entire package 
 			// Loading the entire package will make all of its assets searchable through the in-memory scanning performed by GetAsseets
-			const FString& BackgroundResult = CookedPackageNamesWithoutAssetDataResults.Pop();
+			FString BackgroundResult = CookedPackageNamesWithoutAssetDataResults.PopFrontValue();
 			LoadPackage(nullptr, *BackgroundResult, 0);
 
 			// Check to see if we have run out of time in this tick
@@ -2643,9 +2648,6 @@ void UAssetRegistryImpl::CookedPackageNamesWithoutAssetDataGathered(const double
 		// kill performance. We need a better way.
 		CookedPackageNamesWithoutAssetDataResults.Empty();
 	}
-
-	// Trim the results array
-	CookedPackageNamesWithoutAssetDataResults.Trim();
 }
 
 void UAssetRegistryImpl::AddEmptyPackage(FName PackageName)
@@ -2821,30 +2823,6 @@ void UAssetRegistryImpl::RemovePackageData(const FName PackageName)
 	}
 }
 
-void UAssetRegistryImpl::AddPathToSearch(const FString& Path)
-{
-	if (BackgroundAssetSearch.IsValid())
-	{
-		BackgroundAssetSearch->AddPathToSearch(Path);
-	}
-	else if (GIsEditor)
-	{
-		ScanPathsSynchronous({Path});
-	}
-}
-
-void UAssetRegistryImpl::AddFilesToSearch(const TArray<FString>& Files)
-{
-	if (BackgroundAssetSearch.IsValid())
-	{
-		BackgroundAssetSearch->AddFilesToSearch(Files);
-	}
-	else if (GIsEditor)
-	{
-		ScanFilesSynchronous(Files);
-	}
-}
-
 #if WITH_EDITOR
 
 void UAssetRegistryImpl::OnDirectoryChanged(const TArray<FFileChangeData>& FileChanges)
@@ -2876,6 +2854,7 @@ void UAssetRegistryImpl::OnDirectoryChanged(const TArray<FFileChangeData>& FileC
 		}
 	}
 
+	TArray<FString> NewDirs;
 	TArray<FString> NewFiles;
 	TArray<FString> ModifiedFiles;
 
@@ -2918,11 +2897,10 @@ void UAssetRegistryImpl::OnDirectoryChanged(const TArray<FFileChangeData>& FileC
 			{
 			case FFileChangeData::FCA_Added:
 			{
-				if (FPaths::DirectoryExists(File) && LongPackageName != TEXT("/Game/Collections"))
+				if (FPaths::DirectoryExists(File))
 				{
-					AddPath(LongPackageName);
+					NewDirs.Add(File);
 					UE_LOG(LogAssetRegistry, Verbose, TEXT("Directory was added to content directory: %s"), *File);
-					AddPathToSearch(LongPackageName);
 				}
 				break;
 			}
@@ -2938,9 +2916,20 @@ void UAssetRegistryImpl::OnDirectoryChanged(const TArray<FFileChangeData>& FileC
 		}
 	}
 
-	if (NewFiles.Num())
+	if (NewFiles.Num() || NewDirs.Num())
 	{
-		AddFilesToSearch(NewFiles);
+		if (GlobalGatherer.IsValid())
+		{
+			for (FString& NewDir : NewDirs)
+			{
+				GlobalGatherer->OnDirectoryCreated(NewDir);
+			}
+			GlobalGatherer->OnFilesCreated(NewFiles);
+			if (GlobalGatherer->IsSynchronous())
+			{
+				ScanPathsSynchronousInternal(NewDirs, NewFiles, false /* bForceRescan */, false /* bIgnoreBlacklistScanFilters */, nullptr /* OutFoundAssets */);
+			}
+		}
 	}
 
 	ScanModifiedAssetFiles(ModifiedFiles);
@@ -2951,14 +2940,13 @@ void UAssetRegistryImpl::OnAssetLoaded(UObject *AssetLoaded)
 	LoadedAssetsToProcess.Add(AssetLoaded);
 }
 
-void UAssetRegistryImpl::ProcessLoadedAssetsToUpdateCache(const double TickStartTime)
+void UAssetRegistryImpl::ProcessLoadedAssetsToUpdateCache(const double TickStartTime, bool bUpdateDeferredList)
 {
 	LLM_SCOPE(ELLMTag::AssetRegistry);
-	check(bInitialSearchCompleted && bUpdateDiskCacheAfterLoad);
 
 	const bool bFlushFullBuffer = TickStartTime < 0;
 
-	if (bFlushFullBuffer)
+	if (bFlushFullBuffer || bUpdateDeferredList)
 	{
 		// Retry the previous failures on a flush
 		LoadedAssetsToProcess.Append(LoadedAssetsThatDidNotHaveCachedData);
@@ -3087,7 +3075,7 @@ void UAssetRegistryImpl::ScanModifiedAssetFiles(const TArray<FString>& InFilePat
 
 		// Re-scan and update the asset registry with the new asset data
 		TArray<FName> FoundAssets;
-		ScanPathsAndFilesSynchronous(TArray<FString>(), InFilePaths, BlacklistScanFilters, true, EAssetDataCacheMode::NoCache, &FoundAssets, nullptr);
+		ScanPathsSynchronousInternal(TArray<FString>(), InFilePaths, true /* bForceRescan */, false /* bIgnoreBlackListScanFilters */, &FoundAssets);
 
 		// Remove any assets that are no longer present in the package
 		for (const TArray<FAssetData*, TInlineAllocator<1>>& OldPackageAssets : ExistingFilesAssetData)
@@ -3105,8 +3093,6 @@ void UAssetRegistryImpl::ScanModifiedAssetFiles(const TArray<FString>& InFilePat
 
 void UAssetRegistryImpl::OnContentPathMounted(const FString& InAssetPath, const FString& FileSystemPath)
 {
-	AddSubContentBlacklist(InAssetPath);
-
 	// Sanitize
 	FString AssetPathWithTrailingSlash;
 	if (!InAssetPath.EndsWith(TEXT("/"), ESearchCase::CaseSensitive))
@@ -3122,8 +3108,18 @@ void UAssetRegistryImpl::OnContentPathMounted(const FString& InAssetPath, const
 	// Content roots always exist
 	AddPath(AssetPathWithTrailingSlash);
 
-	// Add this to our list of root paths to process
-	AddPathToSearch(AssetPathWithTrailingSlash);
+	if (GlobalGatherer.IsValid() && bSearchAllAssets)
+	{
+		if (GlobalGatherer->IsSynchronous())
+		{
+			ScanPathsSynchronous({ FileSystemPath });
+		}
+		else
+		{
+			GlobalGatherer->AddMountPoint(FileSystemPath, InAssetPath);
+			GlobalGatherer->SetIsWhitelisted(FileSystemPath, true);
+		}
+	}
 
 	// Listen for directory changes in this content path
 #if WITH_EDITOR
@@ -3164,6 +3160,11 @@ void UAssetRegistryImpl::OnContentPathDismounted(const FString& InAssetPath, con
 		AssetPathNoTrailingSlash.LeftChopInline(1, false);
 	}
 
+	if (GlobalGatherer.IsValid())
+	{
+		GlobalGatherer->RemoveMountPoint(FileSystemPath);
+	}
+
 	// Remove all cached assets found at this location
 	{
 		TArray<FAssetData*> AllAssetDataToRemove;
diff --git a/Engine/Source/Runtime/AssetRegistry/Private/AssetRegistry.h b/Engine/Source/Runtime/AssetRegistry/Private/AssetRegistry.h
--- a/Engine/Source/Runtime/AssetRegistry/Private/AssetRegistry.h
+++ b/Engine/Source/Runtime/AssetRegistry/Private/AssetRegistry.h
@@ -3,18 +3,22 @@
 #pragma once
 
 #include "CoreMinimal.h"
+
 #include "AssetRegistry/AssetData.h"
 #include "AssetRegistry/IAssetRegistry.h"
 #include "AssetRegistry/AssetRegistryState.h"
-#include "PathTree.h"
-#include "PackageDependencyData.h"
-#include "AssetDataGatherer.h"
-#include "BackgroundGatherResults.h"
+#include "Containers/RingBuffer.h"
 #include "ModuleDescriptor.h"
+#include "PackageDependencyData.h"
+#include "PathTree.h"
+#include "Templates/UniquePtr.h"
+
 #include "AssetRegistry.generated.h"
 
 class FDependsNode;
 struct FARFilter;
+class FAssetDataGatherer;
+class FPackageReader;
 
 /**
  * The AssetRegistry singleton gathers information about .uasset files in the background so things
@@ -28,6 +32,7 @@ class UAssetRegistryImpl : public UObject, public IAssetRegistry
 	GENERATED_BODY()
 public:
 	UAssetRegistryImpl(const FObjectInitializer& ObjectInitializer);
+	UAssetRegistryImpl(FVTableHelper& Helper);
 	virtual ~UAssetRegistryImpl();
 
 	/** Gets the asset registry singleton for asset registry module use */
@@ -167,21 +172,26 @@ private:
 	 */
 	void InternalRunAssetsThroughFilter(TArray<FAssetData>& AssetDataList, const FARFilter& Filter, bool bInclusive) const;
 
+	/* Construct the gatherer if it does not already exist */
+	void ConstructGatherer();
+	/**  Called to set up timing variables and launch the on-constructor SearchAllAssets async call */
+	void SearchAllAssetsInitialAsync();
+	/** Read any results from the gatherer, if it is running asynchronously */
+	void TickGatherer(float DeltaTime, TArray<FName>* OutFoundAssets=nullptr);
 	/** Internal handler for ScanPathsSynchronous */
-	void ScanPathsAndFilesSynchronous(const TArray<FString>& InPaths, const TArray<FString>& InSpecificFiles, const TArray<FString>& InBlacklistScanFilters, bool bForceRescan, EAssetDataCacheMode AssetDataCacheMode);
-	void ScanPathsAndFilesSynchronous(const TArray<FString>& InPaths, const TArray<FString>& InSpecificFiles, const TArray<FString>& InBlacklistScanFilters, bool bForceRescan, EAssetDataCacheMode AssetDataCacheMode, TArray<FName>* OutFoundAssets, TArray<FName>* OutFoundPaths);
+	void ScanPathsSynchronousInternal(const TArray<FString>& InDirs, const TArray<FString>& InFiles, bool bForceRescan, bool bIgnoreBlackListScanFilters, TArray<FName>* OutFoundAssets);
 
 	/** Called every tick to when data is retrieved by the background asset search. If TickStartTime is < 0, the entire list of gathered assets will be cached. Also used in sychronous searches */
-	void AssetSearchDataGathered(const double TickStartTime, TBackgroundGatherResults<FAssetData*>& AssetResults);
+	void AssetSearchDataGathered(const double TickStartTime, TRingBuffer<FAssetData*>& AssetResults);
 
 	/** Called every tick when data is retrieved by the background path search. If TickStartTime is < 0, the entire list of gathered assets will be cached. Also used in sychronous searches */
-	void PathDataGathered(const double TickStartTime, TBackgroundGatherResults<FString>& PathResults);
+	void PathDataGathered(const double TickStartTime, TRingBuffer<FString>& PathResults);
 
 	/** Called every tick when data is retrieved by the background dependency search */
-	void DependencyDataGathered(const double TickStartTime, TBackgroundGatherResults<FPackageDependencyData>& DependsResults);
+	void DependencyDataGathered(const double TickStartTime, TRingBuffer<FPackageDependencyData>& DependsResults);
 
 	/** Called every tick when data is retrieved by the background search for cooked packages that do not contain asset data */
-	void CookedPackageNamesWithoutAssetDataGathered(const double TickStartTime, TBackgroundGatherResults<FString>& CookedPackageNamesWithoutAssetDataResults);
+	void CookedPackageNamesWithoutAssetDataGathered(const double TickStartTime, TRingBuffer<FString>& CookedPackageNamesWithoutAssetDataResults);
 
 	/** Adds an asset to the empty package list which contains packages that have no assets left in them */
 	void AddEmptyPackage(FName PackageName);
@@ -210,13 +220,6 @@ private:
 	/** Removes the asset data associated with this package from the look-up maps */
 	void RemovePackageData(const FName PackageName);
 
-	/**
-	 * Adds a root path to be discover files in, when asynchronously scanning the disk for asset files
-	 *
-	 * @param	Path	The path on disk to scan
-	 */
-	void AddPathToSearch(const FString& Path);
-
 	/** Adds a list of files which will be searched for asset data */
 	void AddFilesToSearch (const TArray<FString>& Files);
 
@@ -228,7 +231,7 @@ private:
 	void OnAssetLoaded(UObject *AssetLoaded);
 
 	/** Process Loaded Assets to update cache */
-	void ProcessLoadedAssetsToUpdateCache(const double TickStartTime);
+	void ProcessLoadedAssetsToUpdateCache(const double TickStartTime, bool bUpdateDeferredList);
 
 	/** Update Redirect collector with redirects loaded from asset registry */
 	void UpdateRedirectCollector();
@@ -300,12 +303,6 @@ private:
 	 */
 	void RunAssetsThroughFilterImpl(TArray<FAssetData>& AssetDataList, const FARFilter& Filter, const EARFilterMode FilterMode) const;
 
-	/**
-	 * Add sub content blacklist filter for a new mount point
-	 * @param InMount The mount point
-	 */
-	void AddSubContentBlacklist(const FString& InMount);
-
 	/**
 	 * Returns true if path belongs to one of the mount points provided
 	 *
@@ -365,20 +362,14 @@ private:
 	/** The tree of known cached paths that assets may reside within */
 	FPathTree CachedPathTree;
 
-	/** Set of blacklist paths to filter during full asset scans. */
-	TArray<FString> BlacklistScanFilters;
-
-	/** List of sub content path to filter on every mount during full asset scans. */
-	TArray<FString> BlacklistContentSubPaths;
-
 	/** Async task that gathers asset information from disk */
-	TSharedPtr< class FAssetDataGatherer > BackgroundAssetSearch;
+	TUniquePtr<FAssetDataGatherer> GlobalGatherer;
 
 	/** A list of results that were gathered from the background thread that are waiting to get processed by the main thread */
-	TBackgroundGatherResults<FAssetData*> BackgroundAssetResults;
-	TBackgroundGatherResults<FString> BackgroundPathResults;
-	TBackgroundGatherResults<FPackageDependencyData> BackgroundDependencyResults;
-	TBackgroundGatherResults<FString> BackgroundCookedPackageNamesWithoutAssetDataResults;
+	TRingBuffer<FAssetData*> BackgroundAssetResults;
+	TRingBuffer<FString> BackgroundPathResults;
+	TRingBuffer<FPackageDependencyData> BackgroundDependencyResults;
+	TRingBuffer<FString> BackgroundCookedPackageNamesWithoutAssetDataResults;
 
 	/** The max number of results to process per tick */
 	float MaxSecondsPerFrame;
@@ -418,17 +409,18 @@ private:
 	double AmortizeStartTime;
 	double TotalAmortizeTime;
 
-	/** Flag to enable/disable dependency gathering */
-	bool bGatherDependsData;
-
+	/** Flag to indicate if we used an initial async search */
+	bool bInitialSearchStarted;
 	/** Flag to indicate if the initial background search has completed */
 	bool bInitialSearchCompleted;
+	/** Flag that indicates the background search is idle, so we can take actions when idle changes */
+	bool bGatherIdle;
 
 	/** Enables extra check to make sure path still mounted before adding. Removing mount point can happen between scan (background thread + multiple ticks and the add). */
 	bool bVerifyMountPointAfterGather;
 
-	/** A set used to ignore repeated requests to synchronously scan the same folder or file multiple times */
-	TSet<FString> SynchronouslyScannedPathsAndFiles;
+	/** Record whether SearchAllAssets has been called; if so we will also search new mountpoints when added */
+	bool bSearchAllAssets;
 
 	/** List of all class names derived from Blueprint (including Blueprint itself) */
 	mutable TSet<FName> ClassGeneratorNames;
diff --git a/Engine/Source/Runtime/AssetRegistry/Private/AssetRegistryConsoleCommands.h b/Engine/Source/Runtime/AssetRegistry/Private/AssetRegistryConsoleCommands.h
--- a/Engine/Source/Runtime/AssetRegistry/Private/AssetRegistryConsoleCommands.h
+++ b/Engine/Source/Runtime/AssetRegistry/Private/AssetRegistryConsoleCommands.h
@@ -22,6 +22,7 @@ public:
 	FAutoConsoleCommand GetDependenciesCommand;
 	FAutoConsoleCommand GetReferencersCommand;
 	FAutoConsoleCommand FindInvalidUAssetsCommand;
+	FAutoConsoleCommand ScanPathCommand;
 
 	FAssetRegistryConsoleCommands()
 		: GetByNameCommand(
@@ -52,6 +53,10 @@ public:
 		TEXT( "AssetRegistry.Debug.FindInvalidUAssets" ),
 		*LOCTEXT("CommandText_FindInvalidUAssets", "Finds a list of all assets which are in UAsset files but do not share the name of the package").ToString(),
 		FConsoleCommandWithArgsDelegate::CreateRaw( this, &FAssetRegistryConsoleCommands::FindInvalidUAssets ) )
+	, ScanPathCommand(
+		TEXT("AssetRegistry.ScanPath"),
+		*LOCTEXT("CommandText_ScanPath", "Scan the given filename or directoryname for package files and load them into the assetregistry. Extra string parameters: -forcerescan, -ignoreblacklists, -asfile, -asdir").ToString(),
+		FConsoleCommandWithArgsDelegate::CreateRaw(this, &FAssetRegistryConsoleCommands::ScanPath ) )
 	{}
 
 	void GetByName(const TArray<FString>& Args)
@@ -200,6 +205,67 @@ public:
 			}
 		}
 	}
+
+	void ScanPath(const TArray<FString>& Args)
+	{
+		bool bForceRescan = false;
+		bool bIgnoreBlacklists = false;
+		bool bAsFile = false;
+		bool bAsDir = false;
+
+		FString InPath;
+		for (const FString& Arg : Args)
+		{
+			if (Arg.StartsWith(TEXT("-")))
+			{
+				bForceRescan = bForceRescan || Arg.Equals(TEXT("-forcerescan"), ESearchCase::IgnoreCase);
+				bIgnoreBlacklists = bIgnoreBlacklists || Arg.Equals(TEXT("-ignoreblacklists"), ESearchCase::IgnoreCase);
+				bAsDir = bAsDir || Arg.Equals(TEXT("-asdir"), ESearchCase::IgnoreCase);
+				bAsFile = bAsFile || Arg.Equals(TEXT("-asfile"), ESearchCase::IgnoreCase);
+			}
+			else
+			{
+				InPath = Arg;
+			}
+		}
+		if (InPath.IsEmpty())
+		{
+			UE_LOG(LogAssetRegistry, Log, TEXT("Usage: AssetRegistry.ScanPath [-forcerescan] [-ignoreblacklists] [-asfile] [-asdir] FileOrDirectoryPath"));
+			return;
+		}
+
+		if (!bAsDir && !bAsFile)
+		{
+			bAsDir = true;
+			if (FPackageName::IsValidLongPackageName(InPath))
+			{
+				FString LocalPath;
+				if (FPackageName::TryConvertLongPackageNameToFilename(InPath, LocalPath))
+				{
+					FPackagePath PackagePath = FPackagePath::FromLocalPath(LocalPath);
+					if (FPackageName::DoesPackageExist(PackagePath, &PackagePath))
+					{
+						bAsFile = true;
+						bAsDir = false;
+					}
+				}
+			}
+			else if (IFileManager::Get().FileExists(*InPath))
+			{
+				bAsFile = true;
+				bAsDir = false;
+			}
+		}
+		if (bAsDir)
+		{
+			IAssetRegistry::GetChecked().ScanPathsSynchronous({ InPath }, bForceRescan, bIgnoreBlacklists);
+		}
+		else
+		{
+			IAssetRegistry::GetChecked().ScanFilesSynchronous({ InPath }, bForceRescan);
+		}
+	}
+
 };
 
 #undef LOCTEXT_NAMESPACE
diff --git a/Engine/Source/Runtime/AssetRegistry/Private/BackgroundGatherResults.h b/Engine/Source/Runtime/AssetRegistry/Private/BackgroundGatherResults.h
deleted file mode 100644
diff --git a/Engine/Source/Runtime/AssetRegistry/Private/CriticalSectionQueryable.h b/Engine/Source/Runtime/AssetRegistry/Private/CriticalSectionQueryable.h
new file mode 100644
--- /dev/null
+++ b/Engine/Source/Runtime/AssetRegistry/Private/CriticalSectionQueryable.h
@@ -0,0 +1,104 @@
+// Copyright Epic Games, Inc. All Rights Reserved.
+
+#pragma once
+
+#include <atomic>
+#include "HAL/CriticalSection.h"
+#include "HAL/PlatformTLS.h"
+
+/** A Critical Section with additional function to query whether it is locked on the current thread. */
+class FCriticalSectionQueryable
+{
+public:
+	FCriticalSectionQueryable(const FCriticalSectionQueryable&) = delete;
+	FCriticalSectionQueryable(FCriticalSectionQueryable&&) = delete;
+	FCriticalSectionQueryable& operator=(const FCriticalSectionQueryable&) = delete;
+	FCriticalSectionQueryable& operator=(FCriticalSectionQueryable&&) = delete;
+
+	FCriticalSectionQueryable()
+		: OwnerThreadId(0)
+	{
+	}
+
+	void Lock()
+	{
+		Inner.Lock();
+		SetOwner();
+	}
+
+	bool TryLock()
+	{
+		if (Inner.TryLock())
+		{
+			SetOwner();
+			return true;
+		}
+		return false;
+	}
+
+	void Unlock()
+	{
+		check(IsLockedOnCurrentThread());
+		ClearOwner();
+		Inner.Unlock();
+	}
+
+	bool IsLockedOnCurrentThread() const
+	{
+		return OwnerThreadId.load(std::memory_order::memory_order_relaxed) == FPlatformTLS::GetCurrentThreadId();
+	}
+
+private:
+	void SetOwner()
+	{
+		OwnerThreadId.store(FPlatformTLS::GetCurrentThreadId(), std::memory_order::memory_order_relaxed);
+	}
+
+	void ClearOwner()
+	{
+		OwnerThreadId.store(0, std::memory_order::memory_order_relaxed);
+	}
+
+	FCriticalSection Inner;
+	std::atomic<uint32> OwnerThreadId;
+};
+
+/**
+ * Identical to FScopeLock but works with an FCriticalSectionQueryable
+ */
+class FScopeLockQueryable
+{
+public:
+	FScopeLockQueryable() = delete;
+	FScopeLockQueryable(const FScopeLockQueryable& InScopeLock) = delete;
+	FScopeLockQueryable(FScopeLockQueryable&& InScopeLock) = delete;
+	FScopeLockQueryable& operator=(FScopeLockQueryable& InScopeLock) = delete;
+	FScopeLockQueryable& operator=(FScopeLockQueryable&& InScopeLock) = delete;
+
+	FScopeLockQueryable(FCriticalSectionQueryable* InSynchObject)
+		: SynchObject(InSynchObject)
+	{
+		check(SynchObject);
+		SynchObject->Lock();
+	}
+
+	/** Destructor that performs a release on the synchronization object. */
+	~FScopeLockQueryable()
+	{
+		Unlock();
+	}
+
+	void Unlock()
+	{
+		if (SynchObject)
+		{
+			SynchObject->Unlock();
+			SynchObject = nullptr;
+		}
+	}
+
+private:
+
+	// Holds the synchronization object to aggregate and scope manage.
+	FCriticalSectionQueryable* SynchObject;
+};
diff --git a/Engine/Source/Runtime/CoreUObject/Private/Misc/PackageName.cpp b/Engine/Source/Runtime/CoreUObject/Private/Misc/PackageName.cpp
--- a/Engine/Source/Runtime/CoreUObject/Private/Misc/PackageName.cpp
+++ b/Engine/Source/Runtime/CoreUObject/Private/Misc/PackageName.cpp
@@ -613,7 +613,14 @@ bool FPackageName::TryConvertToMountedPath(FStringView InPath, FString* OutLocal
 		if (OutFailureReason) *OutFailureReason = EErrorCode::PackageNameFullObjectPathNotAllowed;
 		return false;
 	}
-	if (!IsValidTextForLongPackageName(PackageName, &FailureReason))
+	FStringView PackageNameNoTrailingSlash = PackageName;
+	if (PackageName.Len() > 0 && FPathViews::IsSeparator(PackageName[PackageName.Len() - 1]))
+	{
+		// IsValidTextForLongPackageName rejects packagenames with a trailing slash, but we want to allow that
+		// because this function allows both files and directories
+		PackageNameNoTrailingSlash.LeftChopInline(1);
+	}
+	if (!IsValidTextForLongPackageName(PackageNameNoTrailingSlash, &FailureReason))
 	{
 		ClearSuccessOutputs();
 		if (OutFailureReason) *OutFailureReason = FailureReason;
diff --git a/Engine/Source/Runtime/CoreUObject/Public/Misc/PackageName.h b/Engine/Source/Runtime/CoreUObject/Public/Misc/PackageName.h
--- a/Engine/Source/Runtime/CoreUObject/Public/Misc/PackageName.h
+++ b/Engine/Source/Runtime/CoreUObject/Public/Misc/PackageName.h
@@ -500,12 +500,11 @@ public:
 	 * @param	Filename to test. 
 	 * @return	True if the filename ends with a package extension.
 	 */
-	static FORCEINLINE bool IsPackageFilename(const FString& Filename)
+	static FORCEINLINE bool IsPackageFilename(FStringView Filename)
 	{
 		FStringView AssetPackageExtension(LexToString(EPackageExtension::Asset));
 		FStringView MapPackageExtension(LexToString(EPackageExtension::Map));
-		FStringView FilenameView(Filename);
-		return FilenameView.EndsWith(AssetPackageExtension) || FilenameView.EndsWith(MapPackageExtension);
+		return Filename.EndsWith(AssetPackageExtension) || Filename.EndsWith(MapPackageExtension);
 	}
 
 	/**
@@ -675,6 +674,22 @@ public:
 		return DoesPackageExistOverrideDelegate;
 	}
 
+	/**
+	 * Find a mount point that contains the given LocalFilePath or PackageName, and if found, return the MountPoint and RelativePath
+	 *
+	 * @param InFilePathOrPackageName The path to test, either a LocalFilePath or a PackageName or an ObjectPath
+	 * @param OutMountPointPackageName If the MountPoint is found, the PackageName of the MountPoint is copied into this variable, otherwise it is set to empty string
+	 * @param OutMountPointFilePath If the MountPoint is found, the LocalFilePath of the MountPoint is copied into this variable, otherwise it is set to empty string
+	 * @param OutRelPath If the MountPoint is found, the RelativePath from the MountPoint to InFilePathOrPackageName is copied into this variable, otherwise it is set to empty string
+	 * 					 If InFilePathOrPackageName was a filepath, the extension is removed before copying it into OutRelpath
+	 *					 The OutRelPath is the same for both the LocalFilePath and the PackageName
+	 * @param OutFlexNameType If non-null, will be set to whether InFilePathOrPackageName was a PackageName or Filename if the MountPoint is found, otherwise it is set to EFlexNameType::Invalid
+	 * @param OutFailureReason If non-null, will be set to the reason InPath could not be converted, or to EErrorCode::Unknown if the function was successful.
+	 * @return True if the MountPoint was found, else false
+	 */
+	static bool TryGetMountPointForPath(FStringView InFilePathOrPackageName, FStringBuilderBase& OutMountPointPackageName, FStringBuilderBase& OutMountPointFilePath, FStringBuilderBase& OutRelPath,
+		EFlexNameType* OutFlexNameType = nullptr, EErrorCode* OutFailureReason = nullptr);
+
 private:
 
 	/**
@@ -702,22 +717,6 @@ private:
 	static bool TryConvertToMountedPathComponents(FStringView InPath, FStringBuilderBase& OutMountPointPackageName, FStringBuilderBase& OutMountPointFilePath, FStringBuilderBase& OutRelPath,
 		FStringBuilderBase& OutObjectName, EPackageExtension& OutExtension, FStringBuilderBase& OutCustomExtension, EFlexNameType* OutFlexNameType = nullptr, EErrorCode* OutFailureReason = nullptr);
 
-	/**
-	 * Internal helper to find a mount point that contains the given LocalFilePath or PackageName, and if found, return the MountPoint and RelativePath
-	 *
-	 * @param InFilePathOrPackageName The path to test, either a LocalFilePath or a PackageName or an ObjectPath
-	 * @param OutMountPointPackageName If the MountPoint is found, the PackageName of the MountPoint is copied into this variable, otherwise it is set to empty string
-	 * @param OutMountPointFilePath If the MountPoint is found, the LocalFilePath of the MountPoint is copied into this variable, otherwise it is set to empty string
-	 * @param OutRelPath If the MountPoint is found, the RelativePath from the MountPoint to InFilePathOrPackageName is copied into this variable, otherwise it is set to empty string
-	 * 					 If InFilePathOrPackageName was a filepath, the extension is removed before copying it into OutRelpath
-	 *					 The OutRelPath is the same for both the LocalFilePath and the PackageName
-	 * @param OutFlexNameType If non-null, will be set to whether InFilePathOrPackageName was a PackageName or Filename if the MountPoint is found, otherwise it is set to EFlexNameType::Invalid
-	 * @param OutFailureReason If non-null, will be set to the reason InPath could not be converted, or to EErrorCode::Unknown if the function was successful.
-	 * @return True if the MountPoint was found, else false
-	 */
-	static bool TryGetMountPointForPath(FStringView InFilePathOrPackageName, FStringBuilderBase& OutMountPointPackageName, FStringBuilderBase& OutMountPointFilePath, FStringBuilderBase& OutRelPath,
-		EFlexNameType* OutFlexNameType = nullptr, EErrorCode* OutFailureReason = nullptr);
-
 	/** Event that is triggered when a new content path is mounted */
 	static FOnContentPathMountedEvent OnContentPathMountedEvent;
 
