AssetRegistry - Quick hack Fix for ScanPathsSynchronous not working with IgnoreBlacklists.
#rnx

diff --git a/Engine/Source/Runtime/AssetRegistry/Private/AssetDataGatherer.cpp b/Engine/Source/Runtime/AssetRegistry/Private/AssetDataGatherer.cpp
--- a/Engine/Source/Runtime/AssetRegistry/Private/AssetDataGatherer.cpp
+++ b/Engine/Source/Runtime/AssetRegistry/Private/AssetDataGatherer.cpp
@@ -265,16 +265,22 @@ bool FScanDir::IsBlacklisted() const
 		return false;
 	}
 
-	const TSet<FString>& Blacklist = MountDir->GetBlacklist();
+	const TArray<TPair<FString,bool>>& Blacklist = MountDir->GetBlacklist();
 	if (Blacklist.Num())
 	{
 		TStringBuilder<128> MountRelPath;
 		AppendMountRelPath(MountRelPath);
 		FStringView MountRelPathSV(MountRelPath);
-		if (Blacklist.ContainsByHash(GetTypeHash(MountRelPathSV), MountRelPathSV))
+		bool bIsBlacklisted = false;
+		for (const TPair<FString,bool>& BlacklistedPair : Blacklist)
 		{
-			return true;
+			const FString& BlacklistedPath = BlacklistedPair.Get<0>();
+			if (FPathViews::IsParentPathOf(BlacklistedPath, MountRelPathSV))
+			{
+				bIsBlacklisted = BlacklistedPair.Get<1>();
+			}
 		}
+		return bIsBlacklisted;
 	}
 	return false;
 }
@@ -496,7 +502,9 @@ bool FScanDir::TrySetDirectoryProperties(FStringView InRelPath, const FSetPathPr
 		FScanDir* SubDir = nullptr;
 		if (bHasScanned &&
 			(!Properties->HasScanned.IsSet() || *Properties->HasScanned == true) &&
-			(!Properties->IsWhitelisted.IsSet()))
+			(!Properties->IsWhitelisted.IsSet()) &&
+			(!Properties->IgnoreBlacklist.IsSet())
+			)
 		{
 			// If this parent directory has already been scanned and we are not changing the target directory's has-been-scanned value,
 			// and the next child subdirectory does not exist, then the child directory has already been scanned and we do not need to set the properties on it.
@@ -644,10 +652,6 @@ void FScanDir::SetScanResults(FStringView LocalAbsPath, TArrayView<FDiscoveredPa
 			++Index;
 		}
 	}
-	for (FScanDir* ScanDir : SubDirsToRemove)
-	{
-		RemoveSubDir(ScanDir->GetRelPath());
-	}
 
 	// Add the files that were found in the scan, skipping any files that have already been scanned
 	if (InOutFiles.Num())
@@ -798,7 +802,7 @@ void FScanDir::SetComplete(bool bInIsComplete)
 		}
 		// Upon completion, subdirs that do not need to be maintained are deleted, which is done by removing them from the parent
 		// ScanDirs need to be maintained if they are the root, or are whitelisted, or have child ScanDirs that need to be maintained.
-		if (Parent != nullptr && !bIsDirectWhitelisted && SubDirs.IsEmpty())
+		if (Parent != nullptr && !bIsDirectWhitelisted && !IsBlacklisted() && SubDirs.IsEmpty())
 		{
 			Parent->RemoveSubDir(GetRelPath());
 			// *this is Shutdown (e.g. Parent is now null) and it may also have been deallocated
@@ -948,7 +952,7 @@ FStringView FMountDir::GetLongPackageName() const
 	return LongPackageName;
 }
 
-const TSet<FString>& FMountDir::GetBlacklist() const
+const TArray<TPair<FString,bool>>& FMountDir::GetBlacklist() const
 {
 	return BlacklistedRelPaths;
 }
@@ -979,9 +983,9 @@ uint32 FMountDir::GetAllocatedSize() const
 	}
 	Result += LongPackageName.GetAllocatedSize();
 	Result += BlacklistedRelPaths.GetAllocatedSize();
-	for (const FString& Value : BlacklistedRelPaths)
+	for (const TPair<FString,bool>& Value : BlacklistedRelPaths)
 	{
-		Result += Value.GetAllocatedSize();
+		Result += Value.Get<0>().GetAllocatedSize();
 	}
 	return Result;
 }
@@ -1018,14 +1022,15 @@ bool FMountDir::IsBlacklisted(FStringView InLocalAbsPath) const
 {
 	FStringView QueryRelPath;
 	verify(FPathViews::TryMakeChildPathRelativeTo(InLocalAbsPath, GetLocalAbsPath(), QueryRelPath));
-	for (const FString& Blacklist : BlacklistedRelPaths)
+	bool bIsBlacklisted = false;
+	for (const TPair<FString,bool>& Blacklist : BlacklistedRelPaths)
 	{
-		if (FPathViews::IsParentPathOf(Blacklist, QueryRelPath))
+		if (FPathViews::IsParentPathOf(Blacklist.Get<0>(), QueryRelPath))
 		{
-			return true;
+			bIsBlacklisted = Blacklist.Get<1>();
 		}
 	}
-	return false;
+	return bIsBlacklisted;
 }
 
 bool FMountDir::IsMonitored(FStringView InLocalAbsPath) const
@@ -1079,13 +1084,6 @@ bool FMountDir::TrySetDirectoryProperties(FStringView InLocalAbsPath, const FSet
 		
 		MarkDirty(RelPath);
 
-		ModifiedProperties = InProperties;
-		ModifiedProperties->IgnoreBlacklist.Reset();
-		if (!ModifiedProperties->IsSet())
-		{
-			return true;
-		}
-		Properties = &ModifiedProperties.GetValue();
 		bResult = true;
 	}
 
@@ -1101,27 +1099,31 @@ void FMountDir::UpdateBlacklist()
 		FStringView MountRelPath;
 		if (FPathViews::TryMakeChildPathRelativeTo(BlacklistName, LongPackageName, MountRelPath))
 		{
-			FPathData* PathData = FindPathData(MountRelPath);
-			if (!PathData || !PathData->bIgnoreBlacklist)
-			{
-				// Note that an empty RelPath means we blacklist the entire mountpoint
-				BlacklistedRelPaths.Emplace(MountRelPath);
-			}
+			// Note that an empty RelPath means we blacklist the entire mountpoint
+			BlacklistedRelPaths.Emplace(MountRelPath, true);
 		}
 	}
+
 	for (const FString& MountRelPath : Discovery.BlacklistMountRelativePaths)
 	{
 		FPathData* PathData = FindPathData(MountRelPath);
 		if (!PathData || !PathData->bIgnoreBlacklist)
 		{
-			BlacklistedRelPaths.Emplace(MountRelPath);
+			BlacklistedRelPaths.Emplace(MountRelPath, true);
+		}
+	}
+	for (const FPathData& PathData : PathDatas)
+	{
+		if (PathData.bIgnoreBlacklist && !PathData.bIsChildPath)
+		{
+			BlacklistedRelPaths.Emplace(PathData.RelPath, false);
 		}
 	}
 	for (const FPathData& PathData : PathDatas)
 	{
 		if (PathData.bIsChildPath)
 		{
-			BlacklistedRelPaths.Emplace(PathData.RelPath);
+			BlacklistedRelPaths.Emplace(PathData.RelPath, true);
 		}
 	}
 }
@@ -3124,7 +3126,8 @@ void FAssetDataGatherer::SerializeCacheLoad(FAssetRegistryReader& Ar)
 		{
 			// Load the name first to add the entry to the tmap below
 			Ar << PackageNameBlock[AssetIndex];
-			AssetDataBlock[AssetIndex].SerializeForCache(Ar);
+			// Visual Studio Static Analyzer issues C6385 if we call AssetDataBlock[AssetIndex].SerializeForCache
+			(AssetDataBlock + AssetIndex)->SerializeForCache(Ar);
 			if (Ar.IsError())
 			{
 				// There was an error reading the cache. Bail out.
diff --git a/Engine/Source/Runtime/AssetRegistry/Private/AssetDataGathererPrivate.h b/Engine/Source/Runtime/AssetRegistry/Private/AssetDataGathererPrivate.h
--- a/Engine/Source/Runtime/AssetRegistry/Private/AssetDataGathererPrivate.h
+++ b/Engine/Source/Runtime/AssetRegistry/Private/AssetDataGathererPrivate.h
@@ -281,7 +281,7 @@ public:
 	 * Get the set of blacklist paths that are child paths of this MountDir.
 	 * Returned paths are relative paths from the MountDir's root.
 	 */
-	const TSet<FString>& GetBlacklist() const;
+	const TArray<TPair<FString,bool>>& GetBlacklist() const;
 	/** Return the FAssetDataDiscovery that owns this FMountDir. */
 	FAssetDataDiscovery& GetDiscovery() const;
 
@@ -367,7 +367,7 @@ protected:
 	/** Array of data keyed by relative path from the MountDir for information about child paths that should not be stored in ScanDirs. */
 	TArray<FPathData> PathDatas;
 	/** Set of relative path from the MountDir paths that should not be scanned, because they were requested blacklisted by clients or because a childmount owns them. */
-	TSet<FString> BlacklistedRelPaths;
+	TArray<TPair<FString,bool>> BlacklistedRelPaths;
 	/** Absolute path to the root of the MountDir in the local file system. */
 	FString LocalAbsPath;
 	/** LongPackageName that was assigned to the MountDir in FPackageName. */
