- Fixed AssetDataGatherer crash on Mac startup

diff --git a/Engine/Source/Runtime/AssetRegistry/Private/AssetDataGatherer.cpp b/Engine/Source/Runtime/AssetRegistry/Private/AssetDataGatherer.cpp
--- a/Engine/Source/Runtime/AssetRegistry/Private/AssetDataGatherer.cpp
+++ b/Engine/Source/Runtime/AssetRegistry/Private/AssetDataGatherer.cpp
@@ -56,22 +56,6 @@ void AssignStringWithoutShrinking(FString& InOutResult, FStringView Value)
 	}
 }
 
-/** Adapter to allow us to use a lambda for IterateDirectoryStat */
-class FLambdaDirectoryStatVisitor : public IPlatformFile::FDirectoryStatVisitor
-{
-public:
-	typedef TFunctionRef<bool(const TCHAR*, const FFileStatData&)> FLambdaRef;
-	FLambdaRef Callback;
-	explicit FLambdaDirectoryStatVisitor(FLambdaRef InCallback)
-		: Callback(MoveTemp(InCallback))
-	{
-	}
-	virtual bool Visit(const TCHAR* FilenameOrDirectory, const FFileStatData& StatData) override
-	{
-		return Callback(FilenameOrDirectory, StatData);
-	}
-};
-
 FDiscoveredPathData::FDiscoveredPathData(FStringView InLocalAbsPath, FStringView InLongPackageName, FStringView InRelPath, const FDateTime& InPackageTimestamp)
 	: LocalAbsPath(InLocalAbsPath)
 	, LongPackageName(InLongPackageName)
@@ -1449,9 +1433,8 @@ void FAssetDataDiscovery::TickInternal()
 
 	int32 NumIteratedDirs = 0;
 	int32 NumIteratedFiles = 0;
-	FLambdaDirectoryStatVisitor Visitor(FLambdaDirectoryStatVisitor::FLambdaRef(
-		[this, &DirLocalAbsPath, &DirLongPackageName, DirLongPackageNameRootLen, &NumIteratedDirs, &NumIteratedFiles]
-	(const TCHAR* InPackageFilename, const FFileStatData& InPackageStatData)
+	IFileManager::Get().IterateDirectoryStat(DirLocalAbsPath.ToString(), [this, &DirLocalAbsPath, &DirLongPackageName, DirLongPackageNameRootLen, &NumIteratedDirs, &NumIteratedFiles]
+		(const TCHAR* InPackageFilename, const FFileStatData& InPackageStatData)
 		{
 			FStringView LocalAbsPath(InPackageFilename);
 			FStringView RelPath;
@@ -1504,9 +1487,8 @@ void FAssetDataDiscovery::TickInternal()
 				}
 			}
 			return true;
-		}));
+		});
 
-	IFileManager::Get().IterateDirectoryStat(DirLocalAbsPath.ToString(), Visitor);
 	TArrayView<FDiscoveredPathData> LocalSubDirs(IteratedSubDirs.GetData(), NumIteratedDirs);
 	TArrayView<FDiscoveredPathData> LocalDiscoveredFiles(IteratedFiles.GetData(), NumIteratedFiles);
 	bool bValid = false;
