AssetDataGatherer fixes:
Fix ignoreblacklists to allow scanning subdirectories of a blacklisted root.
Fix LocalPathToAdd not checking whether the path is blacklisted.
Fix non-unity compile.
#rnx

diff --git a/Engine/Source/Runtime/AssetRegistry/Private/AssetDataGatherer.cpp b/Engine/Source/Runtime/AssetRegistry/Private/AssetDataGatherer.cpp
--- a/Engine/Source/Runtime/AssetRegistry/Private/AssetDataGatherer.cpp
+++ b/Engine/Source/Runtime/AssetRegistry/Private/AssetDataGatherer.cpp
@@ -242,79 +242,74 @@ FString FScanDir::GetMountRelPath() const
 	return FString(Result);
 }
 
-bool FScanDir::IsBlacklisted() const
+bool FScanDir::FInherited::IsMonitored() const
 {
-	if (!MountDir)
-	{
-		return false;
-	}
+	return IsWhitelisted() && !IsBlacklisted();
+}
 
-	const TArray<TPair<FString,bool>>& Blacklist = MountDir->GetBlacklist();
-	if (Blacklist.Num())
-	{
-		TStringBuilder<128> MountRelPath;
-		AppendMountRelPath(MountRelPath);
-		FStringView MountRelPathSV(MountRelPath);
-		bool bIsBlacklisted = false;
-		for (const TPair<FString,bool>& BlacklistedPair : Blacklist)
-		{
-			const FString& BlacklistedPath = BlacklistedPair.Get<0>();
-			if (FPathViews::IsParentPathOf(BlacklistedPath, MountRelPathSV))
-			{
-				bIsBlacklisted = BlacklistedPair.Get<1>();
-			}
-		}
-		return bIsBlacklisted;
-	}
-	return false;
+bool FScanDir::FInherited::IsBlacklisted() const
+{
+	return bMatchesBlacklist && !bIgnoreBlacklist;
 }
 
-bool FScanDir::IsDirectWhitelisted() const
+bool FScanDir::FInherited::IsWhitelisted() const
 {
-	return bIsDirectWhitelisted;
+	return bIsWhitelisted;
 }
 
-bool FScanDir::IsRecursiveMonitored(bool bParentIsWhitelisted) const
+bool FScanDir::FInherited::HasSetting() const
 {
-	if (!MountDir)
-	{
-		return false;
-	}
-	if (!bParentIsWhitelisted && !bIsDirectWhitelisted)
-	{
-		return false;
-	}
-	if (IsBlacklisted())
-	{
-		return false;
-	}
-	return true;
+	return bIsWhitelisted || bMatchesBlacklist || bIgnoreBlacklist;
 }
 
-bool FScanDir::IsPathWhitelisted(FStringView InRelPath, bool bParentIsWhitelisted) const
+FScanDir::FInherited::FInherited(const FInherited& Parent, const FInherited& Child)
+	: bIsWhitelisted(Parent.bIsWhitelisted || Child.bIsWhitelisted)
+	, bMatchesBlacklist(Parent.bMatchesBlacklist || Child.bMatchesBlacklist)
+	, bIgnoreBlacklist(Parent.bIgnoreBlacklist || Child.bIgnoreBlacklist)
 {
-	bool bIsRecursiveWhitelisted = bParentIsWhitelisted || bIsDirectWhitelisted;
-	if (bIsRecursiveWhitelisted)
+}
+
+void FScanDir::GetMonitorData(FStringView InRelPath, const FInherited& ParentData, FInherited& OutData) const
+{
+	if (!MountDir)
 	{
-		return true;
+		OutData = FInherited();
+		return;
 	}
+
+	FInherited Accumulated(ParentData, DirectData);
+
 	const FScanDir* SubDir = nullptr;
+	FStringView FirstComponent;
+	FStringView RemainingPath;
 	if (!InRelPath.IsEmpty())
 	{
-		FStringView FirstComponent;
-		FPathViews::SplitFirstComponent(InRelPath, FirstComponent, InRelPath);
+		FPathViews::SplitFirstComponent(InRelPath, FirstComponent, RemainingPath);
 		SubDir = FindSubDir(FirstComponent);
 	}
 	if (!SubDir)
+	{
+		OutData = Accumulated;
+	}
+	else
+	{
+		SubDir->GetMonitorData(RemainingPath, Accumulated, OutData);
+	}
+}
+
+bool FScanDir::IsMonitored(const FInherited& ParentData) const
+{
+	if (!MountDir)
 	{
 		return false;
 	}
-	return SubDir->IsPathWhitelisted(InRelPath, false /* bParentIsWhitelisted */);
+	FInherited Accumulated(ParentData, DirectData);
+	return Accumulated.IsMonitored();
 }
 
-bool FScanDir::ShouldScan(bool bParentIsWhitelisted) const
+bool FScanDir::ShouldScan(const FInherited& ParentData) const
 {
-	return !bHasScanned && IsRecursiveMonitored(bParentIsWhitelisted);
+	return !bHasScanned && IsMonitored(ParentData);
 }
 
 bool FScanDir::HasScanned() const
@@ -345,24 +340,24 @@ uint32 FScanDir::GetAllocatedSize() const
 	return Result;
 }
 
-FScanDir* FScanDir::GetControllingDir(FStringView InRelPath, bool bIsDirectory, bool bParentIsWhitelisted, bool& bOutIsWhitelisted, FString& OutRelPath)
+FScanDir* FScanDir::GetControllingDir(FStringView InRelPath, bool bIsDirectory, const FInherited& ParentData, FInherited& OutData, FString& OutRelPath)
 {
 	// GetControllingDir can only be called on valid ScanDirs, which we rely on since we need to call FindOrAddSubDir which relies on that
 	check(IsValid());
 
-	bool bIsWhitelisted = bParentIsWhitelisted || bIsDirectWhitelisted;
+	FInherited Accumulated(ParentData, DirectData);
 	if (InRelPath.IsEmpty())
 	{
 		if (!bIsDirectory)
 		{
 			UE_LOG(LogAssetRegistry, Warning, TEXT("GetControllingDir called on %s with !bIsDirectory, but we have it recorded as a directory. Returning null."), *GetLocalAbsPath());
-			bOutIsWhitelisted = false;
+			OutData = FInherited();
 			OutRelPath.Reset();
 			return nullptr;
 		}
 		else
 		{
-			bOutIsWhitelisted = bIsWhitelisted;
+			OutData = Accumulated;
 			OutRelPath = InRelPath;
 			return this;
 		}
@@ -373,14 +368,14 @@ FScanDir* FScanDir::GetControllingDir(FStringView InRelPath, bool bIsDirectory,
 	FPathViews::SplitFirstComponent(InRelPath, FirstComponent, RemainingPath);
 	if (RemainingPath.IsEmpty() && !bIsDirectory)
 	{
-		bOutIsWhitelisted = bIsWhitelisted;
+		OutData = Accumulated;
 		OutRelPath = InRelPath;
 		return this;
 	}
 	else
 	{
 		FScanDir* SubDir = nullptr;
-		if (ShouldScan(bParentIsWhitelisted))
+		if (ShouldScan(ParentData))
 		{
 			SubDir = &FindOrAddSubDir(FirstComponent);
 		}
@@ -389,12 +384,12 @@ FScanDir* FScanDir::GetControllingDir(FStringView InRelPath, bool bIsDirectory,
 			SubDir = FindSubDir(FirstComponent);
 			if (!SubDir)
 			{
-				bOutIsWhitelisted = bIsWhitelisted;
+				OutData = Accumulated;
 				OutRelPath = InRelPath;
 				return this;
 			}
 		}
-		return SubDir->GetControllingDir(RemainingPath, bIsDirectory, bIsWhitelisted, bOutIsWhitelisted, OutRelPath);
+		return SubDir->GetControllingDir(RemainingPath, bIsDirectory, Accumulated, OutData, OutRelPath);
 	}
 }
 
@@ -407,22 +402,22 @@ bool FScanDir::TrySetDirectoryProperties(FStringView InRelPath, const FSetPathPr
 	if (InRelPath.IsEmpty())
 	{
 		// The properties apply to this entire directory
-		if (InProperties.IsWhitelisted.IsSet() && bIsDirectWhitelisted != *InProperties.IsWhitelisted)
+		if (InProperties.IsWhitelisted.IsSet() && DirectData.bIsWhitelisted != *InProperties.IsWhitelisted)
 		{
 			if (bScanInFlight)
 			{
 				bScanInFlightInvalidated = true;
 			}
-			bIsDirectWhitelisted = *InProperties.IsWhitelisted;
+			DirectData.bIsWhitelisted = *InProperties.IsWhitelisted;
 
-			if (bIsDirectWhitelisted)
+			if (DirectData.bIsWhitelisted)
 			{
 				// Since we are setting this directory to be monitored, we need to implement the guarantee that all Monitored flags of its children are set to false
 				// We also need to SetComplete false on all directories in between this and a previously whitelisted directory, since those non-whitelisted parent directories
 				// marked themselves complete once their whitelisted children finished
 				ForEachDescendent([](FScanDir& ScanDir)
 					{
-						ScanDir.bIsDirectWhitelisted = false;
+						ScanDir.DirectData.bIsWhitelisted = false;
 						ScanDir.SetComplete(false);
 					});
 			}
@@ -438,6 +433,41 @@ bool FScanDir::TrySetDirectoryProperties(FStringView InRelPath, const FSetPathPr
 					});
 			}
 		}
+		if ((InProperties.MatchesBlacklist.IsSet() && DirectData.bMatchesBlacklist != *InProperties.MatchesBlacklist) ||
+			(InProperties.IgnoreBlacklist.IsSet() && DirectData.bIgnoreBlacklist != *InProperties.IgnoreBlacklist))
+		{
+			if (InProperties.MatchesBlacklist.IsSet())
+			{
+				DirectData.bMatchesBlacklist = *InProperties.MatchesBlacklist;
+			}
+			if (InProperties.IgnoreBlacklist.IsSet())
+			{
+				DirectData.bIgnoreBlacklist = *InProperties.IgnoreBlacklist;
+			}
+			bool bIgnoreBlacklist = false;
+			bool bMatchesBlacklist = false;
+			for (FScanDir* Current = this; Current; Current = Current->Parent)
+			{
+				bIgnoreBlacklist = bIgnoreBlacklist || Current->DirectData.bIgnoreBlacklist;
+				bMatchesBlacklist = bMatchesBlacklist || Current->DirectData.bMatchesBlacklist;
+			}
+			bool bIsBlacklisted = bMatchesBlacklist && !bIgnoreBlacklist;
+
+			// Mark all children as incomplete
+			// Also cancel any scans since they are now potentially blacklisted
+			if (bIsBlacklisted && bScanInFlight)
+			{
+				bScanInFlightInvalidated = true;
+			}
+			ForEachDescendent([&bIsBlacklisted](FScanDir& ScanDir)
+				{
+					if (bIsBlacklisted && ScanDir.bScanInFlight)
+					{
+						ScanDir.bScanInFlightInvalidated = true;
+					}
+					ScanDir.SetComplete(false);
+				});
+		}
 		if (InProperties.HasScanned.IsSet())
 		{
 			bool bNewValue = *InProperties.HasScanned;
@@ -457,26 +487,22 @@ bool FScanDir::TrySetDirectoryProperties(FStringView InRelPath, const FSetPathPr
 		{
 			SetDirectPriority(*InProperties.Priority);
 		}
-		// InProperties.IgnoreBlacklist does not require an action on ScanDirs; it is implemented on the MountDir level
 		return true;
 	}
 	else
 	{
 		TOptional<FSetPathProperties> ModifiedProperties;
 		const FSetPathProperties* Properties = &InProperties;
-		if (Properties->IsWhitelisted.IsSet())
+		if (Properties->IsWhitelisted.IsSet() && DirectData.bIsWhitelisted)
 		{
-			if (bIsDirectWhitelisted)
+			// If this directory is set to be monitored, all Monitored flags of its children are unused, are guaranteed set to false, and should not be changed
+			ModifiedProperties = *Properties;
+			ModifiedProperties->IsWhitelisted.Reset();
+			if (!ModifiedProperties->IsSet())
 			{
-				// If this directory is set to be monitored, all Monitored flags of its children are unused, are guaranteed set to false, and should not be changed
-				ModifiedProperties = *Properties;
-				ModifiedProperties->IsWhitelisted.Reset();
-				if (!ModifiedProperties->IsSet())
-				{
-					return false;
-				}
-				Properties = &ModifiedProperties.GetValue();
+				return false;
 			}
+			Properties = &ModifiedProperties.GetValue();
 		}
 
 		FStringView FirstComponent;
@@ -487,8 +513,8 @@ bool FScanDir::TrySetDirectoryProperties(FStringView InRelPath, const FSetPathPr
 		if (bHasScanned &&
 			(!Properties->HasScanned.IsSet() || *Properties->HasScanned == true) &&
 			(!Properties->IsWhitelisted.IsSet()) &&
-			(!Properties->IgnoreBlacklist.IsSet())
-			)
+			(!Properties->IgnoreBlacklist.IsSet()) &&
+			(!Properties->MatchesBlacklist.IsSet()))
 		{
 			// If this parent directory has already been scanned and we are not changing the target directory's has-been-scanned value,
 			// and the next child subdirectory does not exist, then the child directory has already been scanned and we do not need to set the properties on it.
@@ -599,10 +625,8 @@ void FScanDir::OnChildPriorityChanged(EPriority InPriority, int32 Delta)
 	UpdateAccumulatedPriority();
 }
 
-void FScanDir::SetScanResults(FStringView LocalAbsPath, TArrayView<FDiscoveredPathData>& InOutSubDirs, TArrayView<FDiscoveredPathData>& InOutFiles)
+void FScanDir::SetScanResults(FStringView LocalAbsPath, const FInherited& ParentData, TArrayView<FDiscoveredPathData>& InOutSubDirs, TArrayView<FDiscoveredPathData>& InOutFiles)
 {
-	// Note that by contract SetScanResults is only called on Paths with ShouldScan == true, so we do not need to check IsWhitelisted for the files in this directory or for any of its subdirs
-	// We do still need to check IsBlacklisted for subdirs, since the blacklist can be true for a subdirectory even if not true for the parent
 	SetComplete(false);
 	check(!bScanInFlightInvalidated);
 	check(MountDir);
@@ -611,21 +635,14 @@ void FScanDir::SetScanResults(FStringView LocalAbsPath, TArrayView<FDiscoveredPa
 	{
 		return;
 	}
+	FInherited Accumulated(ParentData, DirectData);
 
 	// Add SubDirectories in the tree for the directories found by the scan, and report the directories as discovered directory paths as well
-	// Remove any SubDirectories from the tree that were previously present but are not in the latest scan results
-	TSet<FScanDir*> SubDirsToRemove;
-	ForEachSubDir([&SubDirsToRemove](FScanDir& SubScanDir) { SubDirsToRemove.Add(&SubScanDir); });
 	for (int32 Index = 0; Index < InOutSubDirs.Num(); )
 	{
 		FDiscoveredPathData& SubDirPath = InOutSubDirs[Index];
-		bool bReportResult = false;
-		if (!MountDir->IsBlacklisted(SubDirPath.LocalAbsPath))
-		{
-			FScanDir& SubScanDir = FindOrAddSubDir(SubDirPath.RelPath);
-			SubDirsToRemove.Remove(&SubScanDir);
-			bReportResult = MountDir->GetDiscovery().ShouldDirBeReported(SubDirPath.LongPackageName);
-		}
+		FScanDir& SubScanDir = FindOrAddSubDir(SubDirPath.RelPath);
+		bool bReportResult = SubScanDir.IsMonitored(Accumulated) && MountDir->GetDiscovery().ShouldDirBeReported(SubDirPath.LongPackageName);
 		if (!bReportResult)
 		{
 			Swap(SubDirPath, InOutSubDirs.Last());
@@ -666,14 +683,14 @@ void FScanDir::SetScanResults(FStringView LocalAbsPath, TArrayView<FDiscoveredPa
 	bHasScanned = true;
 }
 
-void FScanDir::Update(FScanDir*& OutCursor, bool& bInOutParentIsWhitelisted)
+void FScanDir::Update(FScanDir*& OutCursor, FInherited& InOutParentData)
 {
 	check(MountDir);
 	if (bIsComplete)
 	{
 		return;
 	}
-	if (ShouldScan(bInOutParentIsWhitelisted))
+	if (ShouldScan(InOutParentData))
 	{
 		OutCursor = this;
 		return;
@@ -685,26 +702,45 @@ void FScanDir::Update(FScanDir*& OutCursor, bool& bInOutParentIsWhitelisted)
 		if (SubDirToScan)
 		{
 			OutCursor = SubDirToScan;
-			bInOutParentIsWhitelisted = bInOutParentIsWhitelisted || bIsDirectWhitelisted;
+			FInherited Accumulated(InOutParentData, DirectData);
+			InOutParentData = Accumulated;
 			return;
 		}
 	}
-	SetComplete(true);
 
 	OutCursor = Parent; // Note this will be null for the root ScanDir
 	if (!Parent)
 	{
-		bInOutParentIsWhitelisted = false;
+		InOutParentData = FInherited();
 	}
-	else if (Parent->bIsDirectWhitelisted)
+	else
 	{
-		// We have a contract bIsDirectWhitelisted is only set on the highest-level directory to monitor and applies to all directories under it.
-		// So we only need to change bInOutParentIsWhitelisted from true to false when we move up the tree into a parent directory with bIsDirectWhitelisted = true.
-		check(!this->bIsDirectWhitelisted); // Verify the contract that children below the ScanDir with bIsDirectWhitelisted = true have bIsDirectWhitelisted = false
-		check(!Parent->Parent || !Parent->Parent->bIsDirectWhitelisted); // Verify the contract that the parent above the ScanDir with bIsDirectWhitelisted = true has bIsDirectWhitelisted = false
-		check(bInOutParentIsWhitelisted); // Verify that the original bInOutParentIsWhitelisted was set to true; it should have been since the parent's bIsDirectWhitelisted is true
-		bInOutParentIsWhitelisted = false;
+		if (Parent->DirectData.bIsWhitelisted)
+		{
+			// We have a contract that bIsWhitelisted is only set on the highest-level directory to monitor and
+			// applies to all directories under it. So we only need to change ParentData.bIsWhitelisted from true to
+			// false when we move up the tree into a parent directory with bIsWhitelisted = true.
+			check(!this->DirectData.bIsWhitelisted); // Verify children below whitelist true are whitelist false
+			check(!Parent->Parent || !Parent->Parent->DirectData.bIsWhitelisted); // Verify above whitelist true is whitelist false
+			check(InOutParentData.bIsWhitelisted); // Verify original InOutParentData matches parent's direct value
+			InOutParentData.bIsWhitelisted = false;
+		}
+		if (Parent->DirectData.bMatchesBlacklist || Parent->DirectData.bIgnoreBlacklist)
+		{
+			// We don't have the same set-once contract for blacklisted information, so when we find a direct blacklisted parent
+			// we have to recalculate the parent's blacklisted information by checking all ancestors.
+			InOutParentData.bMatchesBlacklist = false;
+			InOutParentData.bIgnoreBlacklist = false;
+			for (FScanDir* Current = Parent->Parent; Current; Current = Current->Parent)
+			{
+				InOutParentData.bMatchesBlacklist = InOutParentData.bMatchesBlacklist || Current->DirectData.bMatchesBlacklist;
+				InOutParentData.bIgnoreBlacklist = InOutParentData.bIgnoreBlacklist || Current->DirectData.bIgnoreBlacklist;
+			}
+		}
 	}
+
+	SetComplete(true);
+	// After calling SetComplete, this may have been removed from tree and should no longer run calculations
 }
 
 bool FScanDir::IsScanInFlight() const
@@ -785,8 +821,8 @@ void FScanDir::SetComplete(bool bInIsComplete)
 			}
 		}
 		// Upon completion, subdirs that do not need to be maintained are deleted, which is done by removing them from the parent
-		// ScanDirs need to be maintained if they are the root, or are whitelisted, or have child ScanDirs that need to be maintained.
-		if (Parent != nullptr && !bIsDirectWhitelisted && !IsBlacklisted() && SubDirs.IsEmpty())
+		// ScanDirs need to be maintained if they are the root, or have persistent settings, or have child ScanDirs that need to be maintained.
+		if (Parent != nullptr && !HasPersistentSettings() && SubDirs.IsEmpty())
 		{
 			Parent->RemoveSubDir(GetRelPath());
 			// *this is Shutdown (e.g. Parent is now null) and it may also have been deallocated
@@ -799,6 +835,11 @@ void FScanDir::SetComplete(bool bInIsComplete)
 	}
 }
 
+bool FScanDir::HasPersistentSettings() const
+{
+	return DirectData.HasSetting();
+}
+
 FScanDir* FScanDir::FindSubDir(FStringView SubDirBaseName)
 {
 	int32 Index = FindLowerBoundSubDir(SubDirBaseName);
@@ -936,40 +977,37 @@ FStringView FMountDir::GetLongPackageName() const
 	return LongPackageName;
 }
 
-const TArray<TPair<FString,bool>>& FMountDir::GetBlacklist() const
-{
-	return BlacklistedRelPaths;
-}
-
 FAssetDataDiscovery& FMountDir::GetDiscovery() const
 {
 	return Discovery;
 }
 
-FScanDir* FMountDir::GetControllingDir(FStringView InLocalAbsPath, bool bIsDirectory, bool& bOutIsWhitelisted, FString& OutRelPath)
+FScanDir* FMountDir::GetControllingDir(FStringView InLocalAbsPath, bool bIsDirectory, FScanDir::FInherited& OutData,
+	FString& OutRelPath)
 {
 	FStringView RemainingPath;
 	if (!FPathViews::TryMakeChildPathRelativeTo(InLocalAbsPath, GetLocalAbsPath(), RemainingPath))
 	{
 		return nullptr;
 	}
-	return Root->GetControllingDir(RemainingPath, bIsDirectory, false /* bParentIsWhitelisted */, bOutIsWhitelisted, OutRelPath);
+	return Root->GetControllingDir(RemainingPath, bIsDirectory, FScanDir::FInherited() /* ParentData */,
+		OutData, OutRelPath);
 }
 
 uint32 FMountDir::GetAllocatedSize() const
 {
 	uint32 Result = sizeof(*Root);
 	Result += Root->GetAllocatedSize();
-	Result += PathDatas.GetAllocatedSize();
-	for (const FPathData& Value : PathDatas)
+	Result += ChildMountPaths.GetAllocatedSize();
+	for (const FString& Value : ChildMountPaths)
 	{
 		Result += Value.GetAllocatedSize();
 	}
 	Result += LongPackageName.GetAllocatedSize();
 	Result += BlacklistedRelPaths.GetAllocatedSize();
-	for (const TPair<FString,bool>& Value : BlacklistedRelPaths)
+	for (const FString& Value : BlacklistedRelPaths)
 	{
-		Result += Value.Get<0>().GetAllocatedSize();
+		Result += Value.GetAllocatedSize();
 	}
 	return Result;
 }
@@ -977,7 +1015,7 @@ uint32 FMountDir::GetAllocatedSize() const
 void FMountDir::Shrink()
 {
 	Root->Shrink();
-	PathDatas.Shrink();
+	ChildMountPaths.Shrink();
 	BlacklistedRelPaths.Shrink();
 }
 
@@ -991,45 +1029,23 @@ EPriority FMountDir::GetPriority() const
 	return Root->GetPriority();
 }
 
-bool FMountDir::IsPathWhitelisted(FStringView InLocalAbsPath) const
+void FMountDir::GetMonitorData(FStringView InLocalAbsPath, FScanDir::FInherited& OutData) const
 {
 	FStringView QueryRelPath;
 	if (!ensure(FPathViews::TryMakeChildPathRelativeTo(InLocalAbsPath, GetLocalAbsPath(), QueryRelPath)))
 	{
-		return false;
+		OutData = FScanDir::FInherited();
+		return;
 	}
 
-	return Root->IsPathWhitelisted(QueryRelPath, false);
-}
-
-bool FMountDir::IsBlacklisted(FStringView InLocalAbsPath) const
-{
-	FStringView QueryRelPath;
-	verify(FPathViews::TryMakeChildPathRelativeTo(InLocalAbsPath, GetLocalAbsPath(), QueryRelPath));
-	bool bIsBlacklisted = false;
-	for (const TPair<FString,bool>& Blacklist : BlacklistedRelPaths)
-	{
-		if (FPathViews::IsParentPathOf(Blacklist.Get<0>(), QueryRelPath))
-		{
-			bIsBlacklisted = Blacklist.Get<1>();
-		}
-	}
-	return bIsBlacklisted;
+	return Root->GetMonitorData(QueryRelPath, FScanDir::FInherited() /* ParentData */, OutData);
 }
 
 bool FMountDir::IsMonitored(FStringView InLocalAbsPath) const
 {
-	FStringView QueryRelPath;
-	if (!ensure(FPathViews::TryMakeChildPathRelativeTo(InLocalAbsPath, GetLocalAbsPath(), QueryRelPath)))
-	{
-		return false;
-	}
-	if (IsBlacklisted(InLocalAbsPath))
-	{
-		return false;
-	}
-
-	return Root->IsPathWhitelisted(QueryRelPath, false);
+	FScanDir::FInherited MonitorData;
+	GetMonitorData(InLocalAbsPath, MonitorData);
+	return MonitorData.IsMonitored();
 }
 
 bool FMountDir::TrySetDirectoryProperties(FStringView InLocalAbsPath, const FSetPathProperties& InProperties, bool bConfirmedExists)
@@ -1039,44 +1055,34 @@ bool FMountDir::TrySetDirectoryProperties(FStringView InLocalAbsPath, const FSet
 	{
 		return false;
 	}
-	bool bResult = false;
-
-	const FSetPathProperties* Properties = &InProperties;
-	TOptional<FSetPathProperties> ModifiedProperties;
-	if (Properties->IgnoreBlacklist.IsSet())
+	if (InProperties.IgnoreBlacklist.IsSet())
 	{
-		// IgnoreBlacklist is applied at the MountDir level, so we handle it separately
-		bool bIgnoreBlacklist = *Properties->IgnoreBlacklist;
-		FPathData* PathData = nullptr;
-		if (bIgnoreBlacklist)
+		if (!ensure(!IsChildMountPath(RelPath)))
 		{
-			PathData = &FindOrAddPathData(RelPath);
-		}
-		else
-		{
-			PathData = FindPathData(RelPath);
-		}
-		if (PathData)
-		{
-			PathData->bIgnoreBlacklist = bIgnoreBlacklist;
-			if (PathData->IsEmpty())
+			// Setting IgnoreBlacklist on a child path would break behavior because we use MatchesBlacklist to indicate that 
+			// the scandir is a child path, and setting it to ignoreblacklists will defeat that setting.
+			// This should never be called, because setting ignoreblacklist is only called external to FAssetDataDiscovery, and 
+			// FAssetDataDiscovery would call it on the child mount dir instead of this parent mountdir
+			FSetPathProperties NewProperties(InProperties);
+			NewProperties.IgnoreBlacklist.Reset();
+			if (!NewProperties.IsSet())
 			{
-				RemovePathData(RelPath);
+				return false;
 			}
+			return TrySetDirectoryProperties(InLocalAbsPath, NewProperties, bConfirmedExists);
 		}
-		UpdateBlacklist();
-		
-		MarkDirty(RelPath);
-
-		bResult = true;
 	}
-
-	bResult = Root->TrySetDirectoryProperties(RelPath, InProperties, bConfirmedExists) || bResult;
-	return bResult;
+	return Root->TrySetDirectoryProperties(RelPath, InProperties, bConfirmedExists);
 }
 
 void FMountDir::UpdateBlacklist()
 {
+	TSet<FString> RemovedBlacklists;
+	for (const FString& Old : BlacklistedRelPaths)
+	{
+		RemovedBlacklists.Add(Old);
+	}
+
 	BlacklistedRelPaths.Empty(Discovery.BlacklistMountRelativePaths.Num());
 	for (const FString& BlacklistName : Discovery.BlacklistLongPackageNames)
 	{
@@ -1084,38 +1090,43 @@ void FMountDir::UpdateBlacklist()
 		if (FPathViews::TryMakeChildPathRelativeTo(BlacklistName, LongPackageName, MountRelPath))
 		{
 			// Note that an empty RelPath means we blacklist the entire mountpoint
-			BlacklistedRelPaths.Emplace(MountRelPath, true);
+			BlacklistedRelPaths.Emplace(MountRelPath);
 		}
 	}
-
 	for (const FString& MountRelPath : Discovery.BlacklistMountRelativePaths)
 	{
-		FPathData* PathData = FindPathData(MountRelPath);
-		if (!PathData || !PathData->bIgnoreBlacklist)
-		{
-			BlacklistedRelPaths.Emplace(MountRelPath, true);
-		}
+		BlacklistedRelPaths.Emplace(MountRelPath);
 	}
-	for (const FPathData& PathData : PathDatas)
+	for (const FString& ChildPath : ChildMountPaths)
 	{
-		if (PathData.bIgnoreBlacklist && !PathData.bIsChildPath)
-		{
-			BlacklistedRelPaths.Emplace(PathData.RelPath, false);
-		}
+		BlacklistedRelPaths.Emplace(ChildPath);
 	}
-	for (const FPathData& PathData : PathDatas)
+
+	TSet<FString> AddedBlacklists;
+	for (const FString& New : BlacklistedRelPaths)
 	{
-		if (PathData.bIsChildPath)
+		if (!RemovedBlacklists.Remove(New))
 		{
-			BlacklistedRelPaths.Emplace(PathData.RelPath, true);
+			AddedBlacklists.Add(New);
 		}
 	}
+	FSetPathProperties ChangeBlacklist;
+	ChangeBlacklist.MatchesBlacklist = true;
+	for (const FString& RelPath : AddedBlacklists)
+	{
+		Root->TrySetDirectoryProperties(RelPath, ChangeBlacklist, false /* bConfirmedExists */);
+	}
+	ChangeBlacklist.MatchesBlacklist = false;
+	for (const FString& RelPath : RemovedBlacklists)
+	{
+		Root->TrySetDirectoryProperties(RelPath, ChangeBlacklist, false /* bConfirmedExists */);
+	}
 }
 
-void FMountDir::Update(FScanDir*& OutCursor, bool& bOutCursorParentIsWhitelisted)
+void FMountDir::Update(FScanDir*& OutCursor, FScanDir::FInherited& OutParentData)
 {
-	bOutCursorParentIsWhitelisted = false;
-	Root->Update(OutCursor, bOutCursorParentIsWhitelisted);
+	OutParentData = FScanDir::FInherited();
+	Root->Update(OutCursor, OutParentData);
 }
 
 void FMountDir::SetHasStartedScanning()
@@ -1134,7 +1145,7 @@ void FMountDir::AddChildMount(FMountDir* ChildMount)
 	{
 		return;
 	}
-	FindOrAddPathData(RelPath).bIsChildPath = true;
+	AddChildMountPath(RelPath);
 	if (bHasStartedScanning)
 	{
 		UE_LOG(LogAssetRegistry, Warning, TEXT("AssetDataGatherer directory %.*s has already started scanning when a new mountpoint was added under it at %.*s. ")
@@ -1156,16 +1167,10 @@ void FMountDir::RemoveChildMount(FMountDir* ChildMount)
 	{
 		return;
 	}
-	FPathData* PathData = FindPathData(RelPath);
-	if (!PathData)
+	if (!RemoveChildMountPath(RelPath))
 	{
 		return;
 	}
-	PathData->bIsChildPath = false;
-	if (PathData->IsEmpty())
-	{
-		RemovePathData(RelPath);
-	}
 	if (ChildMount->bHasStartedScanning)
 	{
 		UE_LOG(LogAssetRegistry, Warning, TEXT("AssetDataGatherer directory %.*s has already started scanning when it was removed and merged into its parent mount at %.*s. ")
@@ -1178,9 +1183,7 @@ void FMountDir::RemoveChildMount(FMountDir* ChildMount)
 
 void FMountDir::OnDestroyClearChildMounts()
 {
-	// This function deletes more data than just the child paths; if it becomes used for purposes other than deleting the MountDir,
-	// we will need to change it to only set PathData->bIsChildPath = false;
-	PathDatas.Empty();
+	ChildMountPaths.Empty();
 }
 
 void FMountDir::SetParentMount(FMountDir* Parent)
@@ -1197,18 +1200,15 @@ TArray<FMountDir*> FMountDir::GetChildMounts() const
 {
 	// Called within Discovery's TreeLock
 	TArray<FMountDir*> Result;
-	for (const FPathData& PathData : PathDatas)
+	for (const FString& ChildPath : ChildMountPaths)
 	{
-		if (PathData.bIsChildPath)
+		TStringBuilder<256> ChildAbsPath;
+		ChildAbsPath << LocalAbsPath;
+		FPathViews::AppendPath(ChildAbsPath, ChildPath);
+		FMountDir* ChildMount = Discovery.FindMountPoint(ChildAbsPath);
+		if (ensure(ChildMount)) // This PathData information should have been removed with RemoveChildMount when the child MountDir was removed from the Discovery
 		{
-			TStringBuilder<256> ChildAbsPath;
-			ChildAbsPath << LocalAbsPath;
-			FPathViews::AppendPath(ChildAbsPath, PathData.RelPath);
-			FMountDir* ChildMount = Discovery.FindMountPoint(ChildAbsPath);
-			if (ensure(ChildMount)) // This PathData information should have been removed with RemoveChildMount when the child MountDir was removed from the Discovery
-			{
-				Result.Add(ChildMount);
-			}
+			Result.Add(ChildMount);
 		}
 	}
 	return Result;
@@ -1216,53 +1216,46 @@ TArray<FMountDir*> FMountDir::GetChildMounts() const
 
 void FMountDir::MarkDirty(FStringView MountRelPath)
 {
-	bool bIsWhitelisted = true;
+	FScanDir::FInherited UnusedMonitorData;
 	FString ControlRelPath;
-	FScanDir* ScanDir = Root->GetControllingDir(MountRelPath, true /* bIsDirectory */, false /* bParentIsWhitelisted */, bIsWhitelisted, ControlRelPath);
+	FScanDir* ScanDir = Root->GetControllingDir(MountRelPath, true /* bIsDirectory */, FScanDir::FInherited() /* ParentData */,
+		UnusedMonitorData, ControlRelPath);
 	if (ScanDir)
 	{
-		// If a ScanDir exists for the directory that is being marked dirty, mark all of its descendents dirty as well.
+		// If a ScanDir exists for the directory that is being marked dirty, mark all of its descendants dirty as well.
 		// If the control dir is a parent directory of the requested path, just mark it and its parents dirty
-		// Mark all parent directories that exist as incomplete
+		// Mark all parents dirty in either case
 		bool bDirtyAllDescendents = ControlRelPath.IsEmpty();
 		ScanDir->MarkDirty(bDirtyAllDescendents);
 	}
 }
 
-FMountDir::FPathData* FMountDir::FindPathData(FStringView MountRelPath)
-{
-	return Algo::FindByPredicate(PathDatas, [MountRelPath](const FPathData& PathData) { return FPathViews::Equals(PathData.RelPath, MountRelPath); });
-}
-
-FMountDir::FPathData& FMountDir::FindOrAddPathData(FStringView MountRelPath)
+void FMountDir::AddChildMountPath(FStringView MountRelPath)
 {
-	FPathData* PathData = FindPathData(MountRelPath);
-	if (!PathData)
+	FString* ExistingPath = ChildMountPaths.FindByPredicate([&MountRelPath](const FString& ChildPath) { return FPathViews::Equals(ChildPath, MountRelPath); });
+	if (!ExistingPath)
 	{
-		PathData = &PathDatas.Emplace_GetRef(MountRelPath);
+		ChildMountPaths.Emplace(MountRelPath);
 	}
-	return *PathData;
 }
 
-void FMountDir::RemovePathData(FStringView MountRelPath)
+bool FMountDir::RemoveChildMountPath(FStringView MountRelPath)
 {
-	PathDatas.RemoveAllSwap([MountRelPath](const FPathData& PathData) { return FPathViews::Equals(PathData.RelPath, MountRelPath);  });
+	return ChildMountPaths.RemoveAllSwap([MountRelPath](const FString& ChildPath) { return FPathViews::Equals(ChildPath, MountRelPath);  }) != 0;
 }
 
-FMountDir::FPathData::FPathData(FStringView MountRelPath)
-:RelPath(MountRelPath)
+bool FMountDir::IsChildMountPath(FStringView MountRelPath) const
 {
+	for (const FString& ChildPath : ChildMountPaths)
+	{
+		if (FPathViews::IsParentPathOf(ChildPath, MountRelPath))
+		{
+			return true;
+		}
+	}
+	return false;
 }
 
-bool FMountDir::FPathData::IsEmpty() const
-{
-	return !bIgnoreBlacklist && !bIsChildPath;
-}
-
-uint32 FMountDir::FPathData::GetAllocatedSize() const
-{
-	return RelPath.GetAllocatedSize();
-}
 
 FAssetDataDiscovery::FAssetDataDiscovery(const TArray<FString>& InBlacklistLongPackageNames, const TArray<FString>& InBlacklistMountRelativePaths, bool bInIsSynchronous)
 	: BlacklistLongPackageNames(InBlacklistLongPackageNames)
@@ -1397,7 +1390,7 @@ void FAssetDataDiscovery::TickInternal()
 			if (!Cursor || !Cursor->IsValid())
 			{
 				FScanDir* NewCursor;
-				FindFirstCursor(NewCursor, bCursorParentIsWhitelisted);
+				FindFirstCursor(NewCursor, CursorParentData);
 				Cursor = NewCursor;
 				if (!NewCursor)
 				{
@@ -1405,13 +1398,13 @@ void FAssetDataDiscovery::TickInternal()
 					return;
 				}
 			}
-			if (Cursor->ShouldScan(bCursorParentIsWhitelisted))
+			if (Cursor->ShouldScan(CursorParentData))
 			{
 				break;
 			}
 
 			FScanDir* NewCursor = Cursor.GetReference();
-			NewCursor->Update(NewCursor, bCursorParentIsWhitelisted);
+			NewCursor->Update(NewCursor, CursorParentData);
 			check(NewCursor != Cursor);
 			Cursor = NewCursor;
 		}
@@ -1507,13 +1500,13 @@ void FAssetDataDiscovery::TickInternal()
 		}
 		else
 		{
-			LocalCursor->SetScanResults(DirLocalAbsPath, LocalSubDirs, LocalDiscoveredFiles);
+			LocalCursor->SetScanResults(DirLocalAbsPath, CursorParentData, LocalSubDirs, LocalDiscoveredFiles);
 			bValid = true;
 			FScanDir* NewCursor = Cursor.GetReference();
 			// Other thread may have set the cursor to a new spot; in that case do not update and on the next tick start at the new cursor
 			if (LocalCursor == NewCursor)
 			{
-				LocalCursor->Update(NewCursor, bCursorParentIsWhitelisted);
+				LocalCursor->Update(NewCursor, CursorParentData);
 				Cursor = NewCursor;
 			}
 		}
@@ -1526,7 +1519,7 @@ void FAssetDataDiscovery::TickInternal()
 	}
 }
 
-void FAssetDataDiscovery::FindFirstCursor(FScanDir*& OutCursor, bool& bOutCursorParentIsWhitelisted)
+void FAssetDataDiscovery::FindFirstCursor(FScanDir*& OutCursor, FScanDir::FInherited& OutParentData)
 {
 	CHECK_IS_LOCKED_CURRENT_THREAD(TreeLock);
 	OutCursor = nullptr;
@@ -1550,11 +1543,11 @@ void FAssetDataDiscovery::FindFirstCursor(FScanDir*& OutCursor, bool& bOutCursor
 		if (!WinningMountDir)
 		{
 			OutCursor = nullptr;
-			bOutCursorParentIsWhitelisted = false;
+			OutParentData = FScanDir::FInherited();
 			break;
 		}
 
-		WinningMountDir->Update(OutCursor, bOutCursorParentIsWhitelisted);
+		WinningMountDir->Update(OutCursor, OutParentData);
 		check(OutCursor != nullptr || WinningMountDir->IsComplete()); // The WinningMountDir's update should either return something to update or it should mark itself complete
 	}
 }
@@ -1657,12 +1650,6 @@ void FAssetDataDiscovery::SetPropertiesAndWait(const FString& LocalAbsPath, bool
 			UE_LOG(LogAssetRegistry, Log, TEXT("SetPropertiesAndWait called on %.*s which is not in a mounted directory. Call will be ignored."), LocalAbsPath.Len(), *LocalAbsPath);
 			return;
 		}
-		bool bIsBlacklisted = MountDir->IsBlacklisted(LocalAbsPath);
-		if (bIsBlacklisted && !bIgnoreBlackListScanFilters)
-		{
-			UE_LOG(LogAssetRegistry, Log, TEXT("SetPropertiesAndWait called on %.*s which is blacklisted. Call will be ignored."), LocalAbsPath.Len(), *LocalAbsPath);
-			return;
-		}
 
 		if (StatData.bIsDirectory)
 		{
@@ -1687,11 +1674,11 @@ void FAssetDataDiscovery::SetPropertiesAndWait(const FString& LocalAbsPath, bool
 		}
 
 		FString RelPath;
-		bool bIsWhitelisted;
-		TRefCountPtr<FScanDir> ScanDir = MountDir->GetControllingDir(LocalAbsPath, StatData.bIsDirectory, bIsWhitelisted, RelPath);
-		if (!ScanDir || (!bIsWhitelisted && !bAddToWhitelist))
+		FScanDir::FInherited MonitorData;
+		TRefCountPtr<FScanDir> ScanDir = MountDir->GetControllingDir(LocalAbsPath, StatData.bIsDirectory, MonitorData, RelPath);
+		if (!ScanDir || !MonitorData.IsMonitored())
 		{
-			UE_LOG(LogAssetRegistry, Log, TEXT("SetPropertiesAndWait called on %.*s which is not whitelisted. Call will be ignored."), LocalAbsPath.Len(), *LocalAbsPath);
+			UE_LOG(LogAssetRegistry, Log, TEXT("SetPropertiesAndWait called on %.*s which is not monitored. Call will be ignored."), LocalAbsPath.Len(), *LocalAbsPath);
 			return;
 		}
 
@@ -1706,9 +1693,9 @@ void FAssetDataDiscovery::SetPropertiesAndWait(const FString& LocalAbsPath, bool
 				ScanDir->SetDirectPriority(EPriority::Blocking);
 				InvalidateCursor();
 
+				FScopedPause ScopedPause(*this);
 				TreeScopeLock.Unlock(); // Entering the ticklock, as well as any long duration task such as a tick, has to be done outside of any locks
 
-				FScopedPause ScopedPause(*this);
 				CHECK_IS_NOT_LOCKED_CURRENT_THREAD(TreeLock);
 				CHECK_IS_NOT_LOCKED_CURRENT_THREAD(ResultsLock);
 				FGathererScopeLock TickScopeLock(&TickLock);
@@ -1729,7 +1716,7 @@ void FAssetDataDiscovery::SetPropertiesAndWait(const FString& LocalAbsPath, bool
 		}
 		else
 		{
-			bool bAlreadyScanned = ScanDir->HasScanned() && !bIsBlacklisted;
+			bool bAlreadyScanned = ScanDir->HasScanned() && MonitorData.IsMonitored();
 			if (!bAlreadyScanned || bForceRescan)
 			{
 				FStringView RelPathFromParentDir = FPathViews::GetCleanFilename(RelPath);
@@ -1787,7 +1774,13 @@ bool FAssetDataDiscovery::IsWhitelisted(FStringView LocalAbsPath) const
 	CHECK_IS_NOT_LOCKED_CURRENT_THREAD(ResultsLock);
 	FGathererScopeLock TreeScopeLock(&TreeLock);
 	const FMountDir* MountDir = FindContainingMountPoint(LocalAbsPath);
-	return MountDir && MountDir->IsPathWhitelisted(LocalAbsPath);
+	if (!MountDir)
+	{
+		return false;
+	}
+	FScanDir::FInherited MonitorData;
+	MountDir->GetMonitorData(LocalAbsPath, MonitorData);
+	return MonitorData.IsWhitelisted();
 }
 
 bool FAssetDataDiscovery::IsBlacklisted(FStringView LocalAbsPath) const
@@ -1797,11 +1790,11 @@ bool FAssetDataDiscovery::IsBlacklisted(FStringView LocalAbsPath) const
 	const FMountDir* MountDir = FindContainingMountPoint(LocalAbsPath);
 	if (!MountDir)
 	{
-		// Not mounted, which we report as not blacklisted
 		return false;
 	}
-
-	return MountDir->IsBlacklisted(LocalAbsPath);
+	FScanDir::FInherited MonitorData;
+	MountDir->GetMonitorData(LocalAbsPath, MonitorData);
+	return MonitorData.IsBlacklisted();
 }
 
 bool FAssetDataDiscovery::IsMonitored(FStringView LocalAbsPath) const
@@ -2029,9 +2022,9 @@ void FAssetDataDiscovery::OnFileCreated(const FString& LocalAbsPath)
 	CHECK_IS_LOCKED_CURRENT_THREAD(TreeLock);
 	// Detect whether the file should be scanned and if so pass it through to the gatherer
 	FMountDir* MountDir = FindContainingMountPoint(LocalAbsPath);
-	if (!MountDir || MountDir->IsBlacklisted(LocalAbsPath))
+	if (!MountDir)
 	{
-		// The content root of the file is not registered, or the file is blacklisted; ignore it
+		// The content root of the file is not registered; ignore it
 		return;
 	}
 	FFileStatData StatData = IFileManager::Get().GetStatData(*LocalAbsPath);
@@ -2042,9 +2035,9 @@ void FAssetDataDiscovery::OnFileCreated(const FString& LocalAbsPath)
 	}
 
 	FString FileRelPath;
-	bool bIsWhitelisted;
-	FScanDir* ScanDir = MountDir->GetControllingDir(LocalAbsPath, false /* bIsDirectory */, bIsWhitelisted, FileRelPath);
-	if (!ScanDir || !bIsWhitelisted)
+	FScanDir::FInherited MonitorData;
+	FScanDir* ScanDir = MountDir->GetControllingDir(LocalAbsPath, false /* bIsDirectory */, MonitorData, FileRelPath);
+	if (!ScanDir || !MonitorData.IsMonitored())
 	{
 		// The new file is in an unmonitored directory; ignore it
 		return;
diff --git a/Engine/Source/Runtime/AssetRegistry/Private/AssetDataGathererPrivate.h b/Engine/Source/Runtime/AssetRegistry/Private/AssetDataGathererPrivate.h
--- a/Engine/Source/Runtime/AssetRegistry/Private/AssetDataGathererPrivate.h
+++ b/Engine/Source/Runtime/AssetRegistry/Private/AssetDataGathererPrivate.h
@@ -6,6 +6,7 @@
 
 #include <atomic>
 #include "Containers/Set.h"
+#include "Misc/Optional.h"
 #include "Misc/StringBuilder.h"
 #include "Templates/RefCounting.h"
 #include "Templates/UniquePtr.h"
@@ -36,20 +37,30 @@ struct FSetPathProperties
 {
 	/** The path (usually a plugin's root content path) has been requested for scanning through e.g. ScanPathsSynchronous */
 	TOptional<bool> IsWhitelisted;
-	/** The directory's list of direct file/subdirectory children has been scanned through a call to IFileManager::IterateDirectoryStat after process start or the last request to rescan it */
-	TOptional<bool> HasScanned;
 	/**
-	 * Paths under the directory should be used (if whitelisted) even if they match blacklist entries. Used only on the MountDir.
-	 * Used e.g. when a level is opened in the editor and we want to add to the asset registries all of its one-file-per-actor files that are normally blacklisted.
+	 * Set whether the given directory matches a blacklist entry.
+	 */
+	TOptional<bool> MatchesBlacklist;
+	/**
+	 * Set whether the given directory should ignore blacklist, even if it or its parent matches a blacklist entry.
 	 */
 	TOptional<bool> IgnoreBlacklist;
-	/** Priority that determines which directories are scanned first, for when a directory is being waited on or are not blocking but are wanted for optional features */
+	/**
+	 * The directory's list of direct file/subdirectory children has been scanned through a call to
+	 * IFileManager::IterateDirectoryStat after process start or the last request to rescan it
+	 */
+	TOptional<bool> HasScanned;
+	/**
+	 * Priority that determines which directories are scanned first, for when a directory is being waited on or are
+	 * not blocking but are wanted for optional features
+	 */
 	TOptional<EPriority> Priority;
 
 	/** Used to early-exit from tree traversal when all properties have finished being handled */
 	bool IsSet() const
 	{
-		return IsWhitelisted.IsSet() | HasScanned.IsSet() | IgnoreBlacklist.IsSet() | Priority.IsSet();
+		return IsWhitelisted.IsSet() | HasScanned.IsSet() | MatchesBlacklist.IsSet() |
+			IgnoreBlacklist.IsSet() | Priority.IsSet(); //-V792
 	}
 };
 
@@ -62,16 +73,18 @@ struct FDiscoveredPathData
 	FString LongPackageName;
 	/** The relative path from the path's parent directory */
 	FString RelPath;
-	/** If the path is a file, the modification timestamp of the package file (that it had when it was discovered) */
+	/** If the path is a file, the modification timestamp of the package file (that it had when it was discovered). */
 	FDateTime PackageTimestamp;
 
 	FDiscoveredPathData() = default;
-	FDiscoveredPathData(FStringView InLocalAbsPath, FStringView InLongPackageName, FStringView InRelPath, const FDateTime& InPackageTimestamp);
+	FDiscoveredPathData(FStringView InLocalAbsPath, FStringView InLongPackageName, FStringView InRelPath,
+		const FDateTime& InPackageTimestamp);
 	FDiscoveredPathData(FStringView InLocalAbsPath, FStringView InLongPackageName, FStringView InRelPath);
-	void Assign(FStringView InLocalAbsPath, FStringView InLongPackageName, FStringView InRelPath, const FDateTime& InPackageTimestamp);
+	void Assign(FStringView InLocalAbsPath, FStringView InLongPackageName, FStringView InRelPath,
+		const FDateTime& InPackageTimestamp);
 	void Assign(FStringView InLocalAbsPath, FStringView InLongPackageName, FStringView InRelPath);
 
-	/** Return the total amount of heap memory used by the gatherer (including not-yet-claimed search results). Used for performance metrics. */
+	/** Return the total amount of heap memory used by the gatherer (including not-yet-claimed search results). */
 	uint32 GetAllocatedSize() const;
 };
 
@@ -108,6 +121,25 @@ struct FGatheredPathData
 class FScanDir : public FRefCountBase
 {
 public:
+	/** Holder for fields on FScanDir that are inherited from parents */
+	struct FInherited
+	{
+		FInherited() = default;
+
+		/** Whether the ScanDir has been whitelisted. */
+		bool bIsWhitelisted = false;
+		/** Whether the ScanDir is under one of the blacklisted folders. */
+		bool bMatchesBlacklist = false;
+		/** Whether the ScanDir has been set to ignore blacklist. */
+		bool bIgnoreBlacklist = false;
+
+		bool IsMonitored() const;
+		bool IsBlacklisted() const;
+		bool IsWhitelisted() const;
+		bool HasSetting() const;
+		FInherited(const FInherited& Parent, const FInherited& Child);
+	};
+
 	FScanDir(FMountDir& InMountDir, FScanDir* InParent, FStringView RelPath);
 	~FScanDir();
 
@@ -137,24 +169,16 @@ public:
 	/** Calculate the relative path from the MountPoint to this ScanDir and return it as a string.  */
 	FString GetMountRelPath() const;
 
-	/** Report whether this ScanDir is blacklisted */
-	bool IsBlacklisted() const;
 	/**
-	 * Return whether this ScanDir has been directly marked as whitelisted. A ScanDir may be indirectly whitelisted
-	 * because it has a whitelisted parent directory; that is not reported by this function.
+	 * Report whether the given RelPath is whitelisted and blacklisted, based on parent data and
+	 * on the discovered settings on Stem ScanDirs between *this and the leaf path.
 	 */
-	bool IsDirectWhitelisted() const;
-	/**
-	 * Return whether this scandir is whitelisted and not blacklisted and hence is in the set of directories that
-	 * need to be scanned.
-	 * Considers whitelisting/blacklisting inherited from parents in addition to the direct settings on this scandir.
-	 */
-	bool IsRecursiveMonitored(bool bParentIsWhitelisted) const;
-	/** Report whether the given RelPath is whitelisted, because either it or a parent is direct whitelisted. */
-	bool IsPathWhitelisted(FStringView InRelPath, bool bParentIsWhitelisted) const;
+	void GetMonitorData(FStringView InRelPath, const FInherited& ParentData, FInherited& OutData) const;
+	/** Return whether this scandir is whitelisted and not blacklisted and hence needs to be or has been scanned. */
+	bool IsMonitored(const FInherited& ParentData) const;
 
 	/** Report whether this ScanDir will be scanned in the current or future Tick. */
-	bool ShouldScan(bool bParentIsWhitelisted) const;
+	bool ShouldScan(const FInherited& ParentData) const;
 	/** Report whether this ScanDir has been scanned. */
 	bool HasScanned() const;
 	/** Report whether this ScanDir is complete: has scanned or should not scan, and all subdirs have completed. */
@@ -164,10 +188,11 @@ public:
 	uint32 GetAllocatedSize() const;
 	/**
 	 * Find the Direct parent of InRelPath, or a fallback. Will return null only if !bIsDirectory and InRelPath is empty.
-	 * The fallback is returned if InRelPath has already completed and been deleted, or if InRelPath is not recursively whitelisted or is blacklisted.
+	 * The fallback is returned if InRelPath has already completed and been deleted, or if InRelPath is not monitored.
 	 * The fallback is the lowest existing parent directory of InRelPath.
 	 */
-	FScanDir* GetControllingDir(FStringView InRelPath, bool bIsDirectory, bool bParentIsWhitelisted, bool& bOutIsWhitelisted, FString& OutRelPath);
+	FScanDir* GetControllingDir(FStringView InRelPath, bool bIsDirectory, const FInherited& ParentData,
+		FInherited& OutData, FString& OutRelPath);
 
 	/**
 	 * Set values of fields on the given directory indicated by InRelPath for all of the properties existing on InProperties.
@@ -175,30 +200,41 @@ public:
 	 * directory or the property did not need to be changed.
 	 */
 	bool TrySetDirectoryProperties(FStringView InRelPath, const FSetPathProperties& InProperties, bool bConfirmedExists);
-	/** Mark that the given file has already been scanned, so that it will not be double reporting in the upcoming directory scan, if one is upcoming. */
+	/**
+	 * Mark that the given file has already been scanned, so that it will not be double reporting in the upcoming
+	 * directory scan, if one is upcoming.
+	 */
 	void MarkFileAlreadyScanned(FStringView BaseName);
 	/** Set the direct priority of this ScanDir and update the accumulated priority to match it. */
 	void SetDirectPriority(EPriority InPriority);
 
 	/**
-	 * Called from the Tick; handle the list of subdirs and files that were found from IterateDirectoryStat called on this ScanDir,
-	 * reporting discovered files and updating status variables.
+	 * Called from the Tick; handle the list of subdirs and files that were found from IterateDirectoryStat called
+	 * on this ScanDir, reporting discovered files and updating status variables.
 	 */
-	void SetScanResults(FStringView LocalAbsPath, TArrayView<FDiscoveredPathData>& InOutSubDirs, TArrayView<FDiscoveredPathData>& InOutFiles);
+	void SetScanResults(FStringView LocalAbsPath, const FInherited& ParentData, TArrayView<FDiscoveredPathData>& InOutSubDirs,
+		TArrayView<FDiscoveredPathData>& InOutFiles);
 	/**
 	 * Update the completion state of this ScanDir based on its scan status and its child dirs completion state.
-	 * Set OutCursor to the highest priority child dir that needs to be updated if any children need to be updated and/or scanned first.
+	 * Set OutCursor to the highest priority child dir that needs to be updated if any children need to be updated
+	 * and/or scanned first.
 	 * Set OutCursor to the parent ScanDir (or null if there is no parent) if this ScanDir is now complete.
 	 */
-	void Update(FScanDir*& OutCursor, bool& bInOutParentIsWhitelisted);
+	void Update(FScanDir*& OutCursor, FInherited& InOutParentData);
 
 	/** Thread-synchronization helper - return true if the Tick thread is in the middle of scanning this directory. */
 	bool IsScanInFlight() const;
 	/** Thread-synchronization helper - set that the Tick thread is starting/done with the scan of this directory. */
 	void SetScanInFlight(bool bInScanInFlight);
-	/** Thread-synchronization helper - report whether a non tick thread has marked that this directory is changed or invalidated and the scan should be thrown out. */
+	/**
+	 * Thread-synchronization helper - report whether a non tick thread has marked that this directory is changed or
+	 * invalidated and the scan should be thrown out.
+	 */
 	bool IsScanInFlightInvalidated() const;
-	/** Thread-synchronization helper - set that the current ongoing scan is invalidated, or clear the marker from the tick thread once it has been consumed. */
+	/**
+	 * Thread-synchronization helper - set that the current ongoing scan is invalidated, or clear the marker from
+	 * the tick thread once it has been consumed.
+	 */
 	void SetScanInFlightInvalidated(bool bInValidated);
 
 	/** Set completion flags on this and its parents (and optionally its descendents) so that it will be updated again. */
@@ -210,17 +246,26 @@ public:
 protected:
 
 	/**
-	 * Setting to complete marks that this directory has been scanned, and all of its subdirectories have scanned as well, so it can skipped when it or its parent
-	 * is waited on. Setting back to incomplete can occur when a AssetDataGatherer client wants to rescan the directory.
+	 * Setting to complete marks that this directory has been scanned, and all of its subdirectories have scanned
+	 * as well, so it can be skipped when it or its parent is waited on. Setting back to incomplete can occur when a
+	 * AssetDataGatherer client wants to rescan the directory.
 	 */
-	virtual void SetComplete(bool bInIsComplete);
+	void SetComplete(bool bInIsComplete);
+	/** 
+	 * Return whether this ScanDir has direct settings (whitelist, blacklist, etc) that we need to preserve on it
+	 * even if has finished scanning.
+	 */
+	bool HasPersistentSettings() const;
 
 	/** Find the ScanDir subdirectory for the given basename, or return nullptr. */
 	FScanDir* FindSubDir(FStringView SubDirBaseName);
 	const FScanDir* FindSubDir(FStringView SubDirBaseName) const;
 	/** Find the ScanDir subdirectory for the given basename, and add it if it does not exist. */
 	FScanDir& FindOrAddSubDir(FStringView SubDirBaseName);
-	/** Find the ScanDir subdirectory for the given basename, and if it exists, Shutdown and remove it from SubDirs, which will eventually delete it. */
+	/**
+	 * Find the ScanDir subdirectory for the given basename, and if it exists, Shutdown and remove it from SubDirs,
+	 * which will eventually delete it.
+	 */
 	void RemoveSubDir(FStringView SubDirBaseName);
 
 	/** Find the index of the subdir with the given Relative path. */
@@ -251,8 +296,9 @@ protected:
 	uint8 PriorityRefCounts[CountEPriority] = {}; // Initialize all elements to 0
 	EPriority DirectPriority = EPriority::Normal;
 	EPriority AccumulatedPriority = EPriority::Normal;
+	/** Whether each piece of the inherited data has been set directly on this directory */
+	FInherited DirectData;
 	bool bHasScanned = false;
-	bool bIsDirectWhitelisted = false;
 	bool bScanInFlight = false;
 	bool bScanInFlightInvalidated = false;
 	bool bIsComplete = false;
@@ -272,21 +318,16 @@ public:
 	FMountDir(FAssetDataDiscovery& InDiscovery, FStringView LocalAbsPath, FStringView PackagePath);
 	~FMountDir();
 
-	/** The local filesystem path from FPackageName, as an absolute path d:\root\Engine\Content rather than a relative path ../../../Engine/Content. */
+	/** The local path from FPackageName, absolute d:\root\Engine\Content rather relative ../../../Engine/Content. */
 	FStringView GetLocalAbsPath() const;
 	/** The package path from FPackageName. */
 	FStringView GetLongPackageName() const;
 
-	/**
-	 * Get the set of blacklist paths that are child paths of this MountDir.
-	 * Returned paths are relative paths from the MountDir's root.
-	 */
-	const TArray<TPair<FString,bool>>& GetBlacklist() const;
 	/** Return the FAssetDataDiscovery that owns this FMountDir. */
 	FAssetDataDiscovery& GetDiscovery() const;
 
 	/** Find the direct parent of InRelPath, or the lowest fallback. See FScanDir::GetControllingDir. */
-	FScanDir* GetControllingDir(FStringView LocalAbsPath, bool bIsDirectory, bool& bOutIsWhitelisted, FString& OutRelPath);
+	FScanDir* GetControllingDir(FStringView LocalAbsPath, bool bIsDirectory, FScanDir::FInherited& OutParentData, FString& OutRelPath);
 	/** Return the memory used by the tree under this MountDir, except that sizeof(*this) is excluded. */
 	uint32 GetAllocatedSize() const;
 
@@ -295,15 +336,13 @@ public:
 	/** Return this MountDir's Priority, which is the maximum of any of its ScanDirs' priorities. */
 	EPriority GetPriority() const;
 
-	/** Return whether the given path is a child path of *this and is whitelisted. */
-	bool IsPathWhitelisted(FStringView LocalAbsPath) const;
-	/** Return whether the given path is a child path of *this and matches the blacklist filters. */
-	bool IsBlacklisted(FStringView LocalAbsPath) const;
+	/** Report the collapsed data for the scandir - whitelist, blacklist, etc. Returns false data for non child paths. */
+	void GetMonitorData(FStringView InLocalAbsPath, FScanDir::FInherited& OutData) const;
 	/**
 	 * Return whether the given path is a child path of *this and is whitelisted and is not blacklisted, which means
 	 * it will be or has been scanned.
 	 */
-	bool IsMonitored(FStringView LocalAbsPath) const;
+	bool IsMonitored(FStringView InLocalAbsPath) const;
 
 	/*
 	 * Set values of fields on the given directory for all of the properties existing on InProperties.
@@ -314,12 +353,12 @@ public:
 	/** 
 	 * Update all ScanDirs under this MountDir in priority order until one is found that needs to be scanned.
 	 * Set OutCursor to that ScanDir, or to nullptr if all are complete.
-	 * Set bOutCursorParentIsWhitelisted to whether the parent of the OutCursor is whitelisted,
-	 * or to false if there is no OutCursor.
+	 * Set OutData for the inherited parent data of the OutCursor, set to false data if 
+	 * if OutCursor is null.
 	 */
-	void Update(FScanDir*& OutCursor, bool& bOutCursorParentIsWhitelisted);
+	void Update(FScanDir*& OutCursor, FScanDir::FInherited& OutParentData);
 
-	/** Record that a directory at or under the MountDir's root has been scanned, used to detect if configuration occurs after scanning starts. */
+	/** Record a directory under the MountDir has been scanned, used to detect if configuration occurs after scanning. */
 	void SetHasStartedScanning();
 	/** Minimize data in internal buffers. */
 	void Shrink();
@@ -328,9 +367,15 @@ public:
 	void AddChildMount(FMountDir* ChildMount);
 	/** Mark that a childpath MountDir is being deleted and the childpath should be scanned again by this. */
 	void RemoveChildMount(FMountDir* ChildMount);
-	/** Remove all childmounts. Does not handle properly updating the MountDir to reown those paths; this is used during destruction all MountDirs. */
+	/**
+	 * Remove all childmounts. Does not handle properly updating the MountDir to reown those paths;
+	 * this is used during destruction all MountDirs.
+	 */
 	void OnDestroyClearChildMounts();
-	/** Record the backpointer to the parent mountdir that this mountdir's path is a child path of, or null if the parent no longer exists. */
+	/**
+	 * Record the backpointer to the parent mountdir that this mountdir's path is a child path of,
+	 * or null if the parent no longer exists.
+	 */
 	void SetParentMount(FMountDir* ParentMount);
 	/** Return the parent MountDir. */
 	FMountDir* GetParentMount() const;
@@ -343,40 +388,35 @@ protected:
 	/** Mark that given path needs to be reconsidered by Update. */
 	void MarkDirty(FStringView MountRelPath);
 
-	/** Holds data about child paths directly in the MountDir rather than requiring FScanDir children. This is more memory efficient and is performant for rarely accessed data. */
-	struct FPathData
-	{
-		FPathData(FStringView MountRelPath);
-		bool IsEmpty() const;
-		uint32 GetAllocatedSize() const;
+	/** Add a ChildMountPath, which indicates that another MountPoint owns that path and this Mount should not scan it. */
+	void AddChildMountPath(FStringView MountRelPath);
+	/** Remove a ChildMountPath, return whether there was a match to be removed. */
+	bool RemoveChildMountPath(FStringView MountRelPath);
+	/** Report whether the given Path is equal to or a child of an existing ChildMountPath. */
+	bool IsChildMountPath(FStringView MountRelPath) const;
 
-		/** Relative path from the mountdir. */
-		FString RelPath;
-		/** The path has been marked as should no longer be blacklisted by e.g. ScanPathsSynchronous. */
-		bool bIgnoreBlacklist = false;
-		/** The path is a childpath; another MountDir exists to scan it and it should not be scanned by this MountDir. */
-		bool bIsChildPath = false;
-	};
-	/** Return the PathData with the given MountRelPath, or nullptr if none exists. */
-	FPathData* FindPathData(FStringView MountRelPath);
-	/** Return the PathData with the given MountRelPath, creating it if none exists. */
-	FPathData& FindOrAddPathData(FStringView MountRelPath);
-	/** Remove the PathData with the given MountRelPath if one exists. */
-	void RemovePathData(FStringView MountRelPath);
-
-	/** Array of data keyed by relative path from the MountDir for information about child paths that should not be stored in ScanDirs. */
-	TArray<FPathData> PathDatas;
-	/** Set of relative path from the MountDir paths that should not be scanned, because they were requested blacklisted by clients or because a childmount owns them. */
-	TArray<TPair<FString,bool>> BlacklistedRelPaths;
+	/** Child mount paths; these directories should not be scanned by this MountDir. */
+	TArray<FString> ChildMountPaths;
+	/**
+	 * Set of relative path from the MountDir paths that should not be scanned, because they were requested
+	 * blacklisted by clients or because a childmount owns them.
+	 */
+	TSet<FString> BlacklistedRelPaths;
 	/** Absolute path to the root of the MountDir in the local file system. */
 	FString LocalAbsPath;
 	/** LongPackageName that was assigned to the MountDir in FPackageName. */
 	FString LongPackageName;
-	/** ScanDir for the root directory of this MountDir; child paths to scan will be created (and destroyed after use) as children of the ScanDir (with the exception of childmounts). */
+	/**
+	 * ScanDir for the root directory of this MountDir; child paths to scan will be created (and destroyed after use)
+	 * as children of the ScanDir (with the exception of childmounts).
+	 */
 	TRefCountPtr<FScanDir> Root;
 	/** Backpointer to the Discovery that owns this MountDir. */
 	FAssetDataDiscovery& Discovery;
-	/** If this is a nested MountDir (a nested path was registered with FPackageName) ParentMount is a pointer to the FMountDir that corresponds to the registered parent directory. */
+	/**
+	 * If this is a nested MountDir (a nested path was registered with FPackageName) ParentMount is a
+	 * pointer to the FMountDir that corresponds to the registered parent directory.
+	 */
 	FMountDir* ParentMount = nullptr;
 	/** Records whether any directory at or under the MountDir's root has been scanned. */
 	bool bHasStartedScanning = false;
@@ -387,7 +427,8 @@ protected:
 class FAssetDataDiscovery : public FRunnable
 {
 public:
-	FAssetDataDiscovery(const TArray<FString>& InBlacklistLongPackageNames, const TArray<FString>& InBlacklistMountRelativePaths, bool bInIsSynchronous);
+	FAssetDataDiscovery(const TArray<FString>& InBlacklistLongPackageNames,
+		const TArray<FString>& InBlacklistMountRelativePaths, bool bInIsSynchronous);
 	virtual ~FAssetDataDiscovery();
 
 
@@ -407,30 +448,36 @@ public:
 
 	// Receiving Results and reading properties (possibly while tick is running)
 	/** Gets search results from the file discovery. */
-	void GetAndTrimSearchResults(bool& bOutIsComplete, TArray<FString>& OutDiscoveredPaths, TRingBuffer<FGatheredPathData>& OutDiscoveredFiles, int32& OutNumPathsToSearch);
+	void GetAndTrimSearchResults(bool& bOutIsComplete, TArray<FString>& OutDiscoveredPaths,
+		TRingBuffer<FGatheredPathData>& OutDiscoveredFiles, int32& OutNumPathsToSearch);
 	/** Wait (joining in on the tick) until all currently monitored paths have been scanned. */
 	void WaitForIdle();
 	/** Optionally set some scan properties for the given path and then wait for the scan of it to finish. */
-	void SetPropertiesAndWait(const FString& LocalAbsPath, bool bAddToWhitelist, bool bForceRescan, bool bIgnoreBlackListScanFilters);
+	void SetPropertiesAndWait(const FString& LocalAbsPath, bool bAddToWhitelist, bool bForceRescan,
+		bool bIgnoreBlackListScanFilters);
 	/** Return whether the given path is whitelisted due to e.g. TrySetDirectoryProperties with IsWhitelisted. */
 	bool IsWhitelisted(FStringView LocalAbsPath) const;
-	/** Return whether the given path is blacklisted due to matching blacklist settings and not having been marked IgnoreBlacklist. */
+	/** Return whether the given path matches blacklist and has not been marked IgnoreBlacklist. */
 	bool IsBlacklisted(FStringView LocalAbsPath) const;
-	/** Return whether the given path should be scanned (possibly already has been) because it is whitelisted and not blacklisted. */
+	/** Return whether the given path should or has been scanned because it is whitelisted and not blacklisted. */
 	bool IsMonitored(FStringView LocalAbsPath) const;
 	/** Return the memory used by *this. sizeof(*this) is not included. */
 	uint32 GetAllocatedSize() const;
 
 	// Events and setting of properties (possibly while tick is running)
-	/** Register the given LocalAbsPath/LongPackageName pair that came from FPackageName's list of mount points as a mountpoint to track. Will not be scanned until whitelisted. */
+	/**
+	 * Register the given LocalAbsPath/LongPackageName pair that came from FPackageName's list of mount points as
+	 * a mountpoint to track. Will not be scanned until whitelisted.
+	 */
 	void AddMountPoint(const FString& LocalAbsPath, FStringView LongPackageName);
 	/** Remove the mountpoint because FPackageName has removed it. */
 	void RemoveMountPoint(const FString& LocalAbsPath);
-	/** Set properties on the directory, called from AssetDataGatherer when files are requested whitelisted, blacklisted, or rescanned. */
-	bool TrySetDirectoryProperties(const FString& LocalAbsPath, const UE::AssetDataGather::Private::FSetPathProperties& Properties, bool bConfirmedExists);
-	/** Event called from the directory watcher when a directory is created under a mount point. It will be added to the scan if IsMonitored. */
+	/** Set properties on the directory, called when files are requested whitelisted, blacklisted, or rescanned. */
+	bool TrySetDirectoryProperties(const FString& LocalAbsPath,
+		const UE::AssetDataGather::Private::FSetPathProperties& Properties, bool bConfirmedExists);
+	/** Event called from the directory watcher when a directory is created. It will be scanned if IsMonitored. */
 	void OnDirectoryCreated(FStringView LocalAbsPath);
-	/** Event called from the directory watcher when files are created under a mount point. Each will be reported to the Gatherer if IsMonitored. */
+	/** Event called from the directory watcher when files are created. Each will be reported if IsMonitored. */
 	void OnFilesCreated(TConstArrayView<FString> LocalAbsPaths);
 
 private:
@@ -439,7 +486,10 @@ private:
 	const FMountDir* FindContainingMountPoint(FStringView LocalAbsPath) const;
 	/** Find the mountpoint with the given root. */
 	FMountDir* FindMountPoint(FStringView LocalAbsPath);
-	/** Find the mountpoint with the given root, creating it with root and packagename if it does not already exist. Handle registration of child mountdirs. */
+	/**
+	 * Find the mountpoint with the given root, creating it if it does not already exist.
+	 * Handle registration of child mountdirs.
+	 */
 	FMountDir& FindOrAddMountPoint(FStringView LocalAbsPath, FStringView LongPackageName);
 	/** Execute the search on the sorted-by-path mountdirs, returning the first mountdir with root >= the given path. */
 	int32 FindLowerBoundMountPoint(FStringView LocalAbsPath) const;
@@ -447,10 +497,10 @@ private:
 	/** Run the tick, either called from the async Run or called on thread from a thread executing a synchronous wait. */
 	void TickInternal();
 	/** Search all of the MountDirs for a highest priority ScanDir that needs to be Updated. */
-	void FindFirstCursor(FScanDir*& OutCursor, bool& bOutCursorParentIsWhitelisted);
-	/** Invalidate the current cursor set from the tick, called when properties are set and we need to reconsider whether directories ShouldBeScanned. */
+	void FindFirstCursor(FScanDir*& OutCursor, FScanDir::FInherited& OutParentData);
+	/** Invalidate the tick's current cursor, called when we need to reconsider whether directories ShouldBeScanned. */
 	void InvalidateCursor();
-	/** Mark that this discoverer has either finished and is idle, or has received new data and is no longer idle. Update properties dependent upon the idle state. */
+	/** Mark whether this discoverer has finished and is idle. Update properties dependent upon the idle state. */
 	void SetIsIdle(bool bInIdle);
 
 	/** Store the given discovered files and directories in the results. */
@@ -459,14 +509,16 @@ private:
 	 * Return whether a directory with the given LongPackageName should be reported to the AssetRegistry
 	 * We do not report some directories because they are paths that should not enter the AssetRegistry list of paths if empty,
 	 * and reporting a path to the AssetRegistry adds it unconditionally to the list of paths.
-	 * If ShouldDirBeReported returns false, the directory will still be added to the catalog if non-empty, because the AssetRegistry adds the path of every added file.
+	 * If ShouldDirBeReported returns false, the directory will still be added to the catalog if non-empty,
+	 * because the AssetRegistry adds the path of every added file.
 	 */
 	bool ShouldDirBeReported(FStringView LongPackageName) const;
 
 	/** Handle the actions necessary for a single created file. */
 	void OnFileCreated(const FString& LocalPaths);
-	/** Helper for TrySetDirectoryProperties and other internals that need to set it, called from within the critical section. */
-	bool TrySetDirectoryPropertiesInternal(const FString& LocalAbsPath, const UE::AssetDataGather::Private::FSetPathProperties& Properties, bool bConfirmedExists);
+	/** Helper for TrySetDirectoryProperties, called from within the critical section. */
+	bool TrySetDirectoryPropertiesInternal(const FString& LocalAbsPath,
+		const UE::AssetDataGather::Private::FSetPathProperties& Properties, bool bConfirmedExists);
 
 	/** Add the given path as a MountPoint and update child registrations. */
 	void AddMountPointInternal(const FString& LocalAbsPath, FStringView LongPackageName);
@@ -505,9 +557,15 @@ private:
 
 	// Variable section for variables that are constant during threading.
 
-	/** Blacklisted paths that are full absolute paths. Child paths will not be scanned unless requested to ignore blacklists. Constant during threading. */
+	/**
+	 * Blacklisted full absolute paths. Child paths will not be scanned unless requested to ignore blacklists.
+	 * Constant during threading.
+	 */
 	TArray<FString> BlacklistLongPackageNames;
-	/** Blacklisted paths that are relative paths in each mount. Child paths will not be scanned unless requested to ignore blacklists. Constant during threading. */
+	/**
+	 * Blacklisted relative paths in each mount. Child paths will not be scanned unless requested to ignore blacklists.
+	 * Constant during threading.
+	 */
 	TArray<FString> BlacklistMountRelativePaths;
 	/** LongPackageNames for directories that should not be reported, see ShouldDirBeReported. Constant during threading. */
 	TSet<FString> DirLongPackageNamesToNotReport;
@@ -526,7 +584,10 @@ private:
 	std::atomic<bool> bIsIdle;
 	/** > 0 if we've been asked to abort work in progress at the next opportunity. Read/writable anywhere. */
 	std::atomic<uint32> IsStopped;
-	/** > 0 if we've been asked to pause the worker thread so a synchronous function can take over the tick. Read/writable anywhere. */
+	/**
+	 * > 0 if we've been asked to pause the worker thread so a synchronous function can take over the tick.
+	 * Read/writable anywhere.
+	 */
 	mutable std::atomic<uint32> IsPaused;
 	/**
 	 * Number of directories that have been discovered and IsMonitored but have not yet been scanned.
@@ -541,9 +602,15 @@ private:
 	TArray<FString> DiscoveredDirectories;
 	/** Files found found in the scan. Read/writable only within ResultsLock. */
 	TArray<FGatheredPathData> DiscoveredFiles;
-	/** Time at which the scan was started or last resumed from idle. Used for logging. Read/writable only within ResultsLock. */
+	/**
+	 * Time at which the scan was started or last resumed from idle. Used for logging.
+	 * Read/writable only within ResultsLock.
+	 */
 	double DiscoverStartTime = 0.;
-	/** Number of files discovered during scanning since start or resumed from idle. Used for logging. Read/writable only within ResultsLock. */
+	/**
+	 * Number of files discovered during scanning since start or resumed from idle. Used for logging.
+	 * Read/writable only within ResultsLock.
+	 */
 	int32 NumDiscoveredFiles = 0;
 
 
@@ -557,8 +624,8 @@ private:
 	TArray<TUniquePtr<FMountDir>> MountDirs;
 	/** The next ScanDir to update in Tick. Read/writable only within TreeLock. */
 	TRefCountPtr<FScanDir> Cursor = nullptr;
-	/** The value of ParentIsWhitelisted for Cursor (or false if Cursor is null). Read/writable only within TreeLock. */
-	bool bCursorParentIsWhitelisted = false;
+	/** The value of Parent's InheritedData for Cursor (or false if Cursor is null). Read/writable only within TreeLock. */
+	FScanDir::FInherited CursorParentData;
 
 
 	// Variable section for variables that are read/writable only within TickLock.
diff --git a/Engine/Source/Runtime/AssetRegistry/Private/AssetRegistry.cpp b/Engine/Source/Runtime/AssetRegistry/Private/AssetRegistry.cpp
--- a/Engine/Source/Runtime/AssetRegistry/Private/AssetRegistry.cpp
+++ b/Engine/Source/Runtime/AssetRegistry/Private/AssetRegistry.cpp
@@ -1721,7 +1721,11 @@ bool UAssetRegistryImpl::AddPath(const FString& PathToAdd)
 	// If no GlobalGatherer, then we are in the game or non-cook commandlet and we do not implement blacklisting
 	if (GlobalGatherer.IsValid())
 	{
-		bBlacklisted = GlobalGatherer->IsBlacklisted(PathToAdd);
+		FString LocalPathToAdd;
+		if (FPackageName::TryConvertLongPackageNameToFilename(PathToAdd, LocalPathToAdd))
+		{
+			bBlacklisted = GlobalGatherer->IsBlacklisted(LocalPathToAdd);
+		}
 	}
 	if (bBlacklisted)
 	{
diff --git a/Engine/Source/Runtime/AssetRegistry/Private/AssetRegistryConsoleCommands.h b/Engine/Source/Runtime/AssetRegistry/Private/AssetRegistryConsoleCommands.h
--- a/Engine/Source/Runtime/AssetRegistry/Private/AssetRegistryConsoleCommands.h
+++ b/Engine/Source/Runtime/AssetRegistry/Private/AssetRegistryConsoleCommands.h
@@ -4,6 +4,7 @@
 
 #include "CoreMinimal.h"
 #include "HAL/IConsoleManager.h"
+#include "HAL/FileManager.h"
 #include "Misc/PackageName.h"
 #include "AssetRegistry/AssetData.h"
 #include "Misc/Paths.h"
