Fix AssetDataGatherer's ScanPathsSynchronous failing to scan a file if the path is not whitelisted, even though bAddToWhitelist is true.
Do not call TrySetDirectoryProperties without first verifying the path exists.
#rnx

diff --git a/Engine/Source/Runtime/AssetRegistry/Private/AssetDataGatherer.cpp b/Engine/Source/Runtime/AssetRegistry/Private/AssetDataGatherer.cpp
--- a/Engine/Source/Runtime/AssetRegistry/Private/AssetDataGatherer.cpp
+++ b/Engine/Source/Runtime/AssetRegistry/Private/AssetDataGatherer.cpp
@@ -1110,15 +1110,25 @@ void FMountDir::UpdateBlacklist()
 			AddedBlacklists.Add(New);
 		}
 	}
+
+	TStringBuilder<256> BlacklistAbsPath;
+	IFileManager& FileManager = IFileManager::Get();
 	FSetPathProperties ChangeBlacklist;
 	ChangeBlacklist.MatchesBlacklist = true;
 	for (const FString& RelPath : AddedBlacklists)
 	{
-		Root->TrySetDirectoryProperties(RelPath, ChangeBlacklist, false /* bConfirmedExists */);
+		BlacklistAbsPath.Reset();
+		BlacklistAbsPath << LocalAbsPath;
+		FPathViews::AppendPath(BlacklistAbsPath, RelPath);
+		if (FileManager.DirectoryExists(BlacklistAbsPath.ToString()))
+		{
+			Root->TrySetDirectoryProperties(RelPath, ChangeBlacklist, true /* bConfirmedExists */);
+		}
 	}
 	ChangeBlacklist.MatchesBlacklist = false;
 	for (const FString& RelPath : RemovedBlacklists)
 	{
+		// We don't need to check for existence when setting the removal property, because the scandir already exists
 		Root->TrySetDirectoryProperties(RelPath, ChangeBlacklist, false /* bConfirmedExists */);
 	}
 }
@@ -1676,7 +1686,10 @@ void FAssetDataDiscovery::SetPropertiesAndWait(const FString& LocalAbsPath, bool
 		FString RelPath;
 		FScanDir::FInherited MonitorData;
 		TRefCountPtr<FScanDir> ScanDir = MountDir->GetControllingDir(LocalAbsPath, StatData.bIsDirectory, MonitorData, RelPath);
-		if (!ScanDir || !MonitorData.IsMonitored())
+		bool bIsWhitelistedInThisCall = MonitorData.IsWhitelisted() || bAddToWhitelist;
+		bool bIsBlacklistedInThisCall = MonitorData.IsBlacklisted() && !bIgnoreBlackListScanFilters;
+		bool bIsMonitoredInThisCall = bIsWhitelistedInThisCall && !bIsBlacklistedInThisCall;
+		if (!ScanDir || !bIsMonitoredInThisCall)
 		{
 			UE_LOG(LogAssetRegistry, Log, TEXT("SetPropertiesAndWait called on %.*s which is not monitored. Call will be ignored."), LocalAbsPath.Len(), *LocalAbsPath);
 			return;
