Extend suppression of V792 over both lines of the statement.
#rnx

diff --git a/Engine/Source/Runtime/AssetRegistry/Private/AssetDataGathererPrivate.h b/Engine/Source/Runtime/AssetRegistry/Private/AssetDataGathererPrivate.h
--- a/Engine/Source/Runtime/AssetRegistry/Private/AssetDataGathererPrivate.h
+++ b/Engine/Source/Runtime/AssetRegistry/Private/AssetDataGathererPrivate.h
@@ -59,7 +59,7 @@ struct FSetPathProperties
 	/** Used to early-exit from tree traversal when all properties have finished being handled */
 	bool IsSet() const
 	{
-		return IsWhitelisted.IsSet() | HasScanned.IsSet() | MatchesBlacklist.IsSet() |
+		return IsWhitelisted.IsSet() | HasScanned.IsSet() | MatchesBlacklist.IsSet() | //-V792
 			IgnoreBlacklist.IsSet() | Priority.IsSet(); //-V792
 	}
 };
