AssetRegistry::ScanPathsSynchronous: Handle reentrancy due to arbitrary code execution in FAssetData constructor.
#rnx

diff --git a/Engine/Source/Runtime/AssetRegistry/Private/AssetRegistry.cpp b/Engine/Source/Runtime/AssetRegistry/Private/AssetRegistry.cpp
--- a/Engine/Source/Runtime/AssetRegistry/Private/AssetRegistry.cpp
+++ b/Engine/Source/Runtime/AssetRegistry/Private/AssetRegistry.cpp
@@ -1916,6 +1916,7 @@ void UAssetRegistryImpl::Tick(float DeltaTime)
 
 void UAssetRegistryImpl::TickGatherer(float DeltaTime, TArray<FName>* OutFoundAssets)
 {
+	// Note this function can be reentered due to arbitrary code execution in construction of FAssetData
 	if (!GlobalGatherer.IsValid())
 	{
 		return;
@@ -2214,6 +2215,7 @@ const TSet<FName>& UAssetRegistryImpl::GetCachedEmptyPackages() const
 
 void UAssetRegistryImpl::ScanPathsSynchronousInternal(const TArray<FString>& InDirs, const TArray<FString>& InFiles, bool bForceRescan, bool bIgnoreBlackListScanFilters, TArray<FName>* OutFoundAssets)
 {
+	// Note this function can be reentered due to arbitrary code execution in construction of FAssetData
 	LLM_SCOPE(ELLMTag::AssetRegistry);
 	const double SearchStartTime = FPlatformTime::Seconds();
 	if (OutFoundAssets)
@@ -2978,6 +2980,7 @@ void UAssetRegistryImpl::OnAssetLoaded(UObject *AssetLoaded)
 
 void UAssetRegistryImpl::ProcessLoadedAssetsToUpdateCache(const double TickStartTime, bool bUpdateDeferredList)
 {
+	// Note this function can be reentered due to arbitrary code execution in construction of FAssetData
 	LLM_SCOPE(ELLMTag::AssetRegistry);
 
 	const bool bFlushFullBuffer = TickStartTime < 0;
@@ -2985,7 +2988,11 @@ void UAssetRegistryImpl::ProcessLoadedAssetsToUpdateCache(const double TickStart
 	if (bFlushFullBuffer || bUpdateDeferredList)
 	{
 		// Retry the previous failures on a flush
-		LoadedAssetsToProcess.Append(LoadedAssetsThatDidNotHaveCachedData);
+		LoadedAssetsToProcess.Reserve(LoadedAssetsToProcess.Num() + LoadedAssetsThatDidNotHaveCachedData.Num());
+		for (TWeakObjectPtr<UObject>& Asset : LoadedAssetsThatDidNotHaveCachedData)
+		{
+			LoadedAssetsToProcess.Add(MoveTemp(Asset));
+		}
 		LoadedAssetsThatDidNotHaveCachedData.Reset();
 	}
 
@@ -2993,10 +3000,9 @@ void UAssetRegistryImpl::ProcessLoadedAssetsToUpdateCache(const double TickStart
 	CollectCodeGeneratorClasses();
 
 	// Add the found assets
-	int32 LoadedAssetIndex = 0;
-	for (LoadedAssetIndex = 0; LoadedAssetIndex < LoadedAssetsToProcess.Num(); ++LoadedAssetIndex)
+	while (!LoadedAssetsToProcess.IsEmpty())
 	{
-		UObject* LoadedAsset = LoadedAssetsToProcess[LoadedAssetIndex].Get();
+		UObject* LoadedAsset = LoadedAssetsToProcess.PopFrontValue().Get();
 
 		if (!LoadedAsset)
 		{
@@ -3050,17 +3056,9 @@ void UAssetRegistryImpl::ProcessLoadedAssetsToUpdateCache(const double TickStart
 		// Check to see if we have run out of time in this tick
 		if (!bFlushFullBuffer && (FPlatformTime::Seconds() - TickStartTime) > MaxSecondsPerFrame)
 		{
-			// Increment the index to properly trim the buffer below
-			++LoadedAssetIndex;
 			break;
 		}
 	}
-
-	// Trim the results array
-	if (LoadedAssetIndex > 0)
-	{
-		LoadedAssetsToProcess.RemoveAt(0, LoadedAssetIndex);
-	}
 }
 
 void UAssetRegistryImpl::UpdateRedirectCollector()
diff --git a/Engine/Source/Runtime/AssetRegistry/Private/AssetRegistry.h b/Engine/Source/Runtime/AssetRegistry/Private/AssetRegistry.h
--- a/Engine/Source/Runtime/AssetRegistry/Private/AssetRegistry.h
+++ b/Engine/Source/Runtime/AssetRegistry/Private/AssetRegistry.h
@@ -439,7 +439,7 @@ private:
 
 #if WITH_EDITOR
 	/** List of loaded objects that need to be processed */
-	TArray<TWeakObjectPtr<UObject>> LoadedAssetsToProcess;
+	TRingBuffer<TWeakObjectPtr<UObject>> LoadedAssetsToProcess;
 
 	/** Objects that couldn't be processed because the asset data didn't exist, reprocess these after more directories are scanned */
 	TArray<TWeakObjectPtr<UObject>> LoadedAssetsThatDidNotHaveCachedData;
