Fix issue with UAssetRegistryImpl::ScanModifiedAssetFiles() & external actors
* Was triggering both OnAssetUpdated() & OnAssetRemoved()

diff --git a/Engine/Source/Runtime/AssetRegistry/Private/AssetRegistry.cpp b/Engine/Source/Runtime/AssetRegistry/Private/AssetRegistry.cpp
--- a/Engine/Source/Runtime/AssetRegistry/Private/AssetRegistry.cpp
+++ b/Engine/Source/Runtime/AssetRegistry/Private/AssetRegistry.cpp
@@ -1914,7 +1914,7 @@ void UAssetRegistryImpl::Tick(float DeltaTime)
 	TickGatherer(DeltaTime);
 }
 
-void UAssetRegistryImpl::TickGatherer(float DeltaTime, TArray<FName>* OutFoundAssets)
+void UAssetRegistryImpl::TickGatherer(float DeltaTime, TOptional<UAssetRegistryImpl::FAssetsFoundCallback> AssetsFoundCallback)
 {
 	// Note this function can be reentered due to arbitrary code execution in construction of FAssetData
 	if (!GlobalGatherer.IsValid())
@@ -1963,14 +1963,9 @@ void UAssetRegistryImpl::TickGatherer(float DeltaTime, TArray<FName>* OutFoundAs
 		{
 			AmortizeStartTime = FPlatformTime::Seconds();
 		}
-		if (OutFoundAssets)
+		if (AssetsFoundCallback.IsSet())
 		{
-			OutFoundAssets->Reset();
-			OutFoundAssets->Reserve(BackgroundAssetResults.Num());
-			for (FAssetData* AssetData : BackgroundAssetResults)
-			{
-				OutFoundAssets->Add(AssetData->ObjectPath);
-			}
+			AssetsFoundCallback.GetValue()(BackgroundAssetResults);
 		}
 
 		AssetSearchDataGathered(TickStartTime, BackgroundAssetResults);
@@ -2311,35 +2306,41 @@ void UAssetRegistryImpl::ScanPathsSynchronousInternal(const TArray<FString>& InD
 		Gatherer.LoadCacheFile(CacheFilename);
 	}
 
-	TArray<FName> FoundAssetsBuffer;
-	TArray<FName>* FoundAssets = OutFoundAssets;
-	if (!FoundAssets)
-	{
-		FoundAssets = &FoundAssetsBuffer;
-	}
 	Gatherer.ScanPathsSynchronous(LocalPaths, bForceRescan, bIgnoreBlackListScanFilters, CacheFilename, PackageDirs);
-	TickGatherer(-1.f, FoundAssets);
-	// The gatherer may have added other assets that were scanned as part of the ongoing background scan; remove any assets that were not in the requested paths
-	if (OutFoundAssets)
+
+	int32 NumFoundAssets = 0;
+	FAssetsFoundCallback AssetsFoundCallback = [&PackageFiles, &PackageDirs, &OutFoundAssets, &NumFoundAssets](const TRingBuffer<FAssetData*>& InFoundAssets)
 	{
-		FoundAssets->RemoveAll([&PackageFiles, &PackageDirs](FName ObjectPath)
+		NumFoundAssets = InFoundAssets.Num();
+
+		if (!OutFoundAssets)
+		{
+			return;
+		}
+
+		OutFoundAssets->Reset();
+		OutFoundAssets->Reserve(NumFoundAssets);
+
+		// The gatherer may have added other assets that were scanned as part of the ongoing background scan; remove any assets that were not in the requested paths
+		for (FAssetData* AssetData : InFoundAssets)
+		{
+			bool bIsInRequestedPaths = false;
+
+			TStringBuilder<128> PackageNameStr;
+			AssetData->PackageName.ToString(PackageNameStr);
+			FStringView PackageName(PackageNameStr.ToString(), PackageNameStr.Len());
+
+			for (const FString& RequestedPackageDir : PackageDirs)
 			{
-				TStringBuilder<128> ObjectPathStr;
-				ObjectPath.ToString(ObjectPathStr);
-				FStringView UnusedClassName;
-				FStringView PackageName;
-				FStringView UnusedObjectName;
-				FStringView UnusedSubObjectName;
-				FPackageName::SplitFullObjectPath(ObjectPathStr, UnusedClassName, PackageName, UnusedObjectName, UnusedSubObjectName);
-				bool bIsInRequestedPaths = false;
-				for (const FString& RequestedPackageDir : PackageDirs)
+				if (FPathViews::IsParentPathOf(RequestedPackageDir, PackageName))
 				{
-					if (FPathViews::IsParentPathOf(RequestedPackageDir, PackageName))
-					{
-						bIsInRequestedPaths = true;
-						break;
-					}
+					bIsInRequestedPaths = true;
+					break;
 				}
+			}
+
+			if (!bIsInRequestedPaths)
+			{
 				for (const FString& RequestedPackageFile : PackageFiles)
 				{
 					if (PackageName.Equals(RequestedPackageFile, ESearchCase::IgnoreCase))
@@ -2348,9 +2349,16 @@ void UAssetRegistryImpl::ScanPathsSynchronousInternal(const TArray<FString>& InD
 						break;
 					}
 				}
-				return !bIsInRequestedPaths;
-			});
-	}
+			}
+
+			if (bIsInRequestedPaths)
+			{
+				OutFoundAssets->Add(AssetData->ObjectPath);
+			}
+		}
+	};
+
+	TickGatherer(-1.f, AssetsFoundCallback);
 
 	// Log stats
 	FString PathsString;
@@ -2363,7 +2371,7 @@ void UAssetRegistryImpl::ScanPathsSynchronousInternal(const TArray<FString>& InD
 		PathsString = FString::Printf(TEXT("'%s'"), *LocalPaths[0]);
 	}
 
-	UE_LOG(LogAssetRegistry, Verbose, TEXT("ScanPathsSynchronous completed scanning %s to find %d assets in %0.4f seconds"), *PathsString, FoundAssets->Num(), FPlatformTime::Seconds() - SearchStartTime);
+	UE_LOG(LogAssetRegistry, Verbose, TEXT("ScanPathsSynchronous completed scanning %s to find %d assets in %0.4f seconds"), *PathsString, NumFoundAssets, FPlatformTime::Seconds() - SearchStartTime);
 }
 
 bool UAssetRegistryImpl::IsPathMounted(const FString& Path, const TSet<FString>& MountPointsNoTrailingSlashes, FString& StringBuffer) const
diff --git a/Engine/Source/Runtime/AssetRegistry/Private/AssetRegistry.h b/Engine/Source/Runtime/AssetRegistry/Private/AssetRegistry.h
--- a/Engine/Source/Runtime/AssetRegistry/Private/AssetRegistry.h
+++ b/Engine/Source/Runtime/AssetRegistry/Private/AssetRegistry.h
@@ -176,8 +176,10 @@ private:
 	void ConstructGatherer();
 	/**  Called to set up timing variables and launch the on-constructor SearchAllAssets async call */
 	void SearchAllAssetsInitialAsync();
+
+	typedef TFunctionRef<void(const TRingBuffer<FAssetData*>&)> FAssetsFoundCallback;
 	/** Read any results from the gatherer, if it is running asynchronously */
-	void TickGatherer(float DeltaTime, TArray<FName>* OutFoundAssets=nullptr);
+	void TickGatherer(float DeltaTime, TOptional<FAssetsFoundCallback> AssetsFoundCallback = TOptional<FAssetsFoundCallback>());
 	/** Internal handler for ScanPathsSynchronous */
 	void ScanPathsSynchronousInternal(const TArray<FString>& InDirs, const TArray<FString>& InFiles, bool bForceRescan, bool bIgnoreBlackListScanFilters, TArray<FName>* OutFoundAssets);
 
