Change TFunctionRef lambda assignment to auto to remove dangling lambda pointer.
#rnx

diff --git a/Engine/Source/Runtime/AssetRegistry/Private/AssetRegistry.cpp b/Engine/Source/Runtime/AssetRegistry/Private/AssetRegistry.cpp
--- a/Engine/Source/Runtime/AssetRegistry/Private/AssetRegistry.cpp
+++ b/Engine/Source/Runtime/AssetRegistry/Private/AssetRegistry.cpp
@@ -2309,7 +2309,7 @@ void UAssetRegistryImpl::ScanPathsSynchronousInternal(const TArray<FString>& InD
 	Gatherer.ScanPathsSynchronous(LocalPaths, bForceRescan, bIgnoreBlackListScanFilters, CacheFilename, PackageDirs);
 
 	int32 NumFoundAssets = 0;
-	FAssetsFoundCallback AssetsFoundCallback = [&PackageFiles, &PackageDirs, &OutFoundAssets, &NumFoundAssets](const TRingBuffer<FAssetData*>& InFoundAssets)
+	auto AssetsFoundCallback = [&PackageFiles, &PackageDirs, &OutFoundAssets, &NumFoundAssets](const TRingBuffer<FAssetData*>& InFoundAssets)
 	{
 		NumFoundAssets = InFoundAssets.Num();
 
@@ -2358,7 +2358,7 @@ void UAssetRegistryImpl::ScanPathsSynchronousInternal(const TArray<FString>& InD
 		}
 	};
 
-	TickGatherer(-1.f, AssetsFoundCallback);
+	TickGatherer(-1.f, FAssetsFoundCallback(AssetsFoundCallback));
 
 	// Log stats
 	FString PathsString;
