Prevent an ensure when a newly created plugin is also mounted explicitly

diff --git a/Engine/Source/Runtime/AssetRegistry/Private/AssetRegistry.cpp b/Engine/Source/Runtime/AssetRegistry/Private/AssetRegistry.cpp
--- a/Engine/Source/Runtime/AssetRegistry/Private/AssetRegistry.cpp
+++ b/Engine/Source/Runtime/AssetRegistry/Private/AssetRegistry.cpp
@@ -3238,7 +3238,7 @@ void UAssetRegistryImpl::OnContentPathMounted(const FString& InAssetPath, const
 			// If the path doesn't exist on disk, make it so the watcher will work.
 			IFileManager::Get().MakeDirectory(*FileSystemPath);
 
-			if (ensure(!OnDirectoryChangedDelegateHandles.Contains(AssetPathWithTrailingSlash)))
+			if (!OnDirectoryChangedDelegateHandles.Contains(AssetPathWithTrailingSlash))
 			{
 				FDelegateHandle NewHandle;
 				DirectoryWatcher->RegisterDirectoryChangedCallback_Handle(
