Change UBlueprintGeneratedClasses to be assets so that cooked blueprints are viewable in editor.
But omit BPGC from the global AssetRegistry when they are in uncooked packages to avoid conflicting with the existing UBlueprint asset.
#rnx

diff --git a/Engine/Source/Editor/UnrealEd/Private/Commandlets/AssetRegistryGenerator.cpp b/Engine/Source/Editor/UnrealEd/Private/Commandlets/AssetRegistryGenerator.cpp
--- a/Engine/Source/Editor/UnrealEd/Private/Commandlets/AssetRegistryGenerator.cpp
+++ b/Engine/Source/Editor/UnrealEd/Private/Commandlets/AssetRegistryGenerator.cpp
@@ -2214,7 +2214,7 @@ const FAssetData* FAssetRegistryGenerator::CreateOrFindAssetData(UObject& Object
 	const FAssetData* const AssetData = State.GetAssetByObjectPath(FName(Object.GetPathName()));
 	if (!AssetData)
 	{
-		FAssetData* const NewAssetData = new FAssetData(&Object);
+		FAssetData* const NewAssetData = new FAssetData(&Object, true /* bAllowBlueprintClass */);
 		State.AddAssetData(NewAssetData);
 		return NewAssetData;
 	}
diff --git a/Engine/Source/Runtime/AssetRegistry/Private/AssetRegistry.cpp b/Engine/Source/Runtime/AssetRegistry/Private/AssetRegistry.cpp
--- a/Engine/Source/Runtime/AssetRegistry/Private/AssetRegistry.cpp
+++ b/Engine/Source/Runtime/AssetRegistry/Private/AssetRegistry.cpp
@@ -1014,6 +1014,14 @@ bool UAssetRegistryImpl::EnumerateAssets(const FARCompiledFilter& InFilter, TFun
 					return;
 				}
 
+#if WITH_EDITOR
+				// Skip classes that report themselves as assets but that the editor AssetRegistry is currently not counting as assets
+				if (ShouldSkipAsset(Obj->GetClass()->GetFName(), InMemoryPackage->GetPackageFlags()))
+				{
+					return;
+				}
+#endif
+
 				// Package name
 				const FName PackageName = InMemoryPackage->GetFName();
 
@@ -2478,8 +2486,17 @@ void UAssetRegistryImpl::AssetSearchDataGathered(const double TickStartTime, TRi
 	// Add the found assets
 	while (AssetResults.Num() > 0)
 	{
-		FAssetData* BackgroundResult = AssetResults.PopFrontValue();
-		CA_ASSUME(BackgroundResult);
+		// Delete or take ownership of the BackgroundResult; it was originally new'd by an FPackageReader
+		TUniquePtr<FAssetData> BackgroundResult(AssetResults.PopFrontValue());
+		CA_ASSUME(BackgroundResult.Get() != nullptr);
+
+#if WITH_EDITOR
+		// Skip Assets that we filter out
+		if (ShouldSkipAsset(BackgroundResult->AssetClass, BackgroundResult->PackageFlags))
+		{
+			continue;
+		}
+#endif
 
 		// Try to update any asset data that may already exist
 		FAssetData* AssetData = State.CachedAssetsByObjectPath.FindRef(BackgroundResult->ObjectPath);
@@ -2500,7 +2517,7 @@ void UAssetRegistryImpl::AssetSearchDataGathered(const double TickStartTime, TRi
 		if (AssetData)
 		{
 			// If this ensure fires then we've somehow processed the same result more than once, and that should never happen
-			if (ensure(AssetData != BackgroundResult))
+			if (ensure(AssetData != BackgroundResult.Get()))
 			{
 				// If the current AssetData came from a loaded asset, don't overwrite it with the new one from disk; loaded asset is more authoritative because it has run the postload steps
 #if WITH_EDITOR
@@ -2510,10 +2527,6 @@ void UAssetRegistryImpl::AssetSearchDataGathered(const double TickStartTime, TRi
 					// The asset exists in the cache from disk and has not yet been loaded into memory, update it with the new background data
 					UpdateAssetData(AssetData, *BackgroundResult);
 				}
-
-				// Delete the result that was originally created by an FPackageReader
-				delete BackgroundResult;
-				BackgroundResult = nullptr;
 			}
 		}
 		else
@@ -2521,7 +2534,7 @@ void UAssetRegistryImpl::AssetSearchDataGathered(const double TickStartTime, TRi
 			// The asset isn't in the cache yet, add it and notify subscribers
 			if (bPathIsMounted)
 			{
-				AddAssetData(BackgroundResult);
+				AddAssetData(BackgroundResult.Release());
 			}
 		}
 
@@ -3044,29 +3057,39 @@ void UAssetRegistryImpl::OnDirectoryChanged(const TArray<FFileChangeData>& FileC
 	ScanModifiedAssetFiles(ModifiedFiles);
 }
 
+bool UAssetRegistryImpl::ShouldSkipAsset(FName AssetClass, uint32 PackageFlags) const
+{
+	// We do not yet support having UBlueprintGeneratedClasses be assets when the UBlueprint is also
+	// an asset; the content browser does not handle the multiple assets correctly and displays this
+	// class asset as if it is in a separate package. Revisit when we have removed the UBlueprint as an asset
+	// or when we support multiple assets.
+	static FName NAME_BlueprintGeneratedClass = TEXT("BlueprintGeneratedClass");
+	bool bIsCooked = (PackageFlags & PKG_FilterEditorOnly) != 0;
+	if (!bIsCooked && AssetClass == NAME_BlueprintGeneratedClass)
+	{
+		return true;
+	}
+	return false;
+}
+
 void UAssetRegistryImpl::OnAssetLoaded(UObject *AssetLoaded)
 {
+	UPackage* Package = AssetLoaded->GetPackage();
+	if (Package && ShouldSkipAsset(AssetLoaded->GetClass()->GetFName(), Package->GetPackageFlags()))
+	{
+		return;
+	}
+
 	LoadedAssetsToProcess.Add(AssetLoaded);
 }
 
-void UAssetRegistryImpl::ProcessLoadedAssetsToUpdateCache(const double TickStartTime, bool bUpdateDeferredList)
+void UAssetRegistryImpl::ProcessLoadedAssetsToUpdateCache(const double TickStartTime, bool bBecameIdle)
 {
 	// Note this function can be reentered due to arbitrary code execution in construction of FAssetData
 	LLM_SCOPE(ELLMTag::AssetRegistry);
 
 	const bool bFlushFullBuffer = TickStartTime < 0;
 
-	if (bFlushFullBuffer || bUpdateDeferredList)
-	{
-		// Retry the previous failures on a flush
-		LoadedAssetsToProcess.Reserve(LoadedAssetsToProcess.Num() + LoadedAssetsThatDidNotHaveCachedData.Num());
-		for (TWeakObjectPtr<UObject>& Asset : LoadedAssetsThatDidNotHaveCachedData)
-		{
-			LoadedAssetsToProcess.Add(MoveTemp(Asset));
-		}
-		LoadedAssetsThatDidNotHaveCachedData.Reset();
-	}
-
 	// Refreshes ClassGeneratorNames if out of date due to module load
 	CollectCodeGeneratorClasses();
 
@@ -3101,27 +3124,28 @@ void UAssetRegistryImpl::ProcessLoadedAssetsToUpdateCache(const double TickStart
 			continue;
 		}
 
-		FAssetData** CachedData = State.CachedAssetsByObjectPath.Find(ObjectPath);
-		if (!CachedData)
-		{
-			// Not scanned, can't process right now but try again on next synchronous scan
-			LoadedAssetsThatDidNotHaveCachedData.Add(LoadedAsset);
-			continue;
-		}
+		FAssetData** DataFromGather = State.CachedAssetsByObjectPath.Find(ObjectPath);
 
 		AssetDataObjectPathsUpdatedOnLoad.Add(ObjectPath);
 
-		FAssetData NewAssetData = FAssetData(LoadedAsset);
-
-		if (NewAssetData.TagsAndValues != (*CachedData)->TagsAndValues)
+		if (!DataFromGather)
 		{
-			// We need to actually update disk cache
-			UpdateAssetData(*CachedData, NewAssetData);
+			FAssetData* NewAssetData = new FAssetData(LoadedAsset, true /* bAllowBlueprintClass */);
+			AddAssetData(NewAssetData);
 		}
 		else
 		{
-			// Bundle tags might have changed but CachedAssetsByTag is up to date
-			(*CachedData)->TaggedAssetBundles = NewAssetData.TaggedAssetBundles;
+			FAssetData NewAssetData = FAssetData(LoadedAsset, true /* bAllowBlueprintClass */);
+			if (NewAssetData.TagsAndValues != (*DataFromGather)->TagsAndValues)
+			{
+				// We need to actually update disk cache
+				UpdateAssetData(*DataFromGather, NewAssetData);
+			}
+			else
+			{
+				// Bundle tags might have changed but CachedAssetsByTag is up to date
+				(*DataFromGather)->TaggedAssetBundles = NewAssetData.TaggedAssetBundles;
+			}
 		}
 
 		// Check to see if we have run out of time in this tick
diff --git a/Engine/Source/Runtime/AssetRegistry/Private/AssetRegistry.h b/Engine/Source/Runtime/AssetRegistry/Private/AssetRegistry.h
--- a/Engine/Source/Runtime/AssetRegistry/Private/AssetRegistry.h
+++ b/Engine/Source/Runtime/AssetRegistry/Private/AssetRegistry.h
@@ -232,11 +232,14 @@ private:
 	/** Called when a file in a content directory changes on disk */
 	void OnDirectoryChanged(const TArray<struct FFileChangeData>& Files);
 
+	/** Called to check whether we should filter out assets of the given class and packageflags from the editor's AssetRegistry */
+	bool ShouldSkipAsset(FName AssetClass, uint32 PackageFlags) const;
+
 	/** Called when an asset is loaded, it will possibly update the cache */
 	void OnAssetLoaded(UObject *AssetLoaded);
 
 	/** Process Loaded Assets to update cache */
-	void ProcessLoadedAssetsToUpdateCache(const double TickStartTime, bool bUpdateDeferredList);
+	void ProcessLoadedAssetsToUpdateCache(const double TickStartTime, bool bBecameIdle);
 
 	/** Update Redirect collector with redirects loaded from asset registry */
 	void UpdateRedirectCollector();
@@ -446,9 +449,6 @@ private:
 	/** List of loaded objects that need to be processed */
 	TRingBuffer<TWeakObjectPtr<UObject>> LoadedAssetsToProcess;
 
-	/** Objects that couldn't be processed because the asset data didn't exist, reprocess these after more directories are scanned */
-	TArray<TWeakObjectPtr<UObject>> LoadedAssetsThatDidNotHaveCachedData;
-
 	/** The set of object paths that have had their disk cache updated from the in memory version */
 	TSet<FName> AssetDataObjectPathsUpdatedOnLoad;
 #endif
diff --git a/Engine/Source/Runtime/AssetRegistry/Private/AssetRegistryState.cpp b/Engine/Source/Runtime/AssetRegistry/Private/AssetRegistryState.cpp
--- a/Engine/Source/Runtime/AssetRegistry/Private/AssetRegistryState.cpp
+++ b/Engine/Source/Runtime/AssetRegistry/Private/AssetRegistryState.cpp
@@ -1495,11 +1495,18 @@ void FAssetRegistryState::AddAssetData(FAssetData* AssetData)
 {
 	NumAssets++;
 
+	FAssetData*& ExistingByObjectPath = CachedAssetsByObjectPath.FindOrAdd(AssetData->ObjectPath);
+	if (ExistingByObjectPath)
+	{
+		UE_LOG(LogAssetRegistry, Error, TEXT("AddAssetData called with ObjectPath %s which already exists. ")
+			TEXT("This will overwrite and leak the existing AssetData."), *AssetData->ObjectPath.ToString());
+	}
+	ExistingByObjectPath = AssetData;
+
 	TArray<FAssetData*, TInlineAllocator<1>>& PackageAssets = CachedAssetsByPackageName.FindOrAdd(AssetData->PackageName);
 	TArray<FAssetData*>& PathAssets = CachedAssetsByPath.FindOrAdd(AssetData->PackagePath);
 	TArray<FAssetData*>& ClassAssets = CachedAssetsByClass.FindOrAdd(AssetData->AssetClass);
 
-	CachedAssetsByObjectPath.Add(AssetData->ObjectPath, AssetData);
 	PackageAssets.Add(AssetData);
 	PathAssets.Add(AssetData);
 	ClassAssets.Add(AssetData);
diff --git a/Engine/Source/Runtime/Engine/Classes/Engine/BlueprintGeneratedClass.h b/Engine/Source/Runtime/Engine/Classes/Engine/BlueprintGeneratedClass.h
--- a/Engine/Source/Runtime/Engine/Classes/Engine/BlueprintGeneratedClass.h
+++ b/Engine/Source/Runtime/Engine/Classes/Engine/BlueprintGeneratedClass.h
@@ -746,6 +746,8 @@ public:
 #if WITH_EDITOR
 	virtual UClass* RegenerateClass(UClass* ClassToRegenerate, UObject* PreviousCDO) override;
 #endif	//WITH_EDITOR
+	virtual bool IsAsset() const override;
+
 	// End UObject interface
 
 	// UClass interface
diff --git a/Engine/Source/Runtime/Engine/Private/BlueprintGeneratedClass.cpp b/Engine/Source/Runtime/Engine/Private/BlueprintGeneratedClass.cpp
--- a/Engine/Source/Runtime/Engine/Private/BlueprintGeneratedClass.cpp
+++ b/Engine/Source/Runtime/Engine/Private/BlueprintGeneratedClass.cpp
@@ -162,8 +162,11 @@ void UBlueprintGeneratedClass::PostLoad()
 FPrimaryAssetId UBlueprintGeneratedClass::GetPrimaryAssetId() const
 {
 	FPrimaryAssetId AssetId;
-	if (!ensure(ClassDefaultObject))
+	if (!ClassDefaultObject)
 	{
+		// All UBlueprintGeneratedClass objects should have a pointer to their generated ClassDefaultObject, except for the
+		// CDO itself of the UBlueprintGeneratedClass class.
+		verify(HasAnyFlags(RF_ClassDefaultObject));
 		return AssetId;
 	}
 
@@ -1585,6 +1588,14 @@ UClass* UBlueprintGeneratedClass::RegenerateClass(UClass* ClassToRegenerate, UOb
 }
 #endif
 
+bool UBlueprintGeneratedClass::IsAsset() const
+{
+	// UClass::IsAsset returns false; override that to return true for BlueprintGeneratedClasses,
+	// but only if the instance satisfies the regular definition of an asset (RF_Public, not transient, etc)
+	// and only if it is the active BPGC matching the Blueprint
+	return UObject::IsAsset() && !HasAnyClassFlags(CLASS_NewerVersionExists);
+}
+
 void UBlueprintGeneratedClass::Link(FArchive& Ar, bool bRelinkExistingProperties)
 {
 	Super::Link(Ar, bRelinkExistingProperties);
