BlueprintGeneratedClass Assets: prevent BPGC from being used as the asset in WidgetBlueprints and when saving a package.
#rnx

diff --git a/Engine/Plugins/Editor/ContentBrowser/ContentBrowserAssetDataSource/Source/ContentBrowserAssetDataSource/Private/ContentBrowserAssetDataSource.cpp b/Engine/Plugins/Editor/ContentBrowser/ContentBrowserAssetDataSource/Source/ContentBrowserAssetDataSource/Private/ContentBrowserAssetDataSource.cpp
--- a/Engine/Plugins/Editor/ContentBrowser/ContentBrowserAssetDataSource/Source/ContentBrowserAssetDataSource/Private/ContentBrowserAssetDataSource.cpp
+++ b/Engine/Plugins/Editor/ContentBrowser/ContentBrowserAssetDataSource/Source/ContentBrowserAssetDataSource/Private/ContentBrowserAssetDataSource.cpp
@@ -1566,7 +1566,8 @@ void UContentBrowserAssetDataSource::OnAssetUpdated(const FAssetData& InAssetDat
 
 void UContentBrowserAssetDataSource::OnAssetLoaded(UObject* InAsset)
 {
-	if (InAsset && !InAsset->GetOutermost()->HasAnyPackageFlags(PKG_ForDiffing | PKG_PlayInEditor))
+	if (InAsset && !InAsset->GetOutermost()->HasAnyPackageFlags(PKG_ForDiffing | PKG_PlayInEditor) &&
+		!AssetRegistry->ShouldSkipAsset(InAsset))
 	{
 		FAssetData AssetData(InAsset);
 		if (ContentBrowserAssetData::IsPrimaryAsset(AssetData))
diff --git a/Engine/Source/Editor/PackagesDialog/Private/PackagesDialog.cpp b/Engine/Source/Editor/PackagesDialog/Private/PackagesDialog.cpp
--- a/Engine/Source/Editor/PackagesDialog/Private/PackagesDialog.cpp
+++ b/Engine/Source/Editor/PackagesDialog/Private/PackagesDialog.cpp
@@ -203,27 +203,40 @@ void FPackagesDialogModule::AddPackageItem(UPackage* InPackage, ECheckBoxState I
 	FAssetToolsModule& AssetToolsModule = FModuleManager::LoadModuleChecked<FAssetToolsModule>(TEXT("AssetTools"));
 
 	// Lookup for the first asset in the package
-	ForEachObjectWithPackage(InPackage, [&AssetName, &OwnerName, &AssetToolsModule](UObject* InnerObject)
+	UObject* FoundAsset = nullptr;
+	ForEachObjectWithPackage(InPackage, [&FoundAsset](UObject* InnerObject)
 	{
 		if (InnerObject->IsAsset())
 		{
-			TWeakPtr<IAssetTypeActions> AssetTypeActions = AssetToolsModule.Get().GetAssetTypeActionsForClass(InnerObject->GetClass());
-
-			if (AssetTypeActions.IsValid())
+			if (FAssetData(InnerObject).IsUAsset())
 			{
-				AssetName = *AssetTypeActions.Pin()->GetObjectDisplayName(InnerObject);
+				// If we found the primary asset, use it
+				return false;
 			}
-			else
+			// Otherwise, keep the first found asset but keep looking for a primary asset
+			if (!FoundAsset)
 			{
-				AssetName = InnerObject->GetFName();
+				FoundAsset = InnerObject;
 			}
-
-			OwnerName = InnerObject->GetOutermostObject()->GetFName();
-			return false;
 		}
 		return true;
 	}, /*bIncludeNestedObjects*/ false);
 
+	if (FoundAsset)
+	{
+		TWeakPtr<IAssetTypeActions> AssetTypeActions = AssetToolsModule.Get().GetAssetTypeActionsForClass(FoundAsset->GetClass());
+		if (AssetTypeActions.IsValid())
+		{
+			AssetName = *AssetTypeActions.Pin()->GetObjectDisplayName(FoundAsset);
+		}
+		else
+		{
+			AssetName = FoundAsset->GetFName();
+		}
+
+		OwnerName = FoundAsset->GetOutermostObject()->GetFName();
+	}
+
 	// Last resort, display the package name
 	if (AssetName == NAME_None)
 	{
diff --git a/Engine/Source/Editor/PackagesDialog/Private/SPackagesDialog.cpp b/Engine/Source/Editor/PackagesDialog/Private/SPackagesDialog.cpp
--- a/Engine/Source/Editor/PackagesDialog/Private/SPackagesDialog.cpp
+++ b/Engine/Source/Editor/PackagesDialog/Private/SPackagesDialog.cpp
@@ -14,6 +14,7 @@
 #include "EditorStyleSet.h"
 #include "IAssetTools.h"
 #include "IAssetTypeActions.h"
+#include "AssetRegistry/IAssetRegistry.h"
 #include "AssetToolsModule.h"
 #include "SWarningOrErrorBox.h"
 
@@ -42,7 +43,7 @@ UObject* FPackageItem::GetPackageObject() const
 		GetObjectsWithPackage(Package, ObjectsInPackage, false);
 		for (UObject* Obj : ObjectsInPackage)
 		{
-			if (Obj->IsAsset())
+			if (Obj->IsAsset() && !IAssetRegistry::Get()->ShouldSkipAsset(Obj))
 			{
 				return Obj;
 			}
@@ -61,7 +62,7 @@ bool FPackageItem::HasMultipleAssets() const
 		UObject* FirstObj = nullptr;
 		for (UObject* Obj : ObjectsInPackage)
 		{
-			if (Obj->IsAsset())
+			if (Obj->IsAsset() && !IAssetRegistry::Get()->ShouldSkipAsset(Obj))
 			{
 				if(FirstObj == nullptr)
 				{
diff --git a/Engine/Source/Runtime/AssetRegistry/Private/AssetRegistry.cpp b/Engine/Source/Runtime/AssetRegistry/Private/AssetRegistry.cpp
--- a/Engine/Source/Runtime/AssetRegistry/Private/AssetRegistry.cpp
+++ b/Engine/Source/Runtime/AssetRegistry/Private/AssetRegistry.cpp
@@ -1014,13 +1014,11 @@ bool UAssetRegistryImpl::EnumerateAssets(const FARCompiledFilter& InFilter, TFun
 					return;
 				}
 
-#if WITH_EDITOR
 				// Skip classes that report themselves as assets but that the editor AssetRegistry is currently not counting as assets
-				if (ShouldSkipAsset(Obj->GetClass()->GetFName(), InMemoryPackage->GetPackageFlags()))
+				if (ShouldSkipAsset(Obj))
 				{
 					return;
 				}
-#endif
 
 				// Package name
 				const FName PackageName = InMemoryPackage->GetFName();
@@ -2490,13 +2488,11 @@ void UAssetRegistryImpl::AssetSearchDataGathered(const double TickStartTime, TRi
 		TUniquePtr<FAssetData> BackgroundResult(AssetResults.PopFrontValue());
 		CA_ASSUME(BackgroundResult.Get() != nullptr);
 
-#if WITH_EDITOR
 		// Skip Assets that we filter out
 		if (ShouldSkipAsset(BackgroundResult->AssetClass, BackgroundResult->PackageFlags))
 		{
 			continue;
 		}
-#endif
 
 		// Try to update any asset data that may already exist
 		FAssetData* AssetData = State.CachedAssetsByObjectPath.FindRef(BackgroundResult->ObjectPath);
@@ -2945,6 +2941,70 @@ void UAssetRegistryImpl::RemovePackageData(const FName PackageName)
 	}
 }
 
+bool UAssetRegistryImpl::ShouldSkipAsset(FName AssetClass, uint32 PackageFlags) const
+{
+#if WITH_ENGINE && WITH_EDITOR
+	// We do not yet support having UBlueprintGeneratedClasses be assets when the UBlueprint is also
+	// an asset; the content browser does not handle the multiple assets correctly and displays this
+	// class asset as if it is in a separate package. Revisit when we have removed the UBlueprint as an asset
+	// or when we support multiple assets.
+	if (!bInitializedSkipClasses)
+	{
+		// Since we only collect these the first on-demand time, it is possible we will miss subclasses
+		// from plugins that load later. This flaw is a rare edge case, though, and this solution will
+		// be replaced eventually, so leaving it for now.
+		if (GIsEditor && (!IsRunningCommandlet() || IsRunningCookCommandlet()))
+		{
+			static const FName NAME_BlueprintGeneratedClass("BlueprintGeneratedClass");
+			static const FName NAME_EnginePackage("/Script/Engine");
+			UClass* BlueprintGeneratedClass = nullptr;
+			UPackage* EnginePackage = Cast<UPackage>(StaticFindObjectFast(UPackage::StaticClass(), nullptr, NAME_EnginePackage));
+			if (EnginePackage)
+			{
+				BlueprintGeneratedClass = Cast<UClass>(StaticFindObjectFast(UClass::StaticClass(), EnginePackage, NAME_BlueprintGeneratedClass));
+			}
+			if (!BlueprintGeneratedClass)
+			{
+				UE_LOG(LogAssetRegistry, Warning, TEXT("Could not find BlueprintGeneratedClass; will not be able to filter BPGC from the Content Browser."));
+			}
+			else
+			{
+				SkipClasses.Add(NAME_BlueprintGeneratedClass);
+				for (TObjectIterator<UClass> It; It; ++It)
+				{
+					if (It->IsChildOf(BlueprintGeneratedClass) && !It->HasAnyClassFlags(CLASS_Abstract))
+					{
+						SkipClasses.Add(It->GetFName());
+					}
+				}
+			}
+		}
+		bInitializedSkipClasses = true;
+	}
+
+	bool bIsCooked = (PackageFlags & PKG_FilterEditorOnly) != 0;
+	if (!bIsCooked && SkipClasses.Contains(AssetClass))
+	{
+		return true;
+	}
+#endif
+	return false;
+}
+
+bool UAssetRegistryImpl::ShouldSkipAsset(const UObject* InAsset) const
+{
+	if (!InAsset)
+	{
+		return false;
+	}
+	UPackage* Package = InAsset->GetPackage();
+	if (!Package)
+	{
+		return false;
+	}
+	return ShouldSkipAsset(InAsset->GetClass()->GetFName(), Package->GetPackageFlags());
+}
+
 #if WITH_EDITOR
 
 void UAssetRegistryImpl::OnDirectoryChanged(const TArray<FFileChangeData>& FileChanges)
@@ -3057,21 +3117,6 @@ void UAssetRegistryImpl::OnDirectoryChanged(const TArray<FFileChangeData>& FileC
 	ScanModifiedAssetFiles(ModifiedFiles);
 }
 
-bool UAssetRegistryImpl::ShouldSkipAsset(FName AssetClass, uint32 PackageFlags) const
-{
-	// We do not yet support having UBlueprintGeneratedClasses be assets when the UBlueprint is also
-	// an asset; the content browser does not handle the multiple assets correctly and displays this
-	// class asset as if it is in a separate package. Revisit when we have removed the UBlueprint as an asset
-	// or when we support multiple assets.
-	static FName NAME_BlueprintGeneratedClass = TEXT("BlueprintGeneratedClass");
-	bool bIsCooked = (PackageFlags & PKG_FilterEditorOnly) != 0;
-	if (!bIsCooked && AssetClass == NAME_BlueprintGeneratedClass)
-	{
-		return true;
-	}
-	return false;
-}
-
 void UAssetRegistryImpl::OnAssetLoaded(UObject *AssetLoaded)
 {
 	UPackage* Package = AssetLoaded->GetPackage();
diff --git a/Engine/Source/Runtime/AssetRegistry/Private/AssetRegistry.h b/Engine/Source/Runtime/AssetRegistry/Private/AssetRegistry.h
--- a/Engine/Source/Runtime/AssetRegistry/Private/AssetRegistry.h
+++ b/Engine/Source/Runtime/AssetRegistry/Private/AssetRegistry.h
@@ -151,6 +151,9 @@ public:
 	UE_DEPRECATED(4.17, "IsUsingWorldAssets is now always true, remove any code that assumes it could be false")
 	static bool IsUsingWorldAssets() { return true; }
 
+	virtual bool ShouldSkipAsset(FName AssetClass, uint32 PackageFlags) const override;
+	virtual bool ShouldSkipAsset(const UObject* InAsset) const override;
+
 protected:
 	virtual void SetManageReferences(const TMultiMap<FAssetIdentifier, FAssetIdentifier>& ManagerMap, bool bClearExisting, UE::AssetRegistry::EDependencyCategory RecurseType, TSet<FDependsNode*>& ExistingManagedNodes, ShouldSetManagerPredicate ShouldSetManager = nullptr) override;
 	virtual bool SetPrimaryAssetIdForObjectPath(const FName ObjectPath, FPrimaryAssetId PrimaryAssetId) override;
@@ -232,9 +235,6 @@ private:
 	/** Called when a file in a content directory changes on disk */
 	void OnDirectoryChanged(const TArray<struct FFileChangeData>& Files);
 
-	/** Called to check whether we should filter out assets of the given class and packageflags from the editor's AssetRegistry */
-	bool ShouldSkipAsset(FName AssetClass, uint32 PackageFlags) const;
-
 	/** Called when an asset is loaded, it will possibly update the cache */
 	void OnAssetLoaded(UObject *AssetLoaded);
 
@@ -451,6 +451,9 @@ private:
 
 	/** The set of object paths that have had their disk cache updated from the in memory version */
 	TSet<FName> AssetDataObjectPathsUpdatedOnLoad;
+
+	mutable TSet<FName> SkipClasses;
+	mutable bool bInitializedSkipClasses = false;
 #endif
 
 };
diff --git a/Engine/Source/Runtime/AssetRegistry/Public/AssetRegistry/IAssetRegistry.h b/Engine/Source/Runtime/AssetRegistry/Public/AssetRegistry/IAssetRegistry.h
--- a/Engine/Source/Runtime/AssetRegistry/Public/AssetRegistry/IAssetRegistry.h
+++ b/Engine/Source/Runtime/AssetRegistry/Public/AssetRegistry/IAssetRegistry.h
@@ -530,6 +530,11 @@ public:
 	/** Load FAssetData from the specified package filename */
 	virtual void LoadPackageRegistryData(const FString& PackageFilename, FLoadPackageRegistryData& InOutData) const = 0;
 	
+	/** Called to check whether we should filter out assets of the given class and packageflags from the editor's AssetRegistry */
+	virtual bool ShouldSkipAsset(FName AssetClass, uint32 PackageFlags) const = 0;
+	/** Called to check whether we should filter out the given object (assumed to be an asset) from the editor's AssetRegistry */
+	virtual bool ShouldSkipAsset(const UObject* InAsset) const = 0;
+
 protected:
 	// Functions specifically for calling from the asset manager
 	friend class UAssetManager;
