Fix ProcessLoadedAssetsToUpdateCache to correctly create the FAssetData for BlueprintGeneratedClass assets. Add more errors when adding/editing assets to identify bad adds at the point of addition rather than catching the NumAssets == 0 ensure in Reset.
#rnx

diff --git a/Engine/Source/Runtime/AssetRegistry/Private/AssetRegistry.cpp b/Engine/Source/Runtime/AssetRegistry/Private/AssetRegistry.cpp
--- a/Engine/Source/Runtime/AssetRegistry/Private/AssetRegistry.cpp
+++ b/Engine/Source/Runtime/AssetRegistry/Private/AssetRegistry.cpp
@@ -1817,7 +1817,7 @@ void UAssetRegistryImpl::PrioritizeSearchPath(const FString& PathToPrioritize)
 
 void UAssetRegistryImpl::AssetCreated(UObject* NewAsset)
 {
-	if (ensure(NewAsset) && NewAsset->IsAsset())
+	if (ensure(NewAsset) && NewAsset->IsAsset() && !ShouldSkipAsset(NewAsset))
 	{
 		// Add the newly created object to the package file cache because its filename can already be
 		// determined by its long package name.
@@ -1836,7 +1836,7 @@ void UAssetRegistryImpl::AssetCreated(UObject* NewAsset)
 		AddAssetPath(*FPackageName::GetLongPackagePath(NewPackageName));
 
 		// Let subscribers know that the new asset was added to the registry
-		AssetAddedEvent.Broadcast(FAssetData(NewAsset));
+		AssetAddedEvent.Broadcast(FAssetData(NewAsset, true /* bAllowBlueprintClass */));
 
 		// Notify listeners that an asset was just created
 		InMemoryAssetCreatedEvent.Broadcast(NewAsset);
@@ -1846,7 +1846,7 @@ void UAssetRegistryImpl::AssetCreated(UObject* NewAsset)
 void UAssetRegistryImpl::AssetDeleted(UObject* DeletedAsset)
 {
 	checkf(GIsEditor, TEXT("Updating the AssetRegistry is only available in editor"));
-	if (ensure(DeletedAsset) && DeletedAsset->IsAsset())
+	if (ensure(DeletedAsset) && DeletedAsset->IsAsset() && !ShouldSkipAsset(DeletedAsset))
 	{
 		UPackage* DeletedObjectPackage = DeletedAsset->GetOutermost();
 		if (DeletedObjectPackage != nullptr)
@@ -1869,7 +1869,7 @@ void UAssetRegistryImpl::AssetDeleted(UObject* DeletedAsset)
 			}
 		}
 
-		FAssetData AssetDataDeleted = FAssetData(DeletedAsset);
+		FAssetData AssetDataDeleted = FAssetData(DeletedAsset, true /* bAllowBlueprintClass */);
 
 #if WITH_EDITOR
 		if (bInitialSearchCompleted && AssetDataDeleted.IsRedirector())
@@ -1890,7 +1890,7 @@ void UAssetRegistryImpl::AssetDeleted(UObject* DeletedAsset)
 void UAssetRegistryImpl::AssetRenamed(const UObject* RenamedAsset, const FString& OldObjectPath)
 {
 	checkf(GIsEditor, TEXT("Updating the AssetRegistry is only available in editor"));
-	if (ensure(RenamedAsset) && RenamedAsset->IsAsset())
+	if (ensure(RenamedAsset) && RenamedAsset->IsAsset() && !ShouldSkipAsset(RenamedAsset))
 	{
 		// Add the renamed object to the package file cache because its filename can already be
 		// determined by its long package name.
@@ -1917,7 +1917,7 @@ void UAssetRegistryImpl::AssetRenamed(const UObject* RenamedAsset, const FString
 		// Add the path to the Path Tree, in case it wasn't already there
 		AddAssetPath(*FPackageName::GetLongPackagePath(NewPackageName));
 
-		AssetRenamedEvent.Broadcast(FAssetData(RenamedAsset), OldObjectPath);
+		AssetRenamedEvent.Broadcast(FAssetData(RenamedAsset, true /* bAllowBlueprintClass */), OldObjectPath);
 	}
 }
 
diff --git a/Engine/Source/Runtime/AssetRegistry/Private/AssetRegistryState.cpp b/Engine/Source/Runtime/AssetRegistry/Private/AssetRegistryState.cpp
--- a/Engine/Source/Runtime/AssetRegistry/Private/AssetRegistryState.cpp
+++ b/Engine/Source/Runtime/AssetRegistry/Private/AssetRegistryState.cpp
@@ -1414,7 +1414,7 @@ void FAssetRegistryState::SetAssetDatas(TArrayView<FAssetData> AssetDatas, const
 	UE_CLOG(NumAssets != 0, LogAssetRegistry, Fatal, TEXT("Can only load into empty asset registry states. Load into temporary and append using InitializeFromExisting() instead."));
 
 	NumAssets = AssetDatas.Num();
-	
+
 	auto SetPathCache = [&]() 
 	{
 		CachedAssetsByObjectPath.Empty(AssetDatas.Num());
@@ -1422,6 +1422,7 @@ void FAssetRegistryState::SetAssetDatas(TArrayView<FAssetData> AssetDatas, const
 		{
 			CachedAssetsByObjectPath.Add(AssetData.ObjectPath, &AssetData);
 		}
+		ensure(NumAssets == CachedAssetsByObjectPath.Num());
 	};
 
 	// FAssetDatas sharing package name are very rare.
@@ -1493,14 +1494,16 @@ void FAssetRegistryState::SetAssetDatas(TArrayView<FAssetData> AssetDatas, const
 
 void FAssetRegistryState::AddAssetData(FAssetData* AssetData)
 {
-	NumAssets++;
-
 	FAssetData*& ExistingByObjectPath = CachedAssetsByObjectPath.FindOrAdd(AssetData->ObjectPath);
 	if (ExistingByObjectPath)
 	{
 		UE_LOG(LogAssetRegistry, Error, TEXT("AddAssetData called with ObjectPath %s which already exists. ")
 			TEXT("This will overwrite and leak the existing AssetData."), *AssetData->ObjectPath.ToString());
 	}
+	else
+	{
+		++NumAssets;
+	}
 	ExistingByObjectPath = AssetData;
 
 	TArray<FAssetData*, TInlineAllocator<1>>& PackageAssets = CachedAssetsByPackageName.FindOrAdd(AssetData->PackageName);
@@ -1549,8 +1552,26 @@ void FAssetRegistryState::UpdateAssetData(FAssetData* AssetData, const FAssetDat
 	// Update ObjectPath
 	if (AssetData->PackageName != NewAssetData.PackageName || AssetData->AssetName != NewAssetData.AssetName)
 	{
-		CachedAssetsByObjectPath.Remove(AssetData->ObjectPath);
-		CachedAssetsByObjectPath.Add(NewAssetData.ObjectPath, AssetData);
+		int32 NumRemoved = CachedAssetsByObjectPath.Remove(AssetData->ObjectPath);
+		check(NumRemoved <= 1);
+		if (NumRemoved == 0)
+		{
+			UE_LOG(LogAssetRegistry, Error, TEXT("UpdateAssetData called on AssetData %s that is not present in the AssetRegistry."),
+				*AssetData->ObjectPath.ToString());
+		}
+		NumAssets -= NumRemoved;
+		FAssetData*& Existing = CachedAssetsByObjectPath.FindOrAdd(NewAssetData.ObjectPath, AssetData);
+		if (Existing)
+		{
+			UE_LOG(LogAssetRegistry, Error, TEXT("UpdateAssetData called with a change in ObjectPath from Old=\"%s\" to New=\"%s\", ")
+				TEXT("but the new ObjectPath is already present with another AssetData. This will overwrite and leak the existing AssetData."),
+				*AssetData->ObjectPath.ToString(), *NewAssetData.ObjectPath.ToString());
+		}
+		else
+		{
+			++NumAssets;
+		}
+		Existing = AssetData;
 	}
 
 	// Update PackageName
@@ -1618,48 +1639,57 @@ void FAssetRegistryState::RemoveAssetData(FAssetData* AssetData, bool bRemoveDep
 	bOutRemovedAssetData = false;
 	bOutRemovedPackageData = false;
 
-	if (ensure(AssetData))
+	if (!ensure(AssetData))
 	{
-		TArray<FAssetData*, TInlineAllocator<1>>* OldPackageAssets = CachedAssetsByPackageName.Find(AssetData->PackageName);
-		TArray<FAssetData*>* OldPathAssets = CachedAssetsByPath.Find(AssetData->PackagePath);
-		TArray<FAssetData*>* OldClassAssets = CachedAssetsByClass.Find(AssetData->AssetClass);
-
-		CachedAssetsByObjectPath.Remove(AssetData->ObjectPath);
-		OldPackageAssets->RemoveSingleSwap(AssetData);
-		OldPathAssets->RemoveSingleSwap(AssetData);
-		OldClassAssets->RemoveSingleSwap(AssetData);
+		return;
+	}
 
-		for (auto TagIt = AssetData->TagsAndValues.CreateConstIterator(); TagIt; ++TagIt)
-		{
-			TArray<FAssetData*>* OldTagAssets = CachedAssetsByTag.Find(TagIt.Key());
-			OldTagAssets->RemoveSingleSwap(AssetData);
-		}
+	int32 NumRemoved = CachedAssetsByObjectPath.Remove(AssetData->ObjectPath);
+	check(NumRemoved <= 1);
+	if (NumRemoved == 0)
+	{
+		UE_LOG(LogAssetRegistry, Error, TEXT("RemoveAssetData called on AssetData %s that is not present in the AssetRegistry."),
+			*AssetData->ObjectPath.ToString());
+		return;
+	}
 
-		// Only remove dependencies and package data if there are no other known assets in the package
-		if (OldPackageAssets->Num() == 0)
-		{
-			CachedAssetsByPackageName.Remove(AssetData->PackageName);
+	TArray<FAssetData*, TInlineAllocator<1>>* OldPackageAssets = CachedAssetsByPackageName.Find(AssetData->PackageName);
+	TArray<FAssetData*>* OldPathAssets = CachedAssetsByPath.Find(AssetData->PackagePath);
+	TArray<FAssetData*>* OldClassAssets = CachedAssetsByClass.Find(AssetData->AssetClass);
+	OldPackageAssets->RemoveSingleSwap(AssetData);
+	OldPathAssets->RemoveSingleSwap(AssetData);
+	OldClassAssets->RemoveSingleSwap(AssetData);
 
-			// We need to update the cached dependencies references cache so that they know we no
-			// longer exist and so don't reference them.
-			if (bRemoveDependencyData)
-			{
-				RemoveDependsNode(AssetData->PackageName);
-			}
+	for (auto TagIt = AssetData->TagsAndValues.CreateConstIterator(); TagIt; ++TagIt)
+	{
+		TArray<FAssetData*>* OldTagAssets = CachedAssetsByTag.Find(TagIt.Key());
+		OldTagAssets->RemoveSingleSwap(AssetData);
+	}
 
-			// Remove the package data as well
-			RemovePackageData(AssetData->PackageName);
-			bOutRemovedPackageData = true;
-		}
+	// Only remove dependencies and package data if there are no other known assets in the package
+	if (OldPackageAssets->Num() == 0)
+	{
+		CachedAssetsByPackageName.Remove(AssetData->PackageName);
 
-		// if the assets were preallocated in a block, we can't delete them one at a time, only the whole chunk in the destructor
-		if (PreallocatedAssetDataBuffers.Num() == 0)
+		// We need to update the cached dependencies references cache so that they know we no
+		// longer exist and so don't reference them.
+		if (bRemoveDependencyData)
 		{
-			delete AssetData;
+			RemoveDependsNode(AssetData->PackageName);
 		}
-		NumAssets--;
-		bOutRemovedAssetData = true;
+
+		// Remove the package data as well
+		RemovePackageData(AssetData->PackageName);
+		bOutRemovedPackageData = true;
+	}
+
+	// if the assets were preallocated in a block, we can't delete them one at a time, only the whole chunk in the destructor
+	if (PreallocatedAssetDataBuffers.Num() == 0)
+	{
+		delete AssetData;
 	}
+	NumAssets--;
+	bOutRemovedAssetData = true;
 }
 
 FDependsNode* FAssetRegistryState::FindDependsNode(const FAssetIdentifier& Identifier) const
