AssetRegistry: Reparse skip classes after all plugins have finished loading, since the first parse is called before then and we miss some classes.
#rnx

diff --git a/Engine/Source/Runtime/AssetRegistry/Private/AssetRegistry.cpp b/Engine/Source/Runtime/AssetRegistry/Private/AssetRegistry.cpp
--- a/Engine/Source/Runtime/AssetRegistry/Private/AssetRegistry.cpp
+++ b/Engine/Source/Runtime/AssetRegistry/Private/AssetRegistry.cpp
@@ -291,6 +291,7 @@ UAssetRegistryImpl::UAssetRegistryImpl(const FObjectInitializer& ObjectInitializ
 		bSearchAllAssetsAtStart = bCommandlineAllAssetsAtStart;
 	}
 #endif
+	IPluginManager& PluginManager = IPluginManager::Get();
 	if (bSearchAllAssetsAtStart)
 	{
 		ConstructGatherer();
@@ -322,7 +323,7 @@ UAssetRegistryImpl::UAssetRegistryImpl(const FObjectInitializer& ObjectInitializ
 		}
 #endif // ASSETREGISTRY_ENABLE_PREMADE_REGISTRY_IN_EDITOR 
 
-		TArray<TSharedRef<IPlugin>> ContentPlugins = IPluginManager::Get().GetEnabledPluginsWithContent();
+		TArray<TSharedRef<IPlugin>> ContentPlugins = PluginManager.GetEnabledPluginsWithContent();
 		for (TSharedRef<IPlugin> ContentPlugin : ContentPlugins)
 		{
 			if (ContentPlugin->CanContainContent())
@@ -420,6 +421,12 @@ UAssetRegistryImpl::UAssetRegistryImpl(const FObjectInitializer& ObjectInitializ
 		UE::AssetRegistry::Private::IAssetRegistrySingleton::Singleton = this;
 		IAssetRegistryInterface::Default = &GAssetRegistryInterface;
 	}
+
+	ELoadingPhase::Type LoadingPhase = PluginManager.GetLastCompletedLoadingPhase();
+	if (LoadingPhase == ELoadingPhase::None || LoadingPhase < ELoadingPhase::PostEngineInit)
+	{
+		PluginManager.OnLoadingPhaseComplete().AddUObject(this, &UAssetRegistryImpl::OnPluginLoadingPhaseComplete);
+	}
 }
 
 UAssetRegistryImpl::UAssetRegistryImpl(FVTableHelper& Helper)
@@ -521,10 +528,6 @@ void UAssetRegistryImpl::RegisterReadOfScriptPackages()
 	{
 		ReadScriptPackages();
 	}
-	else
-	{
-		PluginManager.OnLoadingPhaseComplete().AddUObject(this, &UAssetRegistryImpl::OnPluginLoadingPhaseComplete);
-	}
 }
 
 void UAssetRegistryImpl::OnPluginLoadingPhaseComplete(ELoadingPhase::Type LoadingPhase, bool bPhaseSuccessful)
@@ -533,7 +536,19 @@ void UAssetRegistryImpl::OnPluginLoadingPhaseComplete(ELoadingPhase::Type Loadin
 	{
 		return;
 	}
-	ReadScriptPackages();
+	// If we have constructed the GlobalGatherer then we need to readscriptpackages,
+	// otherwise we will read them when constructing the gatherer.
+	if (GlobalGatherer.IsValid())
+	{
+		ReadScriptPackages();
+	}
+
+	// Reparse the skip classes the next time ShouldSkipAsset is called, since available classes
+	// for the search over all classes may have changed
+#if WITH_EDITOR
+	bInitializedSkipClasses = false;
+#endif
+	IPluginManager::Get().OnLoadingPhaseComplete().RemoveAll(this);
 }
 
 void UAssetRegistryImpl::ReadScriptPackages()
@@ -562,7 +577,6 @@ void UAssetRegistryImpl::ReadScriptPackages()
 			}
 		}
 	}
-	IPluginManager::Get().OnLoadingPhaseComplete().RemoveAll(this);
 }
 
 void UAssetRegistryImpl::InitializeSerializationOptions(FAssetRegistrySerializationOptions& Options, const FString& PlatformIniName) const
@@ -2958,6 +2972,7 @@ bool UAssetRegistryImpl::ShouldSkipAsset(FName AssetClass, uint32 PackageFlags)
 		// be replaced eventually, so leaving it for now.
 		if (GIsEditor && (!IsRunningCommandlet() || IsRunningCookCommandlet()))
 		{
+			SkipClasses.Reset();
 			static const FName NAME_BlueprintGeneratedClass("BlueprintGeneratedClass");
 			static const FName NAME_EnginePackage("/Script/Engine");
 			UClass* BlueprintGeneratedClass = nullptr;
