- Filter out UBlueprint assets from cooked packages
- Asset registry filtering is now in outtake instead of intake, to cover all cases: gathered packages, loaded assets and assets from cooked registry (AssetRegistryState)
- Move PKG_ContainsNoAsset filtering from CB to asset registry

diff --git a/Engine/Plugins/Editor/ContentBrowser/ContentBrowserAssetDataSource/Source/ContentBrowserAssetDataSource/Private/ContentBrowserAssetDataSource.cpp b/Engine/Plugins/Editor/ContentBrowser/ContentBrowserAssetDataSource/Source/ContentBrowserAssetDataSource/Private/ContentBrowserAssetDataSource.cpp
--- a/Engine/Plugins/Editor/ContentBrowser/ContentBrowserAssetDataSource/Source/ContentBrowserAssetDataSource/Private/ContentBrowserAssetDataSource.cpp
+++ b/Engine/Plugins/Editor/ContentBrowser/ContentBrowserAssetDataSource/Source/ContentBrowserAssetDataSource/Private/ContentBrowserAssetDataSource.cpp
@@ -1567,7 +1567,7 @@ void UContentBrowserAssetDataSource::OnAssetUpdated(const FAssetData& InAssetDat
 void UContentBrowserAssetDataSource::OnAssetLoaded(UObject* InAsset)
 {
 	if (InAsset && !InAsset->GetOutermost()->HasAnyPackageFlags(PKG_ForDiffing | PKG_PlayInEditor) &&
-		!AssetRegistry->ShouldSkipAsset(InAsset))
+		!UE::AssetRegistry::FFiltering::ShouldSkipAsset(InAsset))
 	{
 		FAssetData AssetData(InAsset);
 		if (ContentBrowserAssetData::IsPrimaryAsset(AssetData))
diff --git a/Engine/Source/Editor/PackagesDialog/PackagesDialog.Build.cs b/Engine/Source/Editor/PackagesDialog/PackagesDialog.Build.cs
--- a/Engine/Source/Editor/PackagesDialog/PackagesDialog.Build.cs
+++ b/Engine/Source/Editor/PackagesDialog/PackagesDialog.Build.cs
@@ -20,7 +20,8 @@ public class PackagesDialog : ModuleRules
                 "EditorStyle",
 				"UnrealEd",
 				"SourceControl",
-				"EditorWidgets"
+				"EditorWidgets",
+				"AssetRegistry"
 			}
 		);
 
diff --git a/Engine/Source/Editor/PackagesDialog/Private/SPackagesDialog.cpp b/Engine/Source/Editor/PackagesDialog/Private/SPackagesDialog.cpp
--- a/Engine/Source/Editor/PackagesDialog/Private/SPackagesDialog.cpp
+++ b/Engine/Source/Editor/PackagesDialog/Private/SPackagesDialog.cpp
@@ -45,7 +45,7 @@ UObject* FPackageItem::GetPackageObject() const
 		{
 			// Don't filter pending kill objects here as we need to determine if the package contains
 			// a single, pending kill object to properly show it in the save dialog.
-			if (Obj->IsAsset() && !IAssetRegistry::Get()->ShouldSkipAsset(Obj))
+			if (Obj->IsAsset() && !UE::AssetRegistry::FFiltering::ShouldSkipAsset(Obj))
 			{
 				return Obj;
 			}
@@ -66,7 +66,7 @@ bool FPackageItem::HasMultipleAssets() const
 		{
 			// Filter pending kill objects here because we don't want the case where a package contains
 			// an actor and a deleted actor to be reported as multiple assets.
-			if (Obj->IsAsset() && !Obj->IsPendingKill() && !IAssetRegistry::Get()->ShouldSkipAsset(Obj))
+			if (Obj->IsAsset() && !Obj->IsPendingKill() && !UE::AssetRegistry::FFiltering::ShouldSkipAsset(Obj))
 			{
 				if(FirstObj == nullptr)
 				{
diff --git a/Engine/Source/Runtime/AssetRegistry/Private/AssetRegistry.cpp b/Engine/Source/Runtime/AssetRegistry/Private/AssetRegistry.cpp
--- a/Engine/Source/Runtime/AssetRegistry/Private/AssetRegistry.cpp
+++ b/Engine/Source/Runtime/AssetRegistry/Private/AssetRegistry.cpp
@@ -232,6 +232,133 @@ EChunkProgressReportingType::Type GetChunkAvailabilityProgressType(EAssetAvailab
 	return ChunkReportType;
 }
 
+
+namespace UE
+{
+namespace AssetRegistry
+{
+#if WITH_ENGINE && WITH_EDITOR
+	TSet<FName> SkipUncookedClasses;
+	TSet<FName> SkipCookedClasses;
+	bool bInitializedSkipClasses = false;
+#endif //if WITH_ENGINE && WITH_EDITOR
+
+	// TODO: replace with a function in CoreGlobals
+	static bool IsRunningCookCommandlet()
+	{
+		FString Commandline = FCommandLine::Get();
+		const bool bIsCookCommandlet = IsRunningCommandlet() && Commandline.Contains(TEXT("run=cook"));
+		return bIsCookCommandlet;
+	}
+
+	bool FFiltering::ShouldSkipAsset(FName AssetClass, uint32 PackageFlags)
+	{
+#if WITH_ENGINE && WITH_EDITOR
+
+		TRACE_CPUPROFILER_EVENT_SCOPE(AssetRegistry::FFiltering::ShouldSkipAsset)
+
+		// We do not yet support having UBlueprintGeneratedClasses be assets when the UBlueprint is also
+		// an asset; the content browser does not handle the multiple assets correctly and displays this
+		// class asset as if it is in a separate package. Revisit when we have removed the UBlueprint as an asset
+		// or when we support multiple assets.
+		if (!bInitializedSkipClasses)
+		{
+			// Since we only collect these the first on-demand time, it is possible we will miss subclasses
+			// from plugins that load later. This flaw is a rare edge case, though, and this solution will
+			// be replaced eventually, so leaving it for now.
+			if (GIsEditor && (!IsRunningCommandlet() || IsRunningCookCommandlet()))
+			{
+				static const FName NAME_EnginePackage("/Script/Engine");
+				UPackage* EnginePackage = Cast<UPackage>(StaticFindObjectFast(UPackage::StaticClass(), nullptr, NAME_EnginePackage));
+
+				{
+					SkipUncookedClasses.Reset();
+
+					static const FName NAME_BlueprintGeneratedClass("BlueprintGeneratedClass");
+					UClass* BlueprintGeneratedClass = nullptr;
+					if (EnginePackage)
+					{
+						BlueprintGeneratedClass = Cast<UClass>(StaticFindObjectFast(UClass::StaticClass(), EnginePackage, NAME_BlueprintGeneratedClass));
+					}
+					if (!BlueprintGeneratedClass)
+					{
+						UE_LOG(LogAssetRegistry, Warning, TEXT("Could not find BlueprintGeneratedClass; will not be able to filter uncooked BPGC"));
+					}
+					else
+					{
+						SkipUncookedClasses.Add(NAME_BlueprintGeneratedClass);
+						for (TObjectIterator<UClass> It; It; ++It)
+						{
+							if (It->IsChildOf(BlueprintGeneratedClass) && !It->HasAnyClassFlags(CLASS_Abstract))
+							{
+								SkipUncookedClasses.Add(It->GetFName());
+							}
+						}
+					}
+				}
+				{
+					SkipCookedClasses.Reset();
+
+					static const FName NAME_Blueprint("Blueprint");
+					UClass* BlueprintClass = nullptr;
+					if (EnginePackage)
+					{
+						BlueprintClass = Cast<UClass>(StaticFindObjectFast(UClass::StaticClass(), EnginePackage, NAME_Blueprint));
+					}
+					if (!BlueprintClass)
+					{
+						UE_LOG(LogAssetRegistry, Warning, TEXT("Could not find BlueprintClass; will not be able to filter cooked BP"));
+					}
+					else
+					{
+						SkipCookedClasses.Add(NAME_Blueprint);
+						for (TObjectIterator<UClass> It; It; ++It)
+						{
+							if (It->IsChildOf(BlueprintClass) && !It->HasAnyClassFlags(CLASS_Abstract))
+							{
+								SkipCookedClasses.Add(It->GetFName());
+							}
+						}
+					}
+				}
+			}
+
+			bInitializedSkipClasses = true;
+		}
+
+		if (PackageFlags & PKG_ContainsNoAsset)
+		{
+			return true;
+		}
+
+		const bool bIsCooked = (PackageFlags & PKG_FilterEditorOnly);
+
+		if ((bIsCooked && SkipCookedClasses.Contains(AssetClass)) ||
+			(!bIsCooked && SkipUncookedClasses.Contains(AssetClass)))
+		{
+			return true;
+		}
+#endif //if WITH_ENGINE && WITH_EDITOR
+
+		return false;
+	}
+
+	bool FFiltering::ShouldSkipAsset(const UObject* InAsset)
+	{
+		if (!InAsset)
+		{
+			return false;
+		}
+		UPackage* Package = InAsset->GetPackage();
+		if (!Package)
+		{
+			return false;
+		}
+		return ShouldSkipAsset(InAsset->GetClass()->GetFName(), Package->GetPackageFlags());
+	}
+}
+}
+
 UAssetRegistry::UAssetRegistry(const FObjectInitializer& ObjectInitializer)
 	: Super(ObjectInitializer)
 {
@@ -546,7 +673,7 @@ void UAssetRegistryImpl::OnPluginLoadingPhaseComplete(ELoadingPhase::Type Loadin
 	// Reparse the skip classes the next time ShouldSkipAsset is called, since available classes
 	// for the search over all classes may have changed
 #if WITH_EDITOR
-	bInitializedSkipClasses = false;
+	UE::AssetRegistry::bInitializedSkipClasses = false;
 #endif
 	IPluginManager::Get().OnLoadingPhaseComplete().RemoveAll(this);
 }
@@ -906,13 +1033,13 @@ void UAssetRegistryImpl::WaitForPackage(const FString& PackageName)
 
 bool UAssetRegistryImpl::HasAssets(const FName PackagePath, const bool bRecursive) const
 {
-	bool bHasAssets = State.HasAssets(PackagePath);
+	bool bHasAssets = State.HasAssets(PackagePath, true /*bARFiltering*/);
 
 	if (!bHasAssets && bRecursive)
 	{
 		CachedPathTree.EnumerateSubPaths(PackagePath, [this, &bHasAssets](FName SubPath)
 		{
-			bHasAssets = State.HasAssets(SubPath);
+			bHasAssets = State.HasAssets(SubPath, true /*bARFiltering*/);
 			return !bHasAssets;
 		});
 	}
@@ -1033,7 +1160,7 @@ bool UAssetRegistryImpl::EnumerateAssets(const FARCompiledFilter& InFilter, TFun
 				}
 
 				// Skip classes that report themselves as assets but that the editor AssetRegistry is currently not counting as assets
-				if (ShouldSkipAsset(Obj))
+				if (UE::AssetRegistry::FFiltering::ShouldSkipAsset(Obj))
 				{
 					return;
 				}
@@ -1137,7 +1264,7 @@ bool UAssetRegistryImpl::EnumerateAssets(const FARCompiledFilter& InFilter, TFun
 		}
 	}
 
-	State.EnumerateAssets(InFilter, PackagesToSkip, Callback);
+	State.EnumerateAssets(InFilter, PackagesToSkip, Callback, true /*bARFiltering*/);
 
 	return true;
 }
@@ -1148,19 +1275,15 @@ FAssetData UAssetRegistryImpl::GetAssetByObjectPath(const FName ObjectPath, bool
 	{
 		UObject* Asset = FindObject<UObject>(nullptr, *ObjectPath.ToString());
 
-		if (Asset != nullptr)
+		if (Asset)
 		{
-			return FAssetData(Asset);
+			return !UE::AssetRegistry::FFiltering::ShouldSkipAsset(Asset) ? FAssetData(Asset) : FAssetData();
 		}
 	}
 
 	const FAssetData* FoundData = State.GetAssetByObjectPath(ObjectPath);
 
-	if (FoundData)
-	{
-		return *FoundData;
-	}
-	return FAssetData();
+	return (FoundData && !UE::AssetRegistry::FFiltering::ShouldSkipAsset(FoundData->AssetClass, FoundData->PackageFlags)) ? *FoundData : FAssetData();
 }
 
 bool UAssetRegistryImpl::GetAllAssets(TArray<FAssetData>& OutAssetData, bool bIncludeOnlyOnDiskAssets) const
@@ -1190,16 +1313,19 @@ bool UAssetRegistryImpl::EnumerateAllAssets(TFunctionRef<bool(const FAssetData&)
 			if (ObjIt->IsAsset())
 			{
 				const FAssetData AssetData(*ObjIt);
-				if (!Callback(AssetData))
+				if (!UE::AssetRegistry::FFiltering::ShouldSkipAsset(*ObjIt))
 				{
-					return true;
+					if (!Callback(AssetData))
+					{
+						return true;
+					}
 				}
 				PackageNamesToSkip.Add(AssetData.PackageName);
 			}
 		}
 	}
 
-	State.EnumerateAllAssets(PackageNamesToSkip, Callback);
+	State.EnumerateAllAssets(PackageNamesToSkip, Callback, true /*bARFiltering*/);
 
 	return true;
 }
@@ -1831,7 +1957,7 @@ void UAssetRegistryImpl::PrioritizeSearchPath(const FString& PathToPrioritize)
 
 void UAssetRegistryImpl::AssetCreated(UObject* NewAsset)
 {
-	if (ensure(NewAsset) && NewAsset->IsAsset() && !ShouldSkipAsset(NewAsset))
+	if (ensure(NewAsset) && NewAsset->IsAsset())
 	{
 		// Add the newly created object to the package file cache because its filename can already be
 		// determined by its long package name.
@@ -1849,18 +1975,21 @@ void UAssetRegistryImpl::AssetCreated(UObject* NewAsset)
 		// Add the path to the Path Tree, in case it wasn't already there
 		AddAssetPath(*FPackageName::GetLongPackagePath(NewPackageName));
 
-		// Let subscribers know that the new asset was added to the registry
-		AssetAddedEvent.Broadcast(FAssetData(NewAsset, true /* bAllowBlueprintClass */));
+		if (!UE::AssetRegistry::FFiltering::ShouldSkipAsset(NewAsset))
+		{
+			// Let subscribers know that the new asset was added to the registry
+			AssetAddedEvent.Broadcast(FAssetData(NewAsset, true /* bAllowBlueprintClass */));
 
-		// Notify listeners that an asset was just created
-		InMemoryAssetCreatedEvent.Broadcast(NewAsset);
+			// Notify listeners that an asset was just created
+			InMemoryAssetCreatedEvent.Broadcast(NewAsset);
+		}
 	}
 }
 
 void UAssetRegistryImpl::AssetDeleted(UObject* DeletedAsset)
 {
 	checkf(GIsEditor, TEXT("Updating the AssetRegistry is only available in editor"));
-	if (ensure(DeletedAsset) && DeletedAsset->IsAsset() && !ShouldSkipAsset(DeletedAsset))
+	if (ensure(DeletedAsset) && DeletedAsset->IsAsset())
 	{
 		UPackage* DeletedObjectPackage = DeletedAsset->GetOutermost();
 		if (DeletedObjectPackage != nullptr)
@@ -1893,18 +2022,21 @@ void UAssetRegistryImpl::AssetDeleted(UObject* DeletedAsset)
 		}
 #endif
 
-		// Let subscribers know that the asset was removed from the registry
-		AssetRemovedEvent.Broadcast(AssetDataDeleted);
+		if (!UE::AssetRegistry::FFiltering::ShouldSkipAsset(DeletedAsset))
+		{
+			// Let subscribers know that the asset was removed from the registry
+			AssetRemovedEvent.Broadcast(AssetDataDeleted);
 
-		// Notify listeners that an in-memory asset was just deleted
-		InMemoryAssetDeletedEvent.Broadcast(DeletedAsset);
+			// Notify listeners that an in-memory asset was just deleted
+			InMemoryAssetDeletedEvent.Broadcast(DeletedAsset);
+		}
 	}
 }
 
 void UAssetRegistryImpl::AssetRenamed(const UObject* RenamedAsset, const FString& OldObjectPath)
 {
 	checkf(GIsEditor, TEXT("Updating the AssetRegistry is only available in editor"));
-	if (ensure(RenamedAsset) && RenamedAsset->IsAsset() && !ShouldSkipAsset(RenamedAsset))
+	if (ensure(RenamedAsset) && RenamedAsset->IsAsset())
 	{
 		// Add the renamed object to the package file cache because its filename can already be
 		// determined by its long package name.
@@ -1931,7 +2063,10 @@ void UAssetRegistryImpl::AssetRenamed(const UObject* RenamedAsset, const FString
 		// Add the path to the Path Tree, in case it wasn't already there
 		AddAssetPath(*FPackageName::GetLongPackagePath(NewPackageName));
 
-		AssetRenamedEvent.Broadcast(FAssetData(RenamedAsset, true /* bAllowBlueprintClass */), OldObjectPath);
+		if (!UE::AssetRegistry::FFiltering::ShouldSkipAsset(RenamedAsset))
+		{
+			AssetRenamedEvent.Broadcast(FAssetData(RenamedAsset, true /* bAllowBlueprintClass */), OldObjectPath);
+		}
 	}
 }
 
@@ -2505,12 +2640,6 @@ void UAssetRegistryImpl::AssetSearchDataGathered(const double TickStartTime, TRi
 		TUniquePtr<FAssetData> BackgroundResult(AssetResults.PopFrontValue());
 		CA_ASSUME(BackgroundResult.Get() != nullptr);
 
-		// Skip Assets that we filter out
-		if (ShouldSkipAsset(BackgroundResult->AssetClass, BackgroundResult->PackageFlags))
-		{
-			continue;
-		}
-
 		// Try to update any asset data that may already exist
 		FAssetData* AssetData = State.CachedAssetsByObjectPath.FindRef(BackgroundResult->ObjectPath);
 
@@ -2958,71 +3087,6 @@ void UAssetRegistryImpl::RemovePackageData(const FName PackageName)
 	}
 }
 
-bool UAssetRegistryImpl::ShouldSkipAsset(FName AssetClass, uint32 PackageFlags) const
-{
-#if WITH_ENGINE && WITH_EDITOR
-	// We do not yet support having UBlueprintGeneratedClasses be assets when the UBlueprint is also
-	// an asset; the content browser does not handle the multiple assets correctly and displays this
-	// class asset as if it is in a separate package. Revisit when we have removed the UBlueprint as an asset
-	// or when we support multiple assets.
-	if (!bInitializedSkipClasses)
-	{
-		// Since we only collect these the first on-demand time, it is possible we will miss subclasses
-		// from plugins that load later. This flaw is a rare edge case, though, and this solution will
-		// be replaced eventually, so leaving it for now.
-		if (GIsEditor && (!IsRunningCommandlet() || IsRunningCookCommandlet()))
-		{
-			SkipClasses.Reset();
-			static const FName NAME_BlueprintGeneratedClass("BlueprintGeneratedClass");
-			static const FName NAME_EnginePackage("/Script/Engine");
-			UClass* BlueprintGeneratedClass = nullptr;
-			UPackage* EnginePackage = Cast<UPackage>(StaticFindObjectFast(UPackage::StaticClass(), nullptr, NAME_EnginePackage));
-			if (EnginePackage)
-			{
-				BlueprintGeneratedClass = Cast<UClass>(StaticFindObjectFast(UClass::StaticClass(), EnginePackage, NAME_BlueprintGeneratedClass));
-			}
-			if (!BlueprintGeneratedClass)
-			{
-				UE_LOG(LogAssetRegistry, Warning, TEXT("Could not find BlueprintGeneratedClass; will not be able to filter BPGC from the Content Browser."));
-			}
-			else
-			{
-				SkipClasses.Add(NAME_BlueprintGeneratedClass);
-				for (TObjectIterator<UClass> It; It; ++It)
-				{
-					if (It->IsChildOf(BlueprintGeneratedClass) && !It->HasAnyClassFlags(CLASS_Abstract))
-					{
-						SkipClasses.Add(It->GetFName());
-					}
-				}
-			}
-		}
-		bInitializedSkipClasses = true;
-	}
-
-	bool bIsCooked = (PackageFlags & PKG_FilterEditorOnly) != 0;
-	if (!bIsCooked && SkipClasses.Contains(AssetClass))
-	{
-		return true;
-	}
-#endif
-	return false;
-}
-
-bool UAssetRegistryImpl::ShouldSkipAsset(const UObject* InAsset) const
-{
-	if (!InAsset)
-	{
-		return false;
-	}
-	UPackage* Package = InAsset->GetPackage();
-	if (!Package)
-	{
-		return false;
-	}
-	return ShouldSkipAsset(InAsset->GetClass()->GetFName(), Package->GetPackageFlags());
-}
-
 #if WITH_EDITOR
 
 void UAssetRegistryImpl::OnDirectoryChanged(const TArray<FFileChangeData>& FileChanges)
@@ -3137,12 +3201,6 @@ void UAssetRegistryImpl::OnDirectoryChanged(const TArray<FFileChangeData>& FileC
 
 void UAssetRegistryImpl::OnAssetLoaded(UObject *AssetLoaded)
 {
-	UPackage* Package = AssetLoaded->GetPackage();
-	if (Package && ShouldSkipAsset(AssetLoaded->GetClass()->GetFName(), Package->GetPackageFlags()))
-	{
-		return;
-	}
-
 	LoadedAssetsToProcess.Add(AssetLoaded);
 }
 
diff --git a/Engine/Source/Runtime/AssetRegistry/Private/AssetRegistry.h b/Engine/Source/Runtime/AssetRegistry/Private/AssetRegistry.h
--- a/Engine/Source/Runtime/AssetRegistry/Private/AssetRegistry.h
+++ b/Engine/Source/Runtime/AssetRegistry/Private/AssetRegistry.h
@@ -151,9 +151,6 @@ public:
 	UE_DEPRECATED(4.17, "IsUsingWorldAssets is now always true, remove any code that assumes it could be false")
 	static bool IsUsingWorldAssets() { return true; }
 
-	virtual bool ShouldSkipAsset(FName AssetClass, uint32 PackageFlags) const override;
-	virtual bool ShouldSkipAsset(const UObject* InAsset) const override;
-
 protected:
 	virtual void SetManageReferences(const TMultiMap<FAssetIdentifier, FAssetIdentifier>& ManagerMap, bool bClearExisting, UE::AssetRegistry::EDependencyCategory RecurseType, TSet<FDependsNode*>& ExistingManagedNodes, ShouldSetManagerPredicate ShouldSetManager = nullptr) override;
 	virtual bool SetPrimaryAssetIdForObjectPath(const FName ObjectPath, FPrimaryAssetId PrimaryAssetId) override;
@@ -451,9 +448,6 @@ private:
 
 	/** The set of object paths that have had their disk cache updated from the in memory version */
 	TSet<FName> AssetDataObjectPathsUpdatedOnLoad;
-
-	mutable TSet<FName> SkipClasses;
-	mutable bool bInitializedSkipClasses = false;
 #endif
 
 };
diff --git a/Engine/Source/Runtime/AssetRegistry/Private/AssetRegistryState.cpp b/Engine/Source/Runtime/AssetRegistry/Private/AssetRegistryState.cpp
--- a/Engine/Source/Runtime/AssetRegistry/Private/AssetRegistryState.cpp
+++ b/Engine/Source/Runtime/AssetRegistry/Private/AssetRegistryState.cpp
@@ -556,19 +556,34 @@ void FAssetRegistryState::PruneAssetData(const TSet<FName>& RequiredPackages, co
 	}
 }
 
-bool FAssetRegistryState::HasAssets(const FName PackagePath) const
+bool FAssetRegistryState::HasAssets(const FName PackagePath, bool bARFiltering) const
 {
 	const TArray<FAssetData*>* FoundAssetArray = CachedAssetsByPath.Find(PackagePath);
-	return FoundAssetArray && FoundAssetArray->Num() > 0;
+	if (FoundAssetArray)
+	{
+		if (bARFiltering)
+		{
+			return FoundAssetArray->ContainsByPredicate([](FAssetData* AssetData)
+			{
+				return AssetData && !UE::AssetRegistry::FFiltering::ShouldSkipAsset(AssetData->AssetClass, AssetData->PackageFlags);
+			});
+		}
+		else
+		{
+			return FoundAssetArray->Num() > 0;
+		}
+	}
+	return false;
 }
 
-bool FAssetRegistryState::GetAssets(const FARCompiledFilter& Filter, const TSet<FName>& PackageNamesToSkip, TArray<FAssetData>& OutAssetData) const
+bool FAssetRegistryState::GetAssets(const FARCompiledFilter& Filter, const TSet<FName>& PackageNamesToSkip, TArray<FAssetData>& OutAssetData, bool bARFiltering) const
 {
 	return EnumerateAssets(Filter, PackageNamesToSkip, [&OutAssetData](const FAssetData& AssetData)
 	{
 		OutAssetData.Emplace(AssetData);
 		return true;
-	});
+	},
+	bARFiltering);
 }
 
 template<class ArrayType>
@@ -597,7 +612,7 @@ TArray<FAssetData*> FindAssets(const TMap<FName, ArrayType>& Map, const TSet<FNa
 	return Out;
 }
 
-bool FAssetRegistryState::EnumerateAssets(const FARCompiledFilter& Filter, const TSet<FName>& PackageNamesToSkip, TFunctionRef<bool(const FAssetData&)> Callback) const
+bool FAssetRegistryState::EnumerateAssets(const FARCompiledFilter& Filter, const TSet<FName>& PackageNamesToSkip, TFunctionRef<bool(const FAssetData&)> Callback, bool bARFiltering) const
 {
 	// Verify filter input. If all assets are needed, use EnumerateAllAssets() instead.
 	if (Filter.IsEmpty() || !IsFilterValid(Filter))
@@ -686,9 +701,14 @@ bool FAssetRegistryState::EnumerateAssets(const FARCompiledFilter& Filter, const
 	{
 		auto SkipAssetData = [&](const FAssetData* AssetData) 
 		{ 
-			return	PackageNamesToSkip.Contains(AssetData->PackageName) |			//-V792
-					AssetData->HasAnyPackageFlags(FilterWithoutPackageFlags) |		//-V792
-					!AssetData->HasAllPackageFlags(FilterWithPackageFlags);			//-V792
+			if (PackageNamesToSkip.Contains(AssetData->PackageName) |			//-V792
+				AssetData->HasAnyPackageFlags(FilterWithoutPackageFlags) |		//-V792
+				!AssetData->HasAllPackageFlags(FilterWithPackageFlags))			//-V792
+			{
+				return true;
+			}
+
+			return bARFiltering && UE::AssetRegistry::FFiltering::ShouldSkipAsset(AssetData->AssetClass, AssetData->PackageFlags);
 		};
 
 		if (FilterResults.Num() > 1)
@@ -739,31 +759,30 @@ bool FAssetRegistryState::EnumerateAssets(const FARCompiledFilter& Filter, const
 	return true;
 }
 
-bool FAssetRegistryState::GetAllAssets(const TSet<FName>& PackageNamesToSkip, TArray<FAssetData>& OutAssetData) const
+bool FAssetRegistryState::GetAllAssets(const TSet<FName>& PackageNamesToSkip, TArray<FAssetData>& OutAssetData, bool bARFiltering) const
 {
 	return EnumerateAllAssets(PackageNamesToSkip, [&OutAssetData](const FAssetData& AssetData)
 	{
 		OutAssetData.Emplace(AssetData);
 		return true;
-	});
+	},
+	bARFiltering);
 }
 
-bool FAssetRegistryState::EnumerateAllAssets(const TSet<FName>& PackageNamesToSkip, TFunctionRef<bool(const FAssetData&)> Callback) const
+bool FAssetRegistryState::EnumerateAllAssets(const TSet<FName>& PackageNamesToSkip, TFunctionRef<bool(const FAssetData&)> Callback, bool bARFiltering) const
 {
 	// All unloaded disk assets
 	for (const TPair<FName, FAssetData*>& AssetDataPair : CachedAssetsByObjectPath)
 	{
 		const FAssetData* AssetData = AssetDataPair.Value;
 
-		if (AssetData != nullptr)
+		if (AssetData &&
+			!PackageNamesToSkip.Contains(AssetData->PackageName) &&
+			(!bARFiltering || !UE::AssetRegistry::FFiltering::ShouldSkipAsset(AssetData->AssetClass, AssetData->PackageFlags)))
 		{
-			// Make sure the asset's package was not loaded then the object was deleted/renamed
-			if (!PackageNamesToSkip.Contains(AssetData->PackageName))
+			if (!Callback(*AssetData))
 			{
-				if (!Callback(*AssetData))
-				{
-					return true;
-				}
+				return true;
 			}
 		}
 	}
diff --git a/Engine/Source/Runtime/AssetRegistry/Public/AssetRegistry/AssetRegistryState.h b/Engine/Source/Runtime/AssetRegistry/Public/AssetRegistry/AssetRegistryState.h
--- a/Engine/Source/Runtime/AssetRegistry/Public/AssetRegistry/AssetRegistryState.h
+++ b/Engine/Source/Runtime/AssetRegistry/Public/AssetRegistry/AssetRegistryState.h
@@ -106,9 +106,10 @@ public:
 
 	/**
 	 * Does the given path contain assets?
+	 * @param bARFiltering Whether to apply filtering from UE::AssetRegistry::FFiltering (false by default)
 	 * @note This function doesn't recurse into sub-paths.
 	 */
-	bool HasAssets(const FName PackagePath) const;
+	bool HasAssets(const FName PackagePath, bool bARFiltering=false) const;
 
 	/**
 	 * Gets asset data for all assets that match the filter.
@@ -118,8 +119,9 @@ public:
 	 * @param Filter filter to apply to the assets in the AssetRegistry
 	 * @param PackageNamesToSkip explicit list of packages to skip, because they were already added
 	 * @param OutAssetData the list of assets in this path
+	 * @param bARFiltering Whether to apply filtering from UE::AssetRegistry::FFiltering (false by default)
 	 */
-	bool GetAssets(const FARCompiledFilter& Filter, const TSet<FName>& PackageNamesToSkip, TArray<FAssetData>& OutAssetData) const;
+	bool GetAssets(const FARCompiledFilter& Filter, const TSet<FName>& PackageNamesToSkip, TArray<FAssetData>& OutAssetData, bool bARFiltering = false) const;
 
 	/**
 	 * Enumerate asset data for all assets that match the filter.
@@ -129,24 +131,27 @@ public:
 	 * @param Filter filter to apply to the assets in the AssetRegistry
 	 * @param PackageNamesToSkip explicit list of packages to skip, because they were already added
 	 * @param Callback function to call for each asset data enumerated
+	 * @param bARFiltering Whether to apply filtering from UE::AssetRegistry::FFiltering (false by default)
 	 */
-	bool EnumerateAssets(const FARCompiledFilter& Filter, const TSet<FName>& PackageNamesToSkip, TFunctionRef<bool(const FAssetData&)> Callback) const;
+	bool EnumerateAssets(const FARCompiledFilter& Filter, const TSet<FName>& PackageNamesToSkip, TFunctionRef<bool(const FAssetData&)> Callback, bool bARFiltering = false) const;
 
 	/**
 	 * Gets asset data for all assets in the registry state.
 	 *
 	 * @param PackageNamesToSkip explicit list of packages to skip, because they were already added
 	 * @param OutAssetData the list of assets
+	 * @param bARFiltering Whether to apply filtering from UE::AssetRegistry::FFiltering (false by default)
 	 */
-	bool GetAllAssets(const TSet<FName>& PackageNamesToSkip, TArray<FAssetData>& OutAssetData) const;
+	bool GetAllAssets(const TSet<FName>& PackageNamesToSkip, TArray<FAssetData>& OutAssetData, bool bARFiltering = false) const;
 
 	/**
 	 * Enumerates asset data for all assets in the registry state.
 	 *
 	 * @param PackageNamesToSkip explicit list of packages to skip, because they were already added
 	 * @param Callback function to call for each asset data enumerated
+	 * @param bARFiltering Whether to apply filtering from UE::AssetRegistry::FFiltering (false by default)
 	 */
-	bool EnumerateAllAssets(const TSet<FName>& PackageNamesToSkip, TFunctionRef<bool(const FAssetData&)> Callback) const;
+	bool EnumerateAllAssets(const TSet<FName>& PackageNamesToSkip, TFunctionRef<bool(const FAssetData&)> Callback, bool bARFiltering = false) const;
 
 	UE_DEPRECATED(4.26, "Use GetDependencies that takes a UE::AssetRegistry::EDependencyCategory instead")
 	bool GetDependencies(const FAssetIdentifier& AssetIdentifier, TArray<FAssetIdentifier>& OutDependencies, EAssetRegistryDependencyType::Type InDependencyType) const;
diff --git a/Engine/Source/Runtime/AssetRegistry/Public/AssetRegistry/IAssetRegistry.h b/Engine/Source/Runtime/AssetRegistry/Public/AssetRegistry/IAssetRegistry.h
--- a/Engine/Source/Runtime/AssetRegistry/Public/AssetRegistry/IAssetRegistry.h
+++ b/Engine/Source/Runtime/AssetRegistry/Public/AssetRegistry/IAssetRegistry.h
@@ -529,11 +529,6 @@ public:
 	
 	/** Load FAssetData from the specified package filename */
 	virtual void LoadPackageRegistryData(const FString& PackageFilename, FLoadPackageRegistryData& InOutData) const = 0;
-	
-	/** Called to check whether we should filter out assets of the given class and packageflags from the editor's AssetRegistry */
-	virtual bool ShouldSkipAsset(FName AssetClass, uint32 PackageFlags) const = 0;
-	/** Called to check whether we should filter out the given object (assumed to be an asset) from the editor's AssetRegistry */
-	virtual bool ShouldSkipAsset(const UObject* InAsset) const = 0;
 
 protected:
 	// Functions specifically for calling from the asset manager
@@ -584,5 +579,15 @@ namespace AssetRegistry
 	// WritePackageData is declared in AssetRegistryInterface.h, in the CoreUObject module, because it is needed by SavePackage in CoreUObject
 	ASSETREGISTRY_API bool ReadPackageDataMain(FArchive& BinaryArchive, const FString& PackageName, const FPackageFileSummary& PackageFileSummary, int64& OutDependencyDataOffset, TArray<FAssetData*>& OutAssetDataList, EReadPackageDataMainErrorCode& OutError);
 	ASSETREGISTRY_API bool ReadPackageDataDependencies(FArchive& BinaryArchive, TBitArray<>& OutImportUsedInGame, TBitArray<>& OutSoftPackageUsedInGame);
+
+	class ASSETREGISTRY_API FFiltering
+	{
+	public:
+		/** Called to check whether we should filter out assets of the given class and package flags from the editor's asset registry */
+		static bool ShouldSkipAsset(FName AssetClass, uint32 PackageFlags);
+
+		/** Called to check whether we should filter out the given object (assumed to be an asset) from the editor's asset registry */
+		static bool ShouldSkipAsset(const UObject* InAsset);
+	};
 }
 }
