Fix missing asset registry filtering checks

diff --git a/Engine/Source/Runtime/AssetRegistry/Private/AssetRegistry.cpp b/Engine/Source/Runtime/AssetRegistry/Private/AssetRegistry.cpp
--- a/Engine/Source/Runtime/AssetRegistry/Private/AssetRegistry.cpp
+++ b/Engine/Source/Runtime/AssetRegistry/Private/AssetRegistry.cpp
@@ -917,9 +917,12 @@ UAssetRegistryImpl::~UAssetRegistryImpl()
 #endif // WITH_EDITOR
 
 	// Clear all listeners
+	PathAddedEvent.Clear();
+	PathRemovedEvent.Clear();
 	AssetAddedEvent.Clear();
 	AssetRemovedEvent.Clear();
 	AssetRenamedEvent.Clear();
+	AssetUpdatedEvent.Clear();
 	InMemoryAssetCreatedEvent.Clear();
 	InMemoryAssetDeletedEvent.Clear();
 	FileLoadedEvent.Clear();
@@ -2280,7 +2283,7 @@ void UAssetRegistryImpl::AppendState(const FAssetRegistryState& InState)
 		// Let subscribers know that the new asset was added to the registry
 		AssetAddedEvent.Broadcast(AssetData);
 		return true;
-	});
+	}, true /*bARFiltering*/);
 }
 
 void UAssetRegistryImpl::CachePathsFromState(const FAssetRegistryState& InState)
@@ -2962,8 +2965,11 @@ void UAssetRegistryImpl::AddAssetData(FAssetData* AssetData)
 {
 	State.AddAssetData(AssetData);
 
-	// Notify subscribers
-	AssetAddedEvent.Broadcast(*AssetData);
+	if (!UE::AssetRegistry::FFiltering::ShouldSkipAsset(AssetData->AssetClass, AssetData->PackageFlags))
+	{
+		// Notify subscribers
+		AssetAddedEvent.Broadcast(*AssetData);
+	}
 
 	// Populate the class map if adding blueprint
 	if (ClassGeneratorNames.Contains(AssetData->AssetClass))
@@ -3012,7 +3018,10 @@ void UAssetRegistryImpl::UpdateAssetData(FAssetData* AssetData, const FAssetData
 
 	State.UpdateAssetData(AssetData, NewAssetData);
 	
-	AssetUpdatedEvent.Broadcast(*AssetData);
+	if (!UE::AssetRegistry::FFiltering::ShouldSkipAsset(AssetData->AssetClass, AssetData->PackageFlags))
+	{
+		AssetUpdatedEvent.Broadcast(*AssetData);
+	}
 }
 
 bool UAssetRegistryImpl::RemoveAssetData(FAssetData* AssetData)
@@ -3021,8 +3030,11 @@ bool UAssetRegistryImpl::RemoveAssetData(FAssetData* AssetData)
 
 	if (ensure(AssetData))
 	{
-		// Notify subscribers
-		AssetRemovedEvent.Broadcast(*AssetData);
+		if (!UE::AssetRegistry::FFiltering::ShouldSkipAsset(AssetData->AssetClass, AssetData->PackageFlags))
+		{
+			// Notify subscribers
+			AssetRemovedEvent.Broadcast(*AssetData);
+		}
 
 		// Remove from the class map if removing a blueprint
 		if (ClassGeneratorNames.Contains(AssetData->AssetClass))
