Add some more instances of the workaround for MSVC static analyzer false positive C6385. Reading with *pointer syntax instead of [] array syntax bypasses the faulty SA rule.
#rnx

diff --git a/Engine/Source/Runtime/AssetRegistry/Private/AssetDataGatherer.cpp b/Engine/Source/Runtime/AssetRegistry/Private/AssetDataGatherer.cpp
--- a/Engine/Source/Runtime/AssetRegistry/Private/AssetDataGatherer.cpp
+++ b/Engine/Source/Runtime/AssetRegistry/Private/AssetDataGatherer.cpp
@@ -3131,9 +3131,9 @@ void FAssetDataGatherer::SerializeCacheLoad(FAssetRegistryReader& Ar)
 		for (int32 AssetIndex = 0; AssetIndex < LocalNumAssets; ++AssetIndex)
 		{
 			// Load the name first to add the entry to the tmap below
-			Ar << PackageNameBlock[AssetIndex];
-			// Visual Studio Static Analyzer issues C6385 if we call AssetDataBlock[AssetIndex].SerializeForCache
-			(AssetDataBlock + AssetIndex)->SerializeForCache(Ar);
+			// Visual Studio Static Analyzer issues C6385 if we call Ar << PackageNameBlock[AssetIndex] or AssetDataBlock[AssetIndex].SerializeForCache
+			Ar << *(PackageNameBlock + AssetIndex); // -C6385
+			(AssetDataBlock + AssetIndex)->SerializeForCache(Ar); // -C6385
 			if (Ar.IsError())
 			{
 				// There was an error reading the cache. Bail out.
@@ -3151,7 +3151,7 @@ void FAssetDataGatherer::SerializeCacheLoad(FAssetRegistryReader& Ar)
 			DiskCachedAssetDataMap.Reserve(DiskCachedAssetDataMap.Num() + LocalNumAssets);
 			for (int32 AssetIndex = 0; AssetIndex < LocalNumAssets; ++AssetIndex)
 			{
-				DiskCachedAssetDataMap.Add(PackageNameBlock[AssetIndex], &AssetDataBlock[AssetIndex]);
+				DiskCachedAssetDataMap.Add(*(PackageNameBlock + AssetIndex), (AssetDataBlock + AssetIndex)); // -C6385
 			}
 			DiskCachedAssetBlocks.Emplace(LocalNumAssets, AssetDataBlock);
 		}
