IAssetRegistry::GetAssetByObjectPath when called on a BlueprintGeneratedClass path that exists in memory should return the FAssetData for the UBlueprint that generated the BPGC to match pre-BPGC-as-asset expectations.

diff --git a/Engine/Source/Runtime/AssetRegistry/Private/AssetRegistry.cpp b/Engine/Source/Runtime/AssetRegistry/Private/AssetRegistry.cpp
--- a/Engine/Source/Runtime/AssetRegistry/Private/AssetRegistry.cpp
+++ b/Engine/Source/Runtime/AssetRegistry/Private/AssetRegistry.cpp
@@ -1280,7 +1280,7 @@ FAssetData UAssetRegistryImpl::GetAssetByObjectPath(const FName ObjectPath, bool
 
 		if (Asset)
 		{
-			return !UE::AssetRegistry::FFiltering::ShouldSkipAsset(Asset) ? FAssetData(Asset) : FAssetData();
+			return FAssetData(Asset);
 		}
 	}
 
