Add IAssetRegistry::GetPackagesByName to find packages by short package name; this replaces FPackageName::SearchForPackageOnDisk, which was slower and does not work with iostore.
#rnx

diff --git a/Engine/Source/Runtime/AssetRegistry/Private/AssetRegistry.cpp b/Engine/Source/Runtime/AssetRegistry/Private/AssetRegistry.cpp
--- a/Engine/Source/Runtime/AssetRegistry/Private/AssetRegistry.cpp
+++ b/Engine/Source/Runtime/AssetRegistry/Private/AssetRegistry.cpp
@@ -1206,6 +1206,16 @@ bool UAssetRegistryImpl::EnumerateAllAssets(TFunctionRef<bool(const FAssetData&)
 	return true;
 }
 
+void UAssetRegistryImpl::GetPackagesByName(FStringView PackageName, TArray<FName>& OutPackageNames) const
+{
+	State.GetPackagesByName(PackageName, OutPackageNames);
+}
+
+FName UAssetRegistryImpl::GetFirstPackageByName(FStringView PackageName) const
+{
+	return State.GetFirstPackageByName(PackageName);
+}
+
 bool UAssetRegistryImpl::GetDependencies(const FAssetIdentifier& AssetIdentifier, TArray<FAssetIdentifier>& OutDependencies, EAssetRegistryDependencyType::Type InDependencyType) const
 {
 	PRAGMA_DISABLE_DEPRECATION_WARNINGS
diff --git a/Engine/Source/Runtime/AssetRegistry/Private/AssetRegistry.h b/Engine/Source/Runtime/AssetRegistry/Private/AssetRegistry.h
--- a/Engine/Source/Runtime/AssetRegistry/Private/AssetRegistry.h
+++ b/Engine/Source/Runtime/AssetRegistry/Private/AssetRegistry.h
@@ -51,6 +51,8 @@ public:
 	virtual FAssetData GetAssetByObjectPath( const FName ObjectPath, bool bIncludeOnlyOnDiskAssets = false ) const override;
 	virtual bool GetAllAssets(TArray<FAssetData>& OutAssetData, bool bIncludeOnlyOnDiskAssets = false) const override;
 	virtual bool EnumerateAllAssets(TFunctionRef<bool(const FAssetData&)> Callback, bool bIncludeOnlyOnDiskAssets = false) const override;
+	virtual void GetPackagesByName(FStringView PackageName, TArray<FName>& OutPackageNames) const override;
+	virtual FName GetFirstPackageByName(FStringView PackageName) const override;
 	UE_DEPRECATED(4.26, "Use GetDependencies that takes a UE::AssetRegistry::EDependencyCategory instead")
 	virtual bool GetDependencies(const FAssetIdentifier& AssetIdentifier, TArray<FAssetIdentifier>& OutDependencies, EAssetRegistryDependencyType::Type InDependencyType) const override;
 	virtual bool GetDependencies(const FAssetIdentifier& AssetIdentifier, TArray<FAssetIdentifier>& OutDependencies, UE::AssetRegistry::EDependencyCategory Category = UE::AssetRegistry::EDependencyCategory::All, const UE::AssetRegistry::FDependencyQuery& Flags = UE::AssetRegistry::FDependencyQuery()) const override;
diff --git a/Engine/Source/Runtime/AssetRegistry/Private/AssetRegistryState.cpp b/Engine/Source/Runtime/AssetRegistry/Private/AssetRegistryState.cpp
--- a/Engine/Source/Runtime/AssetRegistry/Private/AssetRegistryState.cpp
+++ b/Engine/Source/Runtime/AssetRegistry/Private/AssetRegistryState.cpp
@@ -789,6 +789,51 @@ bool FAssetRegistryState::EnumerateAllAssets(const TSet<FName>& PackageNamesToSk
 	return true;
 }
 
+void FAssetRegistryState::GetPackagesByName(FStringView PackageName, TArray<FName>& OutPackageNames) const
+{
+	// Note that we use CachedAssetsByPackageName rather than CachedPackageData because CachedPackageData
+	// is often stripped out of the runtime AssetRegistry
+	if (!FPackageName::IsShortPackageName(PackageName))
+	{
+		FName PackageFName(PackageName);
+		if (CachedAssetsByPackageName.Contains(PackageFName))
+		{
+			OutPackageNames.Add(PackageFName);
+		}
+	}
+	else
+	{
+		TStringBuilder<128> PackageNameStr;
+		for (const TPair<FName, TArray<FAssetData*>>& It : CachedAssetsByPackageName)
+		{
+			It.Key.ToString(PackageNameStr);
+			FStringView ExistingBaseName = FPathViews::GetBaseFilename(PackageNameStr);
+			if (ExistingBaseName.Equals(PackageName, ESearchCase::IgnoreCase))
+			{
+				OutPackageNames.Add(It.Key);
+			}
+		}
+	}
+}
+
+FName FAssetRegistryState::GetFirstPackageByName(FStringView PackageName) const
+{
+	TArray<FName> LongPackageNames;
+	GetPackagesByName(PackageName, LongPackageNames);
+	if (LongPackageNames.Num() == 0)
+	{
+		return NAME_None;
+	}
+	if (LongPackageNames.Num() > 1)
+	{
+		LongPackageNames.Sort(FNameLexicalLess());
+		UE_LOG(LogAssetRegistry, Warning, TEXT("GetFirstPackageByName('%.*s') is returning '%s', but it also found '%s'%s."),
+			PackageName.Len(), PackageName.GetData(), *LongPackageNames[0].ToString(), *LongPackageNames[1].ToString(),
+			(LongPackageNames.Num() > 2 ? *FString::Printf(TEXT(" and %d others"), LongPackageNames.Num() - 2) : TEXT("")));
+	}
+	return LongPackageNames[0];
+}
+
 bool FAssetRegistryState::GetDependencies(const FAssetIdentifier& AssetIdentifier,
 										  TArray<FAssetIdentifier>& OutDependencies,
 										  EAssetRegistryDependencyType::Type InDependencyType) const
diff --git a/Engine/Source/Runtime/AssetRegistry/Public/AssetRegistry/AssetRegistryState.h b/Engine/Source/Runtime/AssetRegistry/Public/AssetRegistry/AssetRegistryState.h
--- a/Engine/Source/Runtime/AssetRegistry/Public/AssetRegistry/AssetRegistryState.h
+++ b/Engine/Source/Runtime/AssetRegistry/Public/AssetRegistry/AssetRegistryState.h
@@ -153,6 +153,25 @@ public:
 	 */
 	bool EnumerateAllAssets(const TSet<FName>& PackageNamesToSkip, TFunctionRef<bool(const FAssetData&)> Callback, bool bARFiltering = false) const;
 
+	/**
+	 * Gets the LongPackageNames for all packages with the given PackageName.
+	 * Call to check existence of a LongPackageName or find all packages with a ShortPackageName.
+	 *
+	 * @param PackageName Name of the package to find, may be a LongPackageName or ShortPackageName.
+	 * @param OutPackageNames All discovered matching LongPackageNames are appended to this array.
+	 */
+	void GetPackagesByName(FStringView PackageName, TArray<FName>& OutPackageNames) const;
+
+	/**
+	 * Returns the first LongPackageName found for the given PackageName.
+	 * Issues a warning and returns the first (sorted lexically) if there is more than one.
+	 * Call to check existence of a LongPackageName or find a package with a ShortPackageName.
+	 *
+	 * @param PackageName Name of the package to find, may be a LongPackageName or ShortPackageName.
+	 * @return The first LongPackageName of the matching package, or NAME_None if not found.
+	 */
+	FName GetFirstPackageByName(FStringView PackageName) const;
+
 	UE_DEPRECATED(4.26, "Use GetDependencies that takes a UE::AssetRegistry::EDependencyCategory instead")
 	bool GetDependencies(const FAssetIdentifier& AssetIdentifier, TArray<FAssetIdentifier>& OutDependencies, EAssetRegistryDependencyType::Type InDependencyType) const;
 	/**
diff --git a/Engine/Source/Runtime/AssetRegistry/Public/AssetRegistry/IAssetRegistry.h b/Engine/Source/Runtime/AssetRegistry/Public/AssetRegistry/IAssetRegistry.h
--- a/Engine/Source/Runtime/AssetRegistry/Public/AssetRegistry/IAssetRegistry.h
+++ b/Engine/Source/Runtime/AssetRegistry/Public/AssetRegistry/IAssetRegistry.h
@@ -4,6 +4,7 @@
 
 #include "CoreMinimal.h"
 #include "Containers/BitArray.h"
+#include "Containers/StringFwd.h"
 #include "AssetRegistry/AssetData.h"
 #include "Misc/AssetRegistryInterface.h"
 #include "UObject/Interface.h"
@@ -208,6 +209,25 @@ public:
 	 */
 	virtual bool EnumerateAllAssets(TFunctionRef<bool(const FAssetData&)> Callback, bool bIncludeOnlyOnDiskAssets = false) const = 0;
 
+	/**
+	 * Gets the LongPackageName for all packages with the given PackageName.
+	 * Call to check existence of a LongPackageName or find all packages with a ShortPackageName.
+	 * 
+	 * @param PackageName Name of the package to find, may be a LongPackageName or ShortPackageName.
+	 * @param OutPackageNames All discovered matching LongPackageNames are appended to this array.
+	 */
+	virtual void GetPackagesByName(FStringView PackageName, TArray<FName>& OutPackageNames) const = 0;
+
+	/**
+	 * Returns the first LongPackageName found for the given PackageName.
+	 * Issues a warning and returns the first (sorted lexically) if there is more than one.
+	 * Call to check existence of a LongPackageName or find a package with a ShortPackageName.
+	 *
+	 * @param PackageName Name of the package to find, may be a LongPackageName or ShortPackageName.
+	 * @return The first LongPackageName of the matching package, or NAME_None if not found.
+	 */
+	virtual FName GetFirstPackageByName(FStringView PackageName) const = 0;
+
 	UE_DEPRECATED(4.26, "Use GetDependencies that takes a UE::AssetRegistry::EDependencyCategory instead")
 	virtual bool GetDependencies(const FAssetIdentifier& AssetIdentifier, TArray<FAssetIdentifier>& OutDependencies, EAssetRegistryDependencyType::Type InDependencyType) const = 0;
 	/**
diff --git a/Engine/Source/Runtime/Engine/Private/LevelStreaming.cpp b/Engine/Source/Runtime/Engine/Private/LevelStreaming.cpp
--- a/Engine/Source/Runtime/Engine/Private/LevelStreaming.cpp
+++ b/Engine/Source/Runtime/Engine/Private/LevelStreaming.cpp
@@ -31,6 +31,7 @@
 #include "Engine/NetDriver.h"
 #include "Engine/PackageMapClient.h"
 #include "Serialization/LoadTimeTrace.h"
+#include "AssetRegistry/AssetRegistryModule.h"
 
 DEFINE_LOG_CATEGORY(LogLevelStreaming);
 
@@ -1928,13 +1929,15 @@ ULevelStreamingDynamic* ULevelStreamingDynamic::LoadLevelInstance(UObject* World
 		return nullptr;
 	}
 
-	// Check whether requested map exists, this could be very slow if LevelName is a short package name
-	FString LongPackageName;
-	bOutSuccess = FPackageName::SearchForPackageOnDisk(LevelName, &LongPackageName);
+	// Check whether requested map exists
+	IAssetRegistry& AssetRegistry = FModuleManager::LoadModuleChecked<FAssetRegistryModule>(TEXT("AssetRegistry")).Get();
+	FName ExistingPackageName = AssetRegistry.GetFirstPackageByName(LevelName);
+	bOutSuccess = !ExistingPackageName.IsNone();
 	if (!bOutSuccess)
 	{
 		return nullptr;
 	}
+	FString LongPackageName = ExistingPackageName.ToString();
 
 	return LoadLevelInstance_Internal(World, LongPackageName, FTransform(Rotation, Location), bOutSuccess, OptionalLevelNameOverride, OptionalLevelStreamingClass);
 }
diff --git a/Engine/Source/Runtime/Engine/Private/URL.cpp b/Engine/Source/Runtime/Engine/Private/URL.cpp
--- a/Engine/Source/Runtime/Engine/Private/URL.cpp
+++ b/Engine/Source/Runtime/Engine/Private/URL.cpp
@@ -492,10 +492,12 @@ FURL::FURL( FURL* Base, const TCHAR* TextURL, ETravelType Type )
 
 				if (!bFoundMap)
 				{
-					// Fall back to incredibly slow disk scan for the package
-					if (FPackageName::SearchForPackageOnDisk(URLStr + FPackageName::GetMapPackageExtension(), &MapFullName))
+					// Fall back to a slow AssetRegistry scan for the package
+					IAssetRegistry& AssetRegistry = FModuleManager::LoadModuleChecked<FAssetRegistryModule>(TEXT("AssetRegistry")).Get();
+					FName ExistingPackageName = AssetRegistry.GetFirstPackageByName(URLStr);
+					if (!ExistingPackageName.IsNone())
 					{
-						Map = MapFullName;
+						ExistingPackageName.ToString(Map);
 						bFoundMap = true;
 					}
 				}
diff --git a/Engine/Source/Runtime/Engine/Private/UnrealEngine.cpp b/Engine/Source/Runtime/Engine/Private/UnrealEngine.cpp
--- a/Engine/Source/Runtime/Engine/Private/UnrealEngine.cpp
+++ b/Engine/Source/Runtime/Engine/Private/UnrealEngine.cpp
@@ -12300,12 +12300,13 @@ bool UEngine::MakeSureMapNameIsValid(FString& InOutMapName)
 	}
 	else
 	{
-		// Look up on disk. Slow!
-		FString LongPackageName;
-		bIsValid = FPackageName::SearchForPackageOnDisk(TestMapName, &LongPackageName);
-		if (bIsValid)
+		// Look up in the AssetRegistry.
+		IAssetRegistry& AssetRegistry = FModuleManager::LoadModuleChecked<FAssetRegistryModule>(TEXT("AssetRegistry")).Get();
+		FName ExistingPackageName = AssetRegistry.GetFirstPackageByName(TestMapName);
+		if (!ExistingPackageName.IsNone())
 		{
-			InOutMapName = LongPackageName;
+			ExistingPackageName.ToString(InOutMapName);
+			bIsValid = true;
 		}
 	}
 	return bIsValid;
