IAssetRegistry::GetPackagesByName: Fix clang compile error.
#rnx

diff --git a/Engine/Source/Runtime/AssetRegistry/Private/AssetRegistryState.cpp b/Engine/Source/Runtime/AssetRegistry/Private/AssetRegistryState.cpp
--- a/Engine/Source/Runtime/AssetRegistry/Private/AssetRegistryState.cpp
+++ b/Engine/Source/Runtime/AssetRegistry/Private/AssetRegistryState.cpp
@@ -804,7 +804,7 @@ void FAssetRegistryState::GetPackagesByName(FStringView PackageName, TArray<FNam
 	else
 	{
 		TStringBuilder<128> PackageNameStr;
-		for (const TPair<FName, TArray<FAssetData*>>& It : CachedAssetsByPackageName)
+		for (const auto& It : CachedAssetsByPackageName)
 		{
 			It.Key.ToString(PackageNameStr);
 			FStringView ExistingBaseName = FPathViews::GetBaseFilename(PackageNameStr);
