Alphabetize includes
#rnx

diff --git a/Engine/Source/Runtime/AssetRegistry/Private/AssetRegistry.cpp b/Engine/Source/Runtime/AssetRegistry/Private/AssetRegistry.cpp
--- a/Engine/Source/Runtime/AssetRegistry/Private/AssetRegistry.cpp
+++ b/Engine/Source/Runtime/AssetRegistry/Private/AssetRegistry.cpp
@@ -2,39 +2,38 @@
 
 #include "AssetRegistry.h"
 
-#include "Misc/CommandLine.h"
-#include "Misc/FileHelper.h"
-#include "Misc/Paths.h"
-#include "Serialization/ArrayReader.h"
-#include "Misc/ConfigCacheIni.h"
-#include "UObject/UObjectHash.h"
-#include "UObject/UObjectIterator.h"
-#include "UObject/MetaData.h"
-#include "UObject/CoreRedirects.h"
-#include "Blueprint/BlueprintSupport.h"
 #include "AssetDataGatherer.h"
-#include "AssetRegistryPrivate.h"
 #include "AssetRegistry/ARFilter.h"
+#include "AssetRegistryConsoleCommands.h"
+#include "AssetRegistryPrivate.h"
+#include "Async/Async.h"
+#include "Blueprint/BlueprintSupport.h"
 #include "DependsNode.h"
-#include "PackageReader.h"
 #include "GenericPlatform/GenericPlatformChunkInstall.h"
+#include "HAL/PlatformMisc.h"
+#include "HAL/ThreadHeartBeat.h"
 #include "Interfaces/IPluginManager.h"
+#include "Misc/CommandLine.h"
+#include "Misc/ConfigCacheIni.h"
 #include "Misc/CoreDelegates.h"
-#include "UObject/ConstructorHelpers.h"
+#include "Misc/FileHelper.h"
+#include "Misc/Paths.h"
 #include "Misc/RedirectCollector.h"
-#include "Async/Async.h"
-#include "Serialization/LargeMemoryReader.h"
 #include "Misc/ScopeExit.h"
-#include "HAL/ThreadHeartBeat.h"
-#include "HAL/PlatformMisc.h"
-#include "AssetRegistryConsoleCommands.h"
-#include "Interfaces/IPluginManager.h"
+#include "PackageReader.h"
+#include "Serialization/ArrayReader.h"
+#include "Serialization/LargeMemoryReader.h"
+#include "UObject/ConstructorHelpers.h"
+#include "UObject/CoreRedirects.h"
+#include "UObject/MetaData.h"
+#include "UObject/UObjectHash.h"
+#include "UObject/UObjectIterator.h"
 
 #if WITH_EDITOR
-#include "IDirectoryWatcher.h"
 #include "DirectoryWatcherModule.h"
-#include "HAL/PlatformProcess.h"
 #include "HAL/IConsoleManager.h"
+#include "HAL/PlatformProcess.h"
+#include "IDirectoryWatcher.h"
 #endif // WITH_EDITOR
 
 static FAssetRegistryConsoleCommands ConsoleCommands; // Registers its various console commands in the constructor
