Asset Registry: don't query asset registry tags for memory mounted packages.

diff --git a/Engine/Source/Runtime/AssetRegistry/Private/AssetRegistry.cpp b/Engine/Source/Runtime/AssetRegistry/Private/AssetRegistry.cpp
--- a/Engine/Source/Runtime/AssetRegistry/Private/AssetRegistry.cpp
+++ b/Engine/Source/Runtime/AssetRegistry/Private/AssetRegistry.cpp
@@ -1036,13 +1036,15 @@ bool UAssetRegistryImpl::EnumerateAssets(const FARCompiledFilter& InFilter, TFun
 
 				// Skip classes that report themselves as assets but that the editor AssetRegistry is currently not counting as assets
 				if (UE::AssetRegistry::FFiltering::ShouldSkipAsset(Obj))
+				// Package name
+				const FName PackageName = InMemoryPackage->GetFName();
+
+				// Skip assets from memory mounted packages (/Memory/...)
+				if (FPackageName::IsMemoryPackage(PackageName.ToString()))
 				{
 					return;
 				}
 
-				// Package name
-				const FName PackageName = InMemoryPackage->GetFName();
-
 				PackagesToSkip.Add(PackageName);
 
 				if (InFilter.PackageNames.Num() && !InFilter.PackageNames.Contains(PackageName))
