Fix bad Robomerge.

diff --git a/Engine/Source/Runtime/AssetRegistry/Private/AssetRegistry.cpp b/Engine/Source/Runtime/AssetRegistry/Private/AssetRegistry.cpp
--- a/Engine/Source/Runtime/AssetRegistry/Private/AssetRegistry.cpp
+++ b/Engine/Source/Runtime/AssetRegistry/Private/AssetRegistry.cpp
@@ -1036,6 +1036,10 @@ bool UAssetRegistryImpl::EnumerateAssets(const FARCompiledFilter& InFilter, TFun
 
 				// Skip classes that report themselves as assets but that the editor AssetRegistry is currently not counting as assets
 				if (UE::AssetRegistry::FFiltering::ShouldSkipAsset(Obj))
+				{
+					return;
+				}
+
 				// Package name
 				const FName PackageName = InMemoryPackage->GetFName();
 
