Reduce the time wasted in package access tracking during cook due to redundant sequential operations on the same referencer, referenced object, or both.

diff --git a/Engine/Source/Runtime/AssetRegistry/Private/PackageBuildDependencyTracker.cpp b/Engine/Source/Runtime/AssetRegistry/Private/PackageBuildDependencyTracker.cpp
--- a/Engine/Source/Runtime/AssetRegistry/Private/PackageBuildDependencyTracker.cpp
+++ b/Engine/Source/Runtime/AssetRegistry/Private/PackageBuildDependencyTracker.cpp
@@ -70,7 +70,21 @@ void FPackageBuildDependencyTracker::StaticOnObjectHandleRead(UObject* ReadObjec
 					if (Referencer != Referenced)
 					{
 						FScopeLock RecordsScopeLock(&Singleton.RecordsLock);
-						Singleton.Records.FindOrAdd(Referencer).Add(Referenced);
+						if (Referencer == Singleton.LastReferencer)
+						{
+							if (Referenced != Singleton.LastReferenced)
+							{
+								Singleton.LastReferencerSet->Add(Referenced);
+								Singleton.LastReferenced = Referenced;
+							}
+						}
+						else
+						{
+							Singleton.LastReferenced = Referenced;
+							Singleton.LastReferencer = Referencer;
+							Singleton.LastReferencerSet = &Singleton.Records.FindOrAdd(Referencer);
+							Singleton.LastReferencerSet->Add(Referenced);
+						}
 					}
 					break;
 				}
diff --git a/Engine/Source/Runtime/AssetRegistry/Private/PackageBuildDependencyTracker.h b/Engine/Source/Runtime/AssetRegistry/Private/PackageBuildDependencyTracker.h
--- a/Engine/Source/Runtime/AssetRegistry/Private/PackageBuildDependencyTracker.h
+++ b/Engine/Source/Runtime/AssetRegistry/Private/PackageBuildDependencyTracker.h
@@ -25,6 +25,9 @@ private:
 
 	mutable FCriticalSection RecordsLock;
 	TMap<FName, TSet<FName>> Records;
+	FName LastReferencer = NAME_None;
+	FName LastReferenced = NAME_None;
+	TSet<FName>* LastReferencerSet = nullptr;
 	static FPackageBuildDependencyTracker Singleton;
 };
 
