#rnx

diff --git a/Engine/Source/Runtime/AssetRegistry/Private/AssetRegistry.cpp b/Engine/Source/Runtime/AssetRegistry/Private/AssetRegistry.cpp
--- a/Engine/Source/Runtime/AssetRegistry/Private/AssetRegistry.cpp
+++ b/Engine/Source/Runtime/AssetRegistry/Private/AssetRegistry.cpp
@@ -1209,12 +1209,36 @@ bool UAssetRegistryImpl::EnumerateAllAssets(TFunctionRef<bool(const FAssetData&)
 
 void UAssetRegistryImpl::GetPackagesByName(FStringView PackageName, TArray<FName>& OutPackageNames) const
 {
+	UE_CLOG(bInitialSearchStarted && !bInitialSearchCompleted, LogAssetRegistry, Warning,
+		TEXT("GetPackagesByName has been called before AssetRegistry gather is complete and it does not wait. ")
+		TEXT("The search may return incomplete results."));
 	State.GetPackagesByName(PackageName, OutPackageNames);
 }
 
 FName UAssetRegistryImpl::GetFirstPackageByName(FStringView PackageName) const
 {
-	return State.GetFirstPackageByName(PackageName);
+	UE_CLOG(bInitialSearchStarted && !bInitialSearchCompleted, LogAssetRegistry, Warning,
+		TEXT("GetFirstPackageByName has been called before AssetRegistry gather is complete and it does not wait. ")
+		TEXT("The search may fail to find the package."));
+	FName LongPackageName = State.GetFirstPackageByName(PackageName);
+#if WITH_EDITOR
+	if (!GIsEditor && !bSearchAllAssets)
+	{
+		// Temporary support for -game:
+		// When running editor.exe with -game, we do not have a cooked AssetRegistry and we do not scan either
+		// In that case, fall back to searching on disk if the search in the AssetRegistry (as expected) fails
+		// In the future we plan to avoid this situation by having -game run the scan as well
+		if (LongPackageName.IsNone())
+		{
+			FString LongPackageNameString;
+			if (FPackageName::SearchForPackageOnDisk(FString(PackageName), &LongPackageNameString))
+			{
+				LongPackageName = FName(*LongPackageNameString);
+			}
+		}
+	}
+#endif
+	return LongPackageName;
 }
 
 bool UAssetRegistryImpl::GetDependencies(const FAssetIdentifier& AssetIdentifier, TArray<FAssetIdentifier>& OutDependencies, EAssetRegistryDependencyType::Type InDependencyType) const
