UAssetRegistryImpl: Early out of serialization when object reference collecting to avoid expensive pointless calculation.
#rnx

diff --git a/Engine/Source/Runtime/AssetRegistry/Private/AssetRegistry.cpp b/Engine/Source/Runtime/AssetRegistry/Private/AssetRegistry.cpp
--- a/Engine/Source/Runtime/AssetRegistry/Private/AssetRegistry.cpp
+++ b/Engine/Source/Runtime/AssetRegistry/Private/AssetRegistry.cpp
@@ -2169,12 +2169,17 @@ void UAssetRegistryImpl::TickGatherPackage(const FString& PackageName)
 
 void UAssetRegistryImpl::Serialize(FArchive& Ar)
 {
-	if (Ar.IsLoading())
+	if (Ar.IsObjectReferenceCollector())
+	{
+		// The Asset Registry does not have any object references, and its serialization function is expensive
+		return;
+	}
+	else if (Ar.IsLoading())
 	{
 		State.Load(Ar);
 		CachePathsFromState(State);
 	}
-	else
+	else if (Ar.IsSaving())
 	{
 		State.Save(Ar, SerializationOptions);
 	}
