Optimize asset registry UpdateTemporaryCaches()

Saves 100ms in Engine::Init() w/o UAssetManager::Push/PopBulkScanning optimization and 35ms with it, measured on Gen5 for an internal project.

diff --git a/Engine/Source/Runtime/AssetRegistry/Private/AssetRegistry.cpp b/Engine/Source/Runtime/AssetRegistry/Private/AssetRegistry.cpp
--- a/Engine/Source/Runtime/AssetRegistry/Private/AssetRegistry.cpp
+++ b/Engine/Source/Runtime/AssetRegistry/Private/AssetRegistry.cpp
@@ -3445,6 +3445,8 @@ void UAssetRegistryImpl::ClearTemporaryCaches() const
 	}
 }
 
+COREUOBJECT_API TMap<UClass*, TSet<UClass*>> GetAllDerivedClasses();
+
 void UAssetRegistryImpl::UpdateTemporaryCaches() const
 {
 	if (bIsTempCachingEnabled && bIsTempCachingUpToDate && TempCachingRegisteredClassesVersionNumber == GetRegisteredClassesVersionNumber())
@@ -3457,38 +3459,46 @@ void UAssetRegistryImpl::UpdateTemporaryCaches() const
 	// Refreshes ClassGeneratorNames if out of date due to module load
 	CollectCodeGeneratorClasses();
 
-	TempCachedInheritanceMap = CachedBPInheritanceMap;
-	TempReverseInheritanceMap.Reset();
 	TempCachingRegisteredClassesVersionNumber = GetRegisteredClassesVersionNumber();
-	for (TObjectIterator<UClass> ClassIt; ClassIt; ++ClassIt)
+	TMap<UClass*, TSet<UClass*>> NativeSubclasses = GetAllDerivedClasses();
+
+	uint32 NumNativeClasses = 1; // UObject has no superclass
+	for (const TPair<UClass*, TSet<UClass*>>& Pair : NativeSubclasses)
 	{
-		UClass* Class = *ClassIt;
+		NumNativeClasses += Pair.Value.Num();
+	}
+	TempCachedInheritanceMap.Reserve(NumNativeClasses + CachedBPInheritanceMap.Num());
+	TempCachedInheritanceMap = CachedBPInheritanceMap;
+	TempCachedInheritanceMap.Add(FName("Object"), FName());
 
-		if (!Class->HasAnyClassFlags(CLASS_Deprecated | CLASS_NewerVersionExists))
-		{
-			FName ClassName = Class->GetFName();
-			if (Class->GetSuperClass())
-			{
-				FName SuperClassName = Class->GetSuperClass()->GetFName();
-				TSet<FName>& ChildClasses = TempReverseInheritanceMap.FindOrAdd(SuperClassName);
-				ChildClasses.Add(ClassName);
+	for (TPair<FName, TArray<FName>>& Pair : TempReverseInheritanceMap)
+	{
+		Pair.Value.Reset();
+	}
+	TempReverseInheritanceMap.Reserve(NativeSubclasses.Num());
 
-				TempCachedInheritanceMap.Add(ClassName, SuperClassName);
-			}
-			else
-			{
-				// This should only be true for a small number of CoreUObject classes
-				TempCachedInheritanceMap.Add(ClassName, NAME_None);
-			}
+	for (const TPair<UClass*, TSet<UClass*>>& Pair : NativeSubclasses)
+	{
+		FName SuperclassName = Pair.Key->GetFName();
 
-			// Add any implemented interfaces to the reverse inheritance map, but not to the forward map
-			for (int32 i = 0; i < Class->Interfaces.Num(); ++i)
+		TArray<FName>& TempCachedSubclasses = TempReverseInheritanceMap.FindOrAdd(SuperclassName);
+		TempCachedSubclasses.Reserve(Pair.Value.Num());
+		for (UClass* Subclass : Pair.Value)
+		{
+			if (!Subclass->HasAnyClassFlags(CLASS_Deprecated | CLASS_NewerVersionExists))
 			{
-				UClass* InterfaceClass = Class->Interfaces[i].Class;
-				if (InterfaceClass) // could be nulled out by ForceDelete of a blueprint interface
+				FName SubclassName = Subclass->GetFName();
+				TempCachedSubclasses.Add(SubclassName);
+				TempCachedInheritanceMap.Add(SubclassName, SuperclassName);
+
+				// Add any implemented interfaces to the reverse inheritance map, but not to the forward map
+				for (const FImplementedInterface& Interface : Subclass->Interfaces)
 				{
-					TSet<FName>& ChildClasses = TempReverseInheritanceMap.FindOrAdd(InterfaceClass->GetFName());
-					ChildClasses.Add(ClassName);
+					if (UClass* InterfaceClass = Interface.Class) // could be nulled out by ForceDelete of a blueprint interface
+					{
+						TArray<FName>& Implementations = TempReverseInheritanceMap.FindOrAdd(InterfaceClass->GetFName());
+						Implementations.Add(SubclassName);
+					}
 				}
 			}
 		}
@@ -3497,10 +3507,10 @@ void UAssetRegistryImpl::UpdateTemporaryCaches() const
 	// Add non-native classes to reverse map
 	for (const TPair<FName, FName>& Kvp : CachedBPInheritanceMap)
 	{
-		const FName& ParentClassName = Kvp.Value;
-		if (ParentClassName != NAME_None)
+		FName ParentClassName = Kvp.Value;
+		if (!ParentClassName.IsNone())
 		{
-			TSet<FName>& ChildClasses = TempReverseInheritanceMap.FindOrAdd(ParentClassName);
+			TArray<FName>& ChildClasses = TempReverseInheritanceMap.FindOrAdd(ParentClassName);
 			ChildClasses.Add(Kvp.Key);
 		}
 	}
@@ -3522,7 +3532,7 @@ void UAssetRegistryImpl::GetSubClasses(const TArray<FName>& InClassNames, const
 	ClearTemporaryCaches();
 }
 
-void UAssetRegistryImpl::GetSubClasses_Recursive(FName InClassName, TSet<FName>& SubClassNames, TSet<FName>& ProcessedClassNames, const TMap<FName, TSet<FName>>& ReverseInheritanceMap, const TSet<FName>& ExcludedClassNames) const
+void UAssetRegistryImpl::GetSubClasses_Recursive(FName InClassName, TSet<FName>& SubClassNames, TSet<FName>& ProcessedClassNames, const TMap<FName, TArray<FName>>& ReverseInheritanceMap, const TSet<FName>& ExcludedClassNames) const
 {
 	if (ExcludedClassNames.Contains(InClassName))
 	{
@@ -3537,7 +3547,7 @@ void UAssetRegistryImpl::GetSubClasses_Recursive(FName InClassName, TSet<FName>&
 		SubClassNames.Add(InClassName);
 		ProcessedClassNames.Add(InClassName);
 
-		const TSet<FName>* FoundSubClassNames = ReverseInheritanceMap.Find(InClassName);
+		const TArray<FName>* FoundSubClassNames = ReverseInheritanceMap.Find(InClassName);
 		if (FoundSubClassNames)
 		{
 			for (FName ClassName : (*FoundSubClassNames))
diff --git a/Engine/Source/Runtime/AssetRegistry/Private/AssetRegistry.h b/Engine/Source/Runtime/AssetRegistry/Private/AssetRegistry.h
--- a/Engine/Source/Runtime/AssetRegistry/Private/AssetRegistry.h
+++ b/Engine/Source/Runtime/AssetRegistry/Private/AssetRegistry.h
@@ -268,7 +268,7 @@ private:
 
 	/** Returns the names of all subclasses of the class whose name is ClassName */
 	void GetSubClasses(const TArray<FName>& InClassNames, const TSet<FName>& ExcludedClassNames, TSet<FName>& SubClassNames) const;
-	void GetSubClasses_Recursive(FName InClassName, TSet<FName>& SubClassNames, TSet<FName>& ProcessedClassNames, const TMap<FName, TSet<FName>>& ReverseInheritanceMap, const TSet<FName>& ExcludedClassNames) const;
+	void GetSubClasses_Recursive(FName InClassName, TSet<FName>& SubClassNames, TSet<FName>& ProcessedClassNames, const TMap<FName, TArray<FName>>& ReverseInheritanceMap, const TSet<FName>& ExcludedClassNames) const;
 
 	/** Finds all class names of classes capable of generating new UClasses */
 	void CollectCodeGeneratorClasses() const;
@@ -351,7 +351,7 @@ private:
 	mutable TMap<FName, FName> TempCachedInheritanceMap;
 
 	/** A reverse map of TempCachedInheritanceMap, only kept during temp caching */
-	mutable TMap<FName, TSet<FName>> TempReverseInheritanceMap;
+	mutable TMap<FName, TArray<FName>> TempReverseInheritanceMap;
 
 	/** If true, temp caching has been computed and is valid.
 	    Set this to false when changing something that might invalidate the cache so it gets recomputed on-demand.
diff --git a/Engine/Source/Runtime/CoreUObject/Private/UObject/UObjectHash.cpp b/Engine/Source/Runtime/CoreUObject/Private/UObject/UObjectHash.cpp
--- a/Engine/Source/Runtime/CoreUObject/Private/UObject/UObjectHash.cpp
+++ b/Engine/Source/Runtime/CoreUObject/Private/UObject/UObjectHash.cpp
@@ -1232,6 +1232,13 @@ void GetDerivedClasses(const UClass* ClassToLookFor, TArray<UClass*>& Results, b
 	}
 }
 
+TMap<UClass*, TSet<UClass*>> GetAllDerivedClasses()
+{
+	FUObjectHashTables& ThreadHash = FUObjectHashTables::Get();
+	FHashTableLock HashLock(ThreadHash);
+	return ThreadHash.ClassToChildListMap;
+}
+
 bool ClassHasInstancesAsyncLoading(const UClass* ClassToLookFor)
 {
 	TArray<const UClass*> ClassesToSearch;
diff --git a/Engine/Source/Runtime/CoreUObject/Public/UObject/UObjectHash.h b/Engine/Source/Runtime/CoreUObject/Public/UObject/UObjectHash.h
--- a/Engine/Source/Runtime/CoreUObject/Public/UObject/UObjectHash.h
+++ b/Engine/Source/Runtime/CoreUObject/Public/UObject/UObjectHash.h
@@ -150,6 +150,9 @@ COREUOBJECT_API void ForEachObjectOfClasses(TArrayView<const UClass*> ClassesToL
  */
 COREUOBJECT_API void GetDerivedClasses(const UClass* ClassToLookFor, TArray<UClass *>& Results, bool bRecursive = true);
 
+/** Get all base classes and their direct subclasses */
+COREUOBJECT_API TMap<UClass*, TSet<UClass*>> GetAllDerivedClasses();
+
 /**
  * Returns true if any instances of the class in question are currently being async loaded.
  *
