Threadsafe AssetRegistry - add synchronization that allows the AssetRegistry to be used from any thread, so that it can be used in asynchronous TryPreload threads by the EditorDomain during cooking.
#rn Minor AssetRegistry

diff --git a/Engine/Plugins/Editor/AssetManagerEditor/Source/AssetManagerEditor/Private/AssetManagerEditorModule.cpp b/Engine/Plugins/Editor/AssetManagerEditor/Source/AssetManagerEditor/Private/AssetManagerEditorModule.cpp
--- a/Engine/Plugins/Editor/AssetManagerEditor/Source/AssetManagerEditor/Private/AssetManagerEditorModule.cpp
+++ b/Engine/Plugins/Editor/AssetManagerEditor/Source/AssetManagerEditor/Private/AssetManagerEditorModule.cpp
@@ -1086,11 +1086,11 @@ bool FAssetManagerEditorModule::GetManagedPackageListForAssetData(const FAssetDa
 
 	TArray<FAssetIdentifier> FoundDependencies;
 
-	if (CurrentRegistrySource->RegistryState)
+	if (CurrentRegistrySource->HasRegistry())
 	{
 		for (const FPrimaryAssetId& PrimaryAssetId : PrimaryAssetSet)
 		{
-			CurrentRegistrySource->RegistryState->GetDependencies(PrimaryAssetId, FoundDependencies, UE::AssetRegistry::EDependencyCategory::Manage);
+			CurrentRegistrySource->GetDependencies(PrimaryAssetId, FoundDependencies, UE::AssetRegistry::EDependencyCategory::Manage);
 		}
 	}
 	
@@ -1107,7 +1107,7 @@ bool FAssetManagerEditorModule::GetManagedPackageListForAssetData(const FAssetDa
 
 bool FAssetManagerEditorModule::GetStringValueForCustomColumn(const FAssetData& AssetData, FName ColumnName, FString& OutValue)
 {
-	if (!CurrentRegistrySource || !CurrentRegistrySource->RegistryState)
+	if (!CurrentRegistrySource || !CurrentRegistrySource->HasRegistry())
 	{
 		return false;
 	}
@@ -1158,10 +1158,10 @@ bool FAssetManagerEditorModule::GetStringValueForCustomColumn(const FAssetData&
 		}
 		else
 		{
-			const FAssetData* PlatformData = CurrentRegistrySource->RegistryState->GetAssetByObjectPath(AssetData.ObjectPath);
-			if (PlatformData)
+			FAssetData PlatformData = CurrentRegistrySource->GetAssetByObjectPath(AssetData.ObjectPath);
+			if (PlatformData.IsValid())
 			{
-				FoundChunks = PlatformData->ChunkIDs;
+				FoundChunks = MoveTemp(PlatformData.ChunkIDs);
 			}
 		}
 		
@@ -1188,7 +1188,7 @@ bool FAssetManagerEditorModule::GetStringValueForCustomColumn(const FAssetData&
 
 bool FAssetManagerEditorModule::GetDisplayTextForCustomColumn(const FAssetData& AssetData, FName ColumnName, FText& OutValue)
 {
-	if (!CurrentRegistrySource || !CurrentRegistrySource->RegistryState)
+	if (!CurrentRegistrySource || !CurrentRegistrySource->HasRegistry())
 	{
 		return false;
 	}
@@ -1256,7 +1256,7 @@ bool FAssetManagerEditorModule::GetDisplayTextForCustomColumn(const FAssetData&
 
 bool FAssetManagerEditorModule::GetIntegerValueForCustomColumn(const FAssetData& AssetData, FName ColumnName, int64& OutValue)
 {
-	if (!CurrentRegistrySource || !CurrentRegistrySource->RegistryState)
+	if (!CurrentRegistrySource || !CurrentRegistrySource->HasRegistry())
 	{
 		return false;
 	}
@@ -1306,7 +1306,7 @@ bool FAssetManagerEditorModule::GetIntegerValueForCustomColumn(const FAssetData&
 	}
 	else if (ColumnName == DiskSizeName)
 	{
-		const FAssetPackageData* FoundData = CurrentRegistrySource->RegistryState->GetAssetPackageData(AssetData.PackageName);
+		TOptional<FAssetPackageData> FoundData = CurrentRegistrySource->GetAssetPackageDataCopy(AssetData.PackageName);
 
 		if (FoundData && FoundData->DiskSize >= 0)
 		{
@@ -1452,12 +1452,9 @@ void FAssetManagerEditorModule::SetCurrentRegistrySource(const FString& SourceNa
 
 		if (CurrentRegistrySource->SourceName == FAssetManagerEditorRegistrySource::CustomSourceName)
 		{
-			if (CurrentRegistrySource->RegistryState)
+			if (CurrentRegistrySource->HasRegistry())
 			{
-				check(!CurrentRegistrySource->bIsEditor);
-
-				delete CurrentRegistrySource->RegistryState;
-				CurrentRegistrySource->RegistryState = nullptr;
+				CurrentRegistrySource->ClearRegistry();
 				CurrentRegistrySource->ChunkAssignments.Reset();
 				CurrentRegistrySource->bManagementDataInitialized = false;
 			}
@@ -1489,7 +1486,7 @@ void FAssetManagerEditorModule::SetCurrentRegistrySource(const FString& SourceNa
 			}
 		}
 
-		if (CurrentRegistrySource->RegistryState == nullptr && !CurrentRegistrySource->SourceFilename.IsEmpty() && !CurrentRegistrySource->bIsEditor)
+		if (!CurrentRegistrySource->HasRegistry() && !CurrentRegistrySource->SourceFilename.IsEmpty() && !CurrentRegistrySource->bIsEditor)
 		{
 			bool bLoaded = false;
 			FArrayReader SerializedAssetData;
@@ -1504,7 +1501,7 @@ void FAssetManagerEditorModule::SetCurrentRegistrySource(const FString& SourceNa
 				if (NewState->GetObjectPathToAssetDataMap().Num() > 0)
 				{
 					bLoaded = true;
-					CurrentRegistrySource->RegistryState = NewState;
+					CurrentRegistrySource->SetRegistryState(NewState);
 				}
 			}
 			if (!bLoaded)
@@ -1516,7 +1513,7 @@ void FAssetManagerEditorModule::SetCurrentRegistrySource(const FString& SourceNa
 			}
 		}
 
-		if (!CurrentRegistrySource->bManagementDataInitialized && CurrentRegistrySource->RegistryState)
+		if (!CurrentRegistrySource->bManagementDataInitialized && CurrentRegistrySource->HasRegistry())
 		{
 			CurrentRegistrySource->ChunkAssignments.Reset();
 			if (CurrentRegistrySource->bIsEditor)
@@ -1530,7 +1527,9 @@ void FAssetManagerEditorModule::SetCurrentRegistrySource(const FString& SourceNa
 			else
 			{
 				// Iterate assets and look for chunks
-				const TMap<FName, const FAssetData*>& AssetDataMap = CurrentRegistrySource->RegistryState->GetObjectPathToAssetDataMap();
+				const FAssetRegistryState* RegistryState = CurrentRegistrySource->GetOwnedRegistryState();
+				checkf(RegistryState, TEXT("Should be non-null because HasRegistry() && !bIsEditor"));
+				const TMap<FName, const FAssetData*>& AssetDataMap = RegistryState->GetObjectPathToAssetDataMap();
 
 				for (const TPair<FName, const FAssetData*>& Pair : AssetDataMap)
 				{
@@ -1538,7 +1537,7 @@ void FAssetManagerEditorModule::SetCurrentRegistrySource(const FString& SourceNa
 					if (AssetData.ChunkIDs.Num() > 0)
 					{
 						TArray<FAssetIdentifier> ManagerAssets;
-						CurrentRegistrySource->RegistryState->GetReferencers(AssetData.PackageName, ManagerAssets, UE::AssetRegistry::EDependencyCategory::Manage);
+						RegistryState->GetReferencers(AssetData.PackageName, ManagerAssets, UE::AssetRegistry::EDependencyCategory::Manage);
 
 						for (int32 ChunkId : AssetData.ChunkIDs)
 						{
@@ -1553,7 +1552,7 @@ void FAssetManagerEditorModule::SetCurrentRegistrySource(const FString& SourceNa
 								
 								TArray<FAssetIdentifier> ManagedAssets;
 
-								CurrentRegistrySource->RegistryState->GetDependencies(ChunkAssetId, ManagedAssets, UE::AssetRegistry::EDependencyCategory::Manage);
+								RegistryState->GetDependencies(ChunkAssetId, ManagedAssets, UE::AssetRegistry::EDependencyCategory::Manage);
 
 								for (const FAssetIdentifier& ManagedAsset : ManagedAssets)
 								{
@@ -1628,9 +1627,11 @@ void FAssetManagerEditorModule::RefreshRegistryData()
 
 bool FAssetManagerEditorModule::IsPackageInCurrentRegistrySource(FName PackageName)
 {
-	if (CurrentRegistrySource && CurrentRegistrySource->RegistryState && !CurrentRegistrySource->bIsEditor)
+	if (CurrentRegistrySource && CurrentRegistrySource->HasRegistry() && !CurrentRegistrySource->bIsEditor)
 	{
-		const FAssetPackageData* FoundData = CurrentRegistrySource->RegistryState->GetAssetPackageData(PackageName);
+		const FAssetRegistryState* RegistryState = CurrentRegistrySource->GetOwnedRegistryState();
+		checkf(RegistryState, TEXT("Should be non-null because HasRegistry() && !bIsEditor"));
+		const FAssetPackageData* FoundData = RegistryState->GetAssetPackageData(PackageName);
 
 		if (!FoundData || FoundData->DiskSize < 0)
 		{
@@ -1652,10 +1653,12 @@ bool IAssetManagerEditorModule::FilterAssetIdentifiersForCurrentRegistrySource(T
 bool FAssetManagerEditorModule::FilterAssetIdentifiersForCurrentRegistrySource(TArray<FAssetIdentifier>& AssetIdentifiers, const FAssetManagerDependencyQuery& DependencyQuery, bool bForwardDependency)
 {
 	bool bMadeChange = false;
-	if (!CurrentRegistrySource || !CurrentRegistrySource->RegistryState || CurrentRegistrySource->bIsEditor)
+	if (!CurrentRegistrySource || !CurrentRegistrySource->HasRegistry() || CurrentRegistrySource->bIsEditor)
 	{
 		return bMadeChange;
 	}
+	const FAssetRegistryState* RegistryState = CurrentRegistrySource->GetOwnedRegistryState();
+	checkf(RegistryState, TEXT("Should be non-null because HasRegistry() && !bIsEditor"));
 
 	for (int32 Index = 0; Index < AssetIdentifiers.Num(); Index++)
 	{
@@ -1682,11 +1685,11 @@ bool FAssetManagerEditorModule::FilterAssetIdentifiersForCurrentRegistrySource(T
 
 							if (bForwardDependency)
 							{
-								CurrentRegistrySource->RegistryState->GetDependencies(PackageName, FoundReferences, DependencyQuery.Categories, DependencyQuery.Flags);
+								RegistryState->GetDependencies(PackageName, FoundReferences, DependencyQuery.Categories, DependencyQuery.Flags);
 							}
 							else
 							{
-								CurrentRegistrySource->RegistryState->GetReferencers(PackageName, FoundReferences, DependencyQuery.Categories, DependencyQuery.Flags);
+								RegistryState->GetReferencers(PackageName, FoundReferences, DependencyQuery.Categories, DependencyQuery.Flags);
 							}
 
 							AssetIdentifiers.Insert(FoundReferences, Index);
@@ -1711,12 +1714,13 @@ void FAssetManagerEditorModule::InitializeRegistrySources(bool bNeedManagementDa
 		RegistrySourceMap.Reset();
 
 		// Add Editor source
-		FAssetManagerEditorRegistrySource EditorSource;
-		EditorSource.SourceName = FAssetManagerEditorRegistrySource::EditorSourceName;
-		EditorSource.bIsEditor = true;
-		EditorSource.RegistryState = AssetRegistry->GetAssetRegistryState();
+		{
+			FAssetManagerEditorRegistrySource EditorSource;
+			EditorSource.SourceName = FAssetManagerEditorRegistrySource::EditorSourceName;
+			EditorSource.SetUseEditorAssetRegistry(AssetRegistry);
 
-		RegistrySourceMap.Add(EditorSource.SourceName, EditorSource);
+			RegistrySourceMap.Add(FAssetManagerEditorRegistrySource::EditorSourceName, MoveTemp(EditorSource));
+		}
 
 		TArray<ITargetPlatform*> Platforms = GetTargetPlatformManager()->GetTargetPlatforms();
 
@@ -1731,18 +1735,20 @@ void FAssetManagerEditorModule::InitializeRegistrySources(bool bNeedManagementDa
 				PlatformSource.SourceFilename = RegistryPath;
 				PlatformSource.TargetPlatform = CheckPlatform;
 
-				RegistrySourceMap.Add(PlatformSource.SourceName, PlatformSource);
+				RegistrySourceMap.Add(PlatformSource.SourceName, MoveTemp(PlatformSource));
 			}
 		}
 
 		// Add Custom source
-		FAssetManagerEditorRegistrySource CustomSource;
-		CustomSource.SourceName = FAssetManagerEditorRegistrySource::CustomSourceName;
+		{
+			FAssetManagerEditorRegistrySource CustomSource;
+			CustomSource.SourceName = FAssetManagerEditorRegistrySource::CustomSourceName;
 
-		RegistrySourceMap.Add(CustomSource.SourceName, CustomSource);
+			RegistrySourceMap.Add(CustomSource.SourceName, MoveTemp(CustomSource));
+		}
 
 		// Select the Editor source by default
-		CurrentRegistrySource = RegistrySourceMap.Find(EditorSource.SourceName);
+		CurrentRegistrySource = RegistrySourceMap.Find(FAssetManagerEditorRegistrySource::EditorSourceName);
 	}
 	check(CurrentRegistrySource);
 
@@ -2082,8 +2088,8 @@ void FAssetManagerEditorModule::DumpAssetRegistry(const TArray<FString>& Args)
 	UAssetManager& Manager = UAssetManager::Get();
 	TArray<FString> ReportLines;
 
-	const FAssetRegistryState* State = Manager.GetAssetRegistry().GetAssetRegistryState();
-	State->Dump(Args, ReportLines);
+	const IAssetRegistry& LocalAssetRegistry = Manager.GetAssetRegistry();
+	LocalAssetRegistry.DumpState(Args, ReportLines);
 
 	Manager.WriteCustomReport(FString::Printf(TEXT("AssetRegistryState-%s.txt"), *FDateTime::Now().ToString()), ReportLines);
 #endif
diff --git a/Engine/Plugins/Editor/AssetManagerEditor/Source/AssetManagerEditor/Private/SSizeMap.cpp b/Engine/Plugins/Editor/AssetManagerEditor/Source/AssetManagerEditor/Private/SSizeMap.cpp
--- a/Engine/Plugins/Editor/AssetManagerEditor/Source/AssetManagerEditor/Private/SSizeMap.cpp
+++ b/Engine/Plugins/Editor/AssetManagerEditor/Source/AssetManagerEditor/Private/SSizeMap.cpp
@@ -397,8 +397,7 @@ namespace SizeMapInternals
 
 void SSizeMap::GatherDependenciesRecursively(TSharedPtr<FAssetThumbnailPool>& InAssetThumbnailPool, TMap<FAssetIdentifier, TSharedPtr<FTreeMapNodeData>>& VisitedAssetIdentifiers, const TArray<FAssetIdentifier>& AssetIdentifiers, const FPrimaryAssetId& FilterPrimaryAsset, const TSharedPtr<FTreeMapNodeData>& Node, TSharedPtr<FTreeMapNodeData>& SharedRootNode, int32& NumAssetsWhichFailedToLoad)
 {
-	const FAssetRegistryState* CurrentPlatformState = CurrentRegistrySource && CurrentRegistrySource->RegistryState ? CurrentRegistrySource->RegistryState : nullptr;
-	if (!CurrentPlatformState)
+	if (!CurrentRegistrySource->HasRegistry())
 	{
 		return;
 	}
@@ -512,11 +511,11 @@ void SSizeMap::GatherDependenciesRecursively(TSharedPtr<FAssetThumbnailPool>& In
 				NodeSizeMapData.AssetData.AssetClass = FName(*LOCTEXT("MissingAsset", "MISSING!").ToString());
 
 				const FString AssetPathString = AssetPackageNameString + TEXT(".") + FPackageName::GetLongPackageAssetName(AssetPackageNameString);
-				const FAssetData* FoundPointer = CurrentPlatformState->GetAssetByObjectPath(FName(*AssetPathString));
+				FAssetData FoundData = CurrentRegistrySource->GetAssetByObjectPath(FName(*AssetPathString));
 
-				if (FoundPointer)
+				if (FoundData.IsValid())
 				{
-					NodeSizeMapData.AssetData = *FoundPointer;
+					NodeSizeMapData.AssetData = MoveTemp(FoundData);
 				}
 			}
 			else
@@ -554,7 +553,7 @@ void SSizeMap::GatherDependenciesRecursively(TSharedPtr<FAssetThumbnailPool>& In
 				}
 				else
 				{
-					CurrentPlatformState->GetDependencies(AssetIdentifier, References, DependencyQuery.Categories, DependencyQuery.Flags);
+					CurrentRegistrySource->GetDependencies(AssetIdentifier, References, DependencyQuery.Categories, DependencyQuery.Flags);
 				}
 				
 				// Filter for registry source
@@ -582,7 +581,7 @@ void SSizeMap::GatherDependenciesRecursively(TSharedPtr<FAssetThumbnailPool>& In
 							{
 								// Check to see if this is managed by the filter asset
 								TArray<FAssetIdentifier> Managers;
-								CurrentPlatformState->GetReferencers(FoundAssetIdentifier, Managers, UE::AssetRegistry::EDependencyCategory::Manage);
+								CurrentRegistrySource->GetReferencers(FoundAssetIdentifier, Managers, UE::AssetRegistry::EDependencyCategory::Manage);
 
 								if (!Managers.Contains(FilterPrimaryAsset))
 								{
diff --git a/Engine/Plugins/Editor/AssetManagerEditor/Source/AssetManagerEditor/Public/AssetManagerEditorModule.h b/Engine/Plugins/Editor/AssetManagerEditor/Source/AssetManagerEditor/Public/AssetManagerEditorModule.h
--- a/Engine/Plugins/Editor/AssetManagerEditor/Source/AssetManagerEditor/Public/AssetManagerEditorModule.h
+++ b/Engine/Plugins/Editor/AssetManagerEditor/Source/AssetManagerEditor/Public/AssetManagerEditorModule.h
@@ -62,9 +62,13 @@ struct ASSETMANAGEREDITOR_API FAssetManagerEditorRegistrySource
 	/** Target platform for this state, may be null */
 	ITargetPlatform* TargetPlatform;
 
-	/** Raw asset registry state, if bIsEditor is true this points to the real editor asset registry */
-	const FAssetRegistryState* RegistryState;
+private:
+	/** Raw asset registry state, only valid if bIsEditor is false */
+	const FAssetRegistryState* RegistryStatePrivate;
+	/** Pointer to the real AssetRegistry, only valid if bIsEditor is true */
+	const IAssetRegistry* AssetRegistryPrivate;
 
+public:
 	/** If true, this is the editor  */
 	uint8 bIsEditor : 1;
 
@@ -82,7 +86,8 @@ struct ASSETMANAGEREDITOR_API FAssetManagerEditorRegistrySource
 
 	FAssetManagerEditorRegistrySource()
 		: TargetPlatform(nullptr)
-		, RegistryState(nullptr)
+		, RegistryStatePrivate(nullptr)
+		, AssetRegistryPrivate(nullptr)
 		, bIsEditor(false)
 		, bManagementDataInitialized(false)
 	{
@@ -90,11 +95,173 @@ struct ASSETMANAGEREDITOR_API FAssetManagerEditorRegistrySource
 
 	~FAssetManagerEditorRegistrySource()
 	{
-		if (RegistryState && !bIsEditor)
+		ClearRegistry();
+	}
+
+	// We delete the copy constructor because it's not clear whether we should copy RegistryStatePrivate
+	FAssetManagerEditorRegistrySource(const FAssetManagerEditorRegistrySource& Other) = delete;
+	FAssetManagerEditorRegistrySource& operator=(const FAssetManagerEditorRegistrySource& Other) = delete;
+
+	FAssetManagerEditorRegistrySource& operator=(FAssetManagerEditorRegistrySource&& Other)
+	{
+		SourceName = MoveTemp(Other.SourceName);
+		SourceFilename = MoveTemp(Other.SourceFilename);
+		TargetPlatform = MoveTemp(Other.TargetPlatform);
+		bManagementDataInitialized = Other.bManagementDataInitialized;
+		Other.bManagementDataInitialized = false;
+		ChunkAssignments = MoveTemp(Other.ChunkAssignments);
+
+		if (Other.bIsEditor)
+		{
+			SetUseEditorAssetRegistry(Other.AssetRegistryPrivate);
+		}
+		else
+		{
+			SetRegistryState(Other.RegistryStatePrivate);
+		}
+		// Remove the Other's RegistryState without deleting it
+		Other.RegistryStatePrivate = nullptr;
+		Other.AssetRegistryPrivate = nullptr;
+		Other.bIsEditor = false;
+
+		return *this;
+	}
+
+	FAssetManagerEditorRegistrySource(FAssetManagerEditorRegistrySource&& Other)
+		:FAssetManagerEditorRegistrySource()
+	{
+		*this = MoveTemp(Other);
+	}
+
+	/** Returns the RegistryState owned by this. Will be null if this does not have a Registry or is using the global editor Registry */
+	const FAssetRegistryState* GetOwnedRegistryState() const
+	{
+		if (!bIsEditor)
+		{
+			return RegistryStatePrivate; // May be null
+		}
+		return nullptr;
+	}
+
+	/**
+	 * Set the RegistryState viewed by this FAssetManagerEditorRegistrySource to a dynamically allocated FAssetRegistryState
+	 * FAssetManagerEditorRegistrySource takes ownership of the given AssetRegistry and deletes it is when cleared.
+	 * Do not call with  AssetRegistry->GetAssetRegistryState; use SetUseEditorAssetRegistry(AssetRegistry) instead.
+	 */
+	void SetRegistryState(const FAssetRegistryState* InRegistryState)
+	{
+		ClearRegistry();
+
+		if (InRegistryState)
+		{
+			bIsEditor = false;
+			RegistryStatePrivate = InRegistryState;
+		}
+	}
+
+	/**
+	 * Set the RegistryState viewed by this FAssetManagerEditorRegistrySource to the FAssetRegistryState owned by the given
+	 * global AssetRegistry.
+	 */
+	void SetUseEditorAssetRegistry(const IAssetRegistry* AssetRegistry)
+	{
+		ClearRegistry();
+
+		if (AssetRegistry)
+		{
+			bIsEditor = true;
+			AssetRegistryPrivate = AssetRegistry;
+		}
+	}
+
+	bool HasRegistry() const
+	{
+		return bIsEditor || RegistryStatePrivate != nullptr;
+	}
+
+	void ClearRegistry()
+	{
+		if (!bIsEditor)
+		{
+			delete RegistryStatePrivate; // May be null
+		}
+		bIsEditor = false;
+		RegistryStatePrivate = nullptr;
+		AssetRegistryPrivate = nullptr;
+	}
+
+	// Functions to forward on to the RegistryState, either the owned one or the editor's global State.
+
+	FAssetData GetAssetByObjectPath(FName ObjectPath) const
+	{
+		if (bIsEditor)
+		{
+			return AssetRegistryPrivate->GetAssetByObjectPath(ObjectPath, true /* bIncludeOnlyOnDiskAssets */);
+		}
+		else if (RegistryStatePrivate)
+		{
+			const FAssetData* AssetData = RegistryStatePrivate->GetAssetByObjectPath(ObjectPath);
+			return AssetData ? *AssetData : FAssetData();
+		}
+		else
+		{
+			return nullptr;
+		}
+	}
+
+	bool GetDependencies(const FAssetIdentifier& AssetIdentifier, TArray<FAssetIdentifier>& OutDependencies,
+		UE::AssetRegistry::EDependencyCategory Category = UE::AssetRegistry::EDependencyCategory::All,
+		const UE::AssetRegistry::FDependencyQuery& Flags = UE::AssetRegistry::FDependencyQuery()) const
+	{
+		if (bIsEditor)
+		{
+			return AssetRegistryPrivate->GetDependencies(AssetIdentifier, OutDependencies, Category, Flags);
+		}
+		else if (RegistryStatePrivate)
+		{
+			return RegistryStatePrivate->GetDependencies(AssetIdentifier, OutDependencies, Category, Flags);
+		}
+		else
 		{
-			delete RegistryState;
+			return false;
 		}
 	}
+
+	bool GetReferencers(const FAssetIdentifier& AssetIdentifier, TArray<FAssetIdentifier>& OutReferencers,
+		UE::AssetRegistry::EDependencyCategory Category = UE::AssetRegistry::EDependencyCategory::All,
+		const UE::AssetRegistry::FDependencyQuery& Flags = UE::AssetRegistry::FDependencyQuery()) const
+	{
+		if (bIsEditor)
+		{
+			return AssetRegistryPrivate->GetReferencers(AssetIdentifier, OutReferencers, Category, Flags);
+		}
+		else if (RegistryStatePrivate)
+		{
+			return RegistryStatePrivate->GetReferencers(AssetIdentifier, OutReferencers, Category, Flags);
+		}
+		else
+		{
+			return false;
+		}
+	}
+
+	TOptional<FAssetPackageData> GetAssetPackageDataCopy(FName PackageName) const
+	{
+		if (bIsEditor)
+		{
+			return AssetRegistryPrivate->GetAssetPackageDataCopy(PackageName);
+		}
+		else if (RegistryStatePrivate)
+		{
+			const FAssetPackageData* DataPtr = RegistryStatePrivate->GetAssetPackageData(PackageName);
+			return DataPtr ? TOptional<FAssetPackageData>(*DataPtr) : TOptional<FAssetPackageData>();
+		}
+		else
+		{
+			return TOptional<FAssetPackageData>();
+		}
+	}
+
 };
 
 /**
diff --git a/Engine/Plugins/Editor/ContentBrowser/ContentBrowserAssetDataSource/Source/ContentBrowserAssetDataSource/Private/ContentBrowserAssetDataCore.cpp b/Engine/Plugins/Editor/ContentBrowser/ContentBrowserAssetDataSource/Source/ContentBrowserAssetDataSource/Private/ContentBrowserAssetDataCore.cpp
--- a/Engine/Plugins/Editor/ContentBrowser/ContentBrowserAssetDataSource/Source/ContentBrowserAssetDataSource/Private/ContentBrowserAssetDataCore.cpp
+++ b/Engine/Plugins/Editor/ContentBrowser/ContentBrowserAssetDataSource/Source/ContentBrowserAssetDataSource/Private/ContentBrowserAssetDataCore.cpp
@@ -1100,7 +1100,7 @@ bool GetDiskSizeItemAttribute(const FAssetData& InAssetData, IAssetRegistry* InA
 	check(InAssetData.IsValid());
 	check(InAssetRegistry);
 
-	if (const FAssetPackageData* PackageData = InAssetRegistry->GetAssetPackageData(InAssetData.PackageName))
+	if (TOptional<FAssetPackageData> PackageData = InAssetRegistry->GetAssetPackageDataCopy(InAssetData.PackageName))
 	{
 		OutAttributeValue.SetValue(PackageData->DiskSize);
 
diff --git a/Engine/Plugins/Runtime/AssetTags/Source/AssetTags/Private/AssetTagsSubsystem.cpp b/Engine/Plugins/Runtime/AssetTags/Source/AssetTags/Private/AssetTagsSubsystem.cpp
--- a/Engine/Plugins/Runtime/AssetTags/Source/AssetTags/Private/AssetTagsSubsystem.cpp
+++ b/Engine/Plugins/Runtime/AssetTags/Source/AssetTags/Private/AssetTagsSubsystem.cpp
@@ -5,6 +5,7 @@
 #include "IAssetRegistry.h"
 #include "AssetRegistryState.h"
 #include "AssetRegistryModule.h"
+#include "Containers/StringView.h"
 
 #if WITH_EDITOR
 #include "ICollectionManager.h"
@@ -322,11 +323,8 @@ bool UAssetTagsSubsystem::CollectionExists(const FName Name)
 	{
 		IAssetRegistry& AssetRegistry = FModuleManager::LoadModuleChecked<FAssetRegistryModule>(AssetRegistryConstants::ModuleName).Get();
 
-		if (const FAssetRegistryState* StatePtr = AssetRegistry.GetAssetRegistryState())
-		{
-			const FString TagNameStr = FString::Printf(TEXT("%s%s"), FAssetData::GetCollectionTagPrefix(), *Name.ToString());
-			bExists = StatePtr->GetTagToAssetDatasMap().Contains(*TagNameStr);
-		}
+		const FString TagNameStr = FString::Printf(TEXT("%s%s"), FAssetData::GetCollectionTagPrefix(), *Name.ToString());
+		bExists = AssetRegistry.ContainsTag(*TagNameStr);
 	}
 #endif	// WITH_EDITOR
 
@@ -354,21 +352,19 @@ TArray<FName> UAssetTagsSubsystem::GetCollections()
 	{
 		IAssetRegistry& AssetRegistry = FModuleManager::LoadModuleChecked<FAssetRegistryModule>(AssetRegistryConstants::ModuleName).Get();
 
-		if (const FAssetRegistryState* StatePtr = AssetRegistry.GetAssetRegistryState())
-		{
-			const TCHAR* CollectionTagPrefix = FAssetData::GetCollectionTagPrefix();
-			const int32 CollectionTagPrefixLen = FCString::Strlen(CollectionTagPrefix);
+		FStringView CollectionTagPrefix(FAssetData::GetCollectionTagPrefix());
 
-			for (const auto& TagAndArrayPair : StatePtr->GetTagToAssetDatasMap())
+		FString TagNameStr;
+		AssetRegistry.ReadLockEnumerateTagToAssetDatas(
+			[&TagNameStr, &CollectionTagPrefix, &CollectionNames](FName TagName, const TArray<const FAssetData*>& Assets)
 			{
-				const FString TagNameStr = TagAndArrayPair.Key.ToString();
-				if (TagNameStr.StartsWith(CollectionTagPrefix))
+				TagName.ToString(TagNameStr);
+				if (FStringView(TagNameStr).StartsWith(CollectionTagPrefix, ESearchCase::IgnoreCase))
 				{
-					const FString TrimmedTagNameStr = TagNameStr.Mid(CollectionTagPrefixLen);
+					const FString TrimmedTagNameStr = TagNameStr.Mid(CollectionTagPrefix.Len());
 					CollectionNames.Add(*TrimmedTagNameStr);
 				}
-			}
-		}
+			});
 	}
 #endif	// WITH_EDITOR
 
diff --git a/Engine/Source/Editor/ContentBrowser/Private/SContentBrowser.cpp b/Engine/Source/Editor/ContentBrowser/Private/SContentBrowser.cpp
--- a/Engine/Source/Editor/ContentBrowser/Private/SContentBrowser.cpp
+++ b/Engine/Source/Editor/ContentBrowser/Private/SContentBrowser.cpp
@@ -3373,18 +3373,17 @@ void SContentBrowser::OnAssetSearchSuggestionFilter(const FText& SearchText, TAr
 	{
 		IAssetRegistry& AssetRegistry = FModuleManager::LoadModuleChecked<FAssetRegistryModule>(AssetRegistryConstants::ModuleName).Get();
 
-		if (const FAssetRegistryState* StatePtr = AssetRegistry.GetAssetRegistryState())
-		{
-			const FText MetaDataCategoryName = NSLOCTEXT("ContentBrowser", "MetaDataCategoryName", "Meta-Data");
-			for (const auto& TagAndArrayPair : StatePtr->GetTagToAssetDatasMap())
+		const FText MetaDataCategoryName = NSLOCTEXT("ContentBrowser", "MetaDataCategoryName", "Meta-Data");
+		FString TagNameStr;
+		AssetRegistry.ReadLockEnumerateTagToAssetDatas(
+			[&PassesValueFilter, &PossibleSuggestions, &MetaDataCategoryName, &TagNameStr](FName TagName, const TArray<const FAssetData*>& Assets)
 			{
-				const FString TagNameStr = TagAndArrayPair.Key.ToString();
+				TagName.ToString(TagNameStr);
 				if (PassesValueFilter(TagNameStr))
 				{
 					PossibleSuggestions.Add(FAssetSearchBoxSuggestion{ TagNameStr, FText::FromString(TagNameStr), MetaDataCategoryName });
 				}
-			}
-		}
+			});
 	}
 
 	SuggestionHighlightText = FText::FromString(FilterValue);
diff --git a/Engine/Source/Editor/UnrealEd/Private/CookOnTheFlyServer.cpp b/Engine/Source/Editor/UnrealEd/Private/CookOnTheFlyServer.cpp
--- a/Engine/Source/Editor/UnrealEd/Private/CookOnTheFlyServer.cpp
+++ b/Engine/Source/Editor/UnrealEd/Private/CookOnTheFlyServer.cpp
@@ -5795,7 +5795,7 @@ void DumpAssetRegistryForCooker(IAssetRegistry* AssetRegistry)
 	PlatformFile.CreateDirectoryTree(*DumpDir);
 	TArray<FString> Pages;
 	TArray<FString> Arguments({ TEXT("ObjectPath"),TEXT("PackageName"),TEXT("Path"),TEXT("Class"),TEXT("Tag"), TEXT("DependencyDetails"), TEXT("PackageData"), TEXT("LegacyDependencies") });
-	AssetRegistry->GetAssetRegistryState()->Dump(Arguments, Pages, 10000 /* LinesPerPage */);
+	AssetRegistry->DumpState(Arguments, Pages, 10000 /* LinesPerPage */);
 	int PageIndex = 0;
 	TStringBuilder<256> FileName;
 	for (FString& PageText : Pages)
diff --git a/Engine/Source/Editor/UnrealEd/Private/Cooker/CookPackageData.cpp b/Engine/Source/Editor/UnrealEd/Private/Cooker/CookPackageData.cpp
--- a/Engine/Source/Editor/UnrealEd/Private/Cooker/CookPackageData.cpp
+++ b/Engine/Source/Editor/UnrealEd/Private/Cooker/CookPackageData.cpp
@@ -1027,7 +1027,7 @@ namespace Cook
 				return false;
 			}
 			UE::Cook::FPackageData& PackageData = PackageDatas.FindOrAddPackageData(PackageFName, FileName);
-			if (IAssetRegistry::Get()->GetAssetPackageData(PackageFName) && !PackageData.IsGenerated())
+			if (IAssetRegistry::Get()->GetAssetPackageDataCopy(PackageFName).IsSet() && !PackageData.IsGenerated())
 			{
 				UE_LOG(LogCook, Warning, TEXT("PackageSplitter specified a generated package that already exists in the workspace domain. Splitter=%s, Generated=%s."),
 					*this->GetSplitDataObjectName().ToString(), *PackageName);
diff --git a/Engine/Source/Editor/UnrealEd/Private/EditorDomain/EditorDomainUtils.cpp b/Engine/Source/Editor/UnrealEd/Private/EditorDomain/EditorDomainUtils.cpp
--- a/Engine/Source/Editor/UnrealEd/Private/EditorDomain/EditorDomainUtils.cpp
+++ b/Engine/Source/Editor/UnrealEd/Private/EditorDomain/EditorDomainUtils.cpp
@@ -61,7 +61,7 @@ EPackageDigestResult GetPackageDigest(IAssetRegistry& AssetRegistry, FName Packa
 		return EPackageDigestResult::WrongThread;
 	}
 	AssetRegistry.WaitForPackage(PackageName.ToString());
-	const FAssetPackageData* PackageData = AssetRegistry.GetAssetPackageData(PackageName);
+	TOptional<FAssetPackageData> PackageData = AssetRegistry.GetAssetPackageDataCopy(PackageName);
 	if (!PackageData)
 	{
 		return EPackageDigestResult::FileDoesNotExist;
diff --git a/Engine/Source/Runtime/AssetRegistry/Private/AssetRegistry.cpp b/Engine/Source/Runtime/AssetRegistry/Private/AssetRegistry.cpp
--- a/Engine/Source/Runtime/AssetRegistry/Private/AssetRegistry.cpp
+++ b/Engine/Source/Runtime/AssetRegistry/Private/AssetRegistry.cpp
@@ -17,12 +17,15 @@
 #include "Misc/ConfigCacheIni.h"
 #include "Misc/CoreDelegates.h"
 #include "Misc/FileHelper.h"
+#include "Misc/Optional.h"
 #include "Misc/Paths.h"
 #include "Misc/RedirectCollector.h"
 #include "Misc/ScopeExit.h"
+#include "Misc/ScopeRWLock.h"
 #include "PackageReader.h"
 #include "Serialization/ArrayReader.h"
 #include "Serialization/LargeMemoryReader.h"
+#include "Templates/UnrealTemplate.h"
 #include "UObject/ConstructorHelpers.h"
 #include "UObject/CoreRedirects.h"
 #include "UObject/MetaData.h"
@@ -36,8 +39,63 @@
 #include "IDirectoryWatcher.h"
 #endif // WITH_EDITOR
 
+/**
+ * ********** AssetRegistry threading model **********
+ * *** Functions and InterfaceLock ***
+ * All data(except events and RWLock) in the AssetRegistry is stored on the FAssetRegistryImpl GuardedData object.
+ * No data can be read on GuardedData unless the caller has entered the InterfaceLock.
+ * All data on FAssetRegistryImpl is private; this allows us to mark threading model with function prototypes.
+ * All functions on FAssetRegistryImpl are intended to be called only within a critical section.
+ * const functions require a ReadLock critical section; non-const require a WriteLock.
+ * The requirement that functions must be called only from within a critical section(and non-const only within a
+ * WriteLock) is not enforced technically; change authors need to carefully follow the synchronization model.
+
+ * *** Events, Callbacks, and Object Virtuals ***
+ * The AssetRegistry provides several Events(e.g.AssetAddedEvent) that can be subscribed to from arbitrary engine or
+ * licensee code, and some functions(e.g.EnumerateAssets) take a callback, and some functions call arbitrary
+ * UObject virtuals(e.g. new FAssetData(UObject*)). Some of this arbitrary code can call AssetRegistry functions of
+ * their own, and if they were called from within the lock that reentrancy would cause a deadlock when we tried
+ * to acquire the RWLock (RWLocks are not reenterable on the same thread). With some exceptions AssetRegistryImpl code
+ * is therefore not allowed to call callbacks, send events, or call UObject virtuals from inside a lock.
+
+ * FEventContext allows deferring events to a point in the top-level interface function outside the lock. The top-level
+ * function passes the EventContext in to the GuardedData functions, which add events on to it, and then it broadcasts
+ * the events outside the lock. FEventContext also handles deferring events to the Tick function executed from
+ * the GameThread, as we have a contract that events are only called from the game thread.
+
+ * Callbacks are handled on a case-by-case basis; each interface function handles queuing up the data for the callback
+ * functions and calling it outside the lock. The one exception is the ShouldSetManager function, which we call
+ * from inside the lock, since it is relatively well-behaved code as it is only used by UAssetManager and licensee
+ * subclasses of UAssetManager.
+
+ * UObject virtuals are handled on a case-by-case basis; the primary example is `new FAssetData(UObject*)`, which
+ * ProcessLoadedAssetsToUpdateCache takes care to call outside the lock and only on the game thread.
+ *
+ * *** Updating Caches - InheritanceContext ***
+ * The AssetRegistry has a cache for CodeGeneratorClasses and for an InheritanceMap of classes - native and blueprint.
+ * Updating these caches needs to be done within a writelock; for CodeGeneratorClasses we do this normally by marking
+ * all functions that need to update it as non-const. For InheritanceMap that would be overly pessimistic as several
+ * otherwise-const functions need to occasionally update the caches. For InheritanceMap we therefore have
+ * FClassInheritanceContext and FClassInheritanceBuffer. The top-level interface functions check whether the
+ * inheritance map will need to be updated during their execution, and if so they enter a write lock with the ability
+ * to update the members in the InheritanceContext. Otherwise they enter a readlock and the InheritanceBuffer will not
+ * be modified. All functions that use the cached data require the InheritanceContext to give them access, to ensure
+ * they are only using correctly updated cache data.
+ *
+ * *** Returning Internal Data ***
+ * All interface functions that return internal data return it by copy, or provide a ReadLockEnumerate function that
+ * calls a callback under the readlock, where the author of the callback has to ensure other AssetRegistry functions
+ * are not called.
+ */
+
 static FAssetRegistryConsoleCommands ConsoleCommands; // Registers its various console commands in the constructor
 
+namespace UE::AssetRegistry::Impl
+{
+	/** The max time to spend in UAssetRegistryImpl::Tick */
+	const float MaxSecondsPerFrame = 0.04;
+}
+
 /**
  * Implementation of IAssetRegistryInterface; forwards calls from the CoreUObject-accessible IAssetRegistryInterface into the AssetRegistry-accessible IAssetRegistry
  */
@@ -77,9 +135,6 @@ static FAutoConsoleVariableRef CVarLoadPremadeRegistryInEditor(
 	TEXT(""));
 #endif
 
-/** This will always read the ini, public version may return cache */
-static void InitializeSerializationOptionsFromIni(FAssetRegistrySerializationOptions& Options, const FString& PlatformIniName);
-
 static bool LoadAssetRegistry(const TCHAR* Path, const FAssetRegistryLoadOptions& Options, FAssetRegistryState& Out)
 {
 	check(Path);
@@ -236,8 +291,40 @@ UAssetRegistry::UAssetRegistry(const FObjectInitializer& ObjectInitializer)
 {
 }
 
+namespace UE::AssetRegistry::Impl
+{
+
+struct FInitializeContext
+{
+	FEventContext Events;
+	TArray<FString> RootContentPaths;
+	bool bRedirectorsNeedSubscribe = false;
+	bool bUpdateDiskCacheAfterLoad = false;
+};
+
+}
+
 UAssetRegistryImpl::UAssetRegistryImpl(const FObjectInitializer& ObjectInitializer)
 	: Super(ObjectInitializer)
+{
+	UE::AssetRegistry::Impl::FInitializeContext Context;
+
+	{
+		FWriteScopeLock InterfaceScopeLock(InterfaceLock);
+		GuardedData.Initialize(Context);
+		InitializeEvents(Context);
+	}
+	Broadcast(Context.Events);
+}
+
+namespace UE::AssetRegistry
+{
+
+FAssetRegistryImpl::FAssetRegistryImpl()
+{
+}
+
+void FAssetRegistryImpl::Initialize(Impl::FInitializeContext& Context)
 {
 	LLM_SCOPE(ELLMTag::AssetRegistry);
 
@@ -250,18 +337,14 @@ UAssetRegistryImpl::UAssetRegistryImpl(const FObjectInitializer& ObjectInitializ
 	AmortizeStartTime = 0;
 	TotalAmortizeTime = 0;
 
-	MaxSecondsPerFrame = 0.04;
-
 	// By default update the disk cache once on asset load, to incorporate changes made in PostLoad. This only happens in editor builds
 	bUpdateDiskCacheAfterLoad = true;
 
 	bIsTempCachingAlwaysEnabled = ASSETREGISTRY_CACHE_ALWAYS_ENABLED;
 	bIsTempCachingEnabled = bIsTempCachingAlwaysEnabled;
-	bIsTempCachingUpToDate = false;
-	
-	// The initial value doesn't matter since caching has not yet been computed
-	TempCachingRegisteredClassesVersionNumber = 0;
-	ClassGeneratorNamesRegisteredClassesVersionNumber = 0;
+	TempCachedInheritanceBuffer.bDirty = true;
+
+	ClassGeneratorNamesRegisteredClassesVersionNumber = MAX_uint64;
 
 	// By default do not double check mount points are still valid when gathering new assets
 	bVerifyMountPointAfterGather = false;
@@ -276,9 +359,12 @@ UAssetRegistryImpl::UAssetRegistryImpl(const FObjectInitializer& ObjectInitializ
 
 	// Collect all code generator classes (currently BlueprintCore-derived ones)
 	CollectCodeGeneratorClasses();
+#if WITH_ENGINE && WITH_EDITOR
+	Utils::PopulateSkipClasses(SkipUncookedClasses, SkipCookedClasses);
+#endif
 
 	// Read default serialization options
-	InitializeSerializationOptionsFromIni(SerializationOptions, FString());
+	Utils::InitializeSerializationOptionsFromIni(SerializationOptions, FString());
 
 	// If in the editor or cookcommandlet, we start the GlobalGatherer now
 	// In the game or other commandlets, we do not construct it until project or commandlet code calls SearchAllAssets or ScanPathsSynchronous
@@ -290,13 +376,12 @@ UAssetRegistryImpl::UAssetRegistryImpl(const FObjectInitializer& ObjectInitializ
 		bSearchAllAssetsAtStart = bCommandlineAllAssetsAtStart;
 	}
 #endif
-	IPluginManager& PluginManager = IPluginManager::Get();
 	if (bSearchAllAssetsAtStart)
 	{
 		ConstructGatherer();
 		if (!GlobalGatherer->IsSynchronous())
 		{
-			SearchAllAssetsInitialAsync();
+			SearchAllAssetsInitialAsync(Context.Events);
 		}
 	}
 #if ASSETREGISTRY_ENABLE_PREMADE_REGISTRY_IN_EDITOR 
@@ -306,7 +391,7 @@ UAssetRegistryImpl::UAssetRegistryImpl(const FObjectInitializer& ObjectInitializ
 		if (LoadAssetRegistry(*(FPaths::ProjectDir() / TEXT("AssetRegistry.bin")), LoadOptions, State))
 		{
 			UE_LOG(LogAssetRegistry, Log, TEXT("Loaded premade asset registry"));
-			CachePathsFromState(State);
+			CachePathsFromState(Context.Events, State);
 		}
 		else
 		{
@@ -318,11 +403,11 @@ UAssetRegistryImpl::UAssetRegistryImpl(const FObjectInitializer& ObjectInitializ
 		if (SerializationOptions.bSerializeAssetRegistry &&
 			GCookedAssetRegistryPreloader.Consume(/* Out */ State))
 		{
-			CachePathsFromState(State);
+			CachePathsFromState(Context.Events, State);
 		}
 #endif // ASSETREGISTRY_ENABLE_PREMADE_REGISTRY_IN_EDITOR 
 
-		TArray<TSharedRef<IPlugin>> ContentPlugins = PluginManager.GetEnabledPluginsWithContent();
+		TArray<TSharedRef<IPlugin>> ContentPlugins = IPluginManager::Get().GetEnabledPluginsWithContent();
 		for (TSharedRef<IPlugin> ContentPlugin : ContentPlugins)
 		{
 			if (ContentPlugin->CanContainContent())
@@ -336,7 +421,7 @@ UAssetRegistryImpl::UAssetRegistryImpl(const FObjectInitializer& ObjectInitializ
 					PluginState.Load(SerializedAssetData);
 
 					State.InitializeFromExisting(PluginState, SerializationOptions, FAssetRegistryState::EInitializationMode::Append);
-					CachePathsFromState(PluginState);
+					CachePathsFromState(Context.Events, PluginState);
 				}
 			}
 		}
@@ -347,6 +432,42 @@ UAssetRegistryImpl::UAssetRegistryImpl(const FObjectInitializer& ObjectInitializ
 	// Report startup time. This does not include DirectoryWatcher startup time.
 	UE_LOG(LogAssetRegistry, Log, TEXT("FAssetRegistry took %0.4f seconds to start up"), FPlatformTime::Seconds() - StartupStartTime);
 
+#if WITH_EDITOR
+	if (GConfig)
+	{
+		GConfig->GetBool(TEXT("AssetRegistry"), TEXT("bUpdateDiskCacheAfterLoad"), bUpdateDiskCacheAfterLoad, GEngineIni);
+	}
+#endif
+	Context.bUpdateDiskCacheAfterLoad = bUpdateDiskCacheAfterLoad;
+
+	// Content roots always exist; add them as paths
+	FPackageName::QueryRootContentPaths(Context.RootContentPaths);
+	for (const FString& AssetPath : Context.RootContentPaths)
+	{
+		AddPath(Context.Events, AssetPath);
+	}
+
+	InitRedirectors(Context.Events, Context.bRedirectorsNeedSubscribe);
+}
+
+}
+
+void UAssetRegistryImpl::InitializeEvents(UE::AssetRegistry::Impl::FInitializeContext& Context)
+{
+	if (HasAnyFlags(RF_ClassDefaultObject))
+	{
+		check(UE::AssetRegistry::Private::IAssetRegistrySingleton::Singleton == nullptr && IAssetRegistryInterface::Default == nullptr);
+		UE::AssetRegistry::Private::IAssetRegistrySingleton::Singleton = this;
+		IAssetRegistryInterface::Default = &GAssetRegistryInterface;
+	}
+
+	if (Context.bRedirectorsNeedSubscribe)
+	{
+		FCoreDelegates::FResolvePackageNameDelegate PackageResolveDelegate;
+		PackageResolveDelegate.BindUObject(this, &UAssetRegistryImpl::OnResolveRedirect);
+		FCoreDelegates::PackageNameResolvers.Add(PackageResolveDelegate);
+	}
+
 #if WITH_EDITOR
 	// In-game doesn't listen for directory changes
 	if (GIsEditor)
@@ -356,9 +477,7 @@ UAssetRegistryImpl::UAssetRegistryImpl(const FObjectInitializer& ObjectInitializ
 
 		if (DirectoryWatcher)
 		{
-			TArray<FString> RootContentPaths;
-			FPackageName::QueryRootContentPaths(RootContentPaths);
-			for (TArray<FString>::TConstIterator RootPathIt(RootContentPaths); RootPathIt; ++RootPathIt)
+			for (TArray<FString>::TConstIterator RootPathIt(Context.RootContentPaths); RootPathIt; ++RootPathIt)
 			{
 				const FString& RootPath = *RootPathIt;
 				const FString& ContentFolder = FPackageName::LongPackageNameToFilename(RootPath);
@@ -380,47 +499,21 @@ UAssetRegistryImpl::UAssetRegistryImpl(const FObjectInitializer& ObjectInitializ
 		}
 	}
 
-
-
-	if (GConfig)
-	{
-		GConfig->GetBool(TEXT("AssetRegistry"), TEXT("bUpdateDiskCacheAfterLoad"), bUpdateDiskCacheAfterLoad, GEngineIni);
-	}
-
-	if (bUpdateDiskCacheAfterLoad)
+	if (Context.bUpdateDiskCacheAfterLoad)
 	{
 		FCoreUObjectDelegates::OnAssetLoaded.AddUObject(this, &UAssetRegistryImpl::OnAssetLoaded);
 	}
 #endif // WITH_EDITOR
 
-	// Content roots always exist
-	{
-		TArray<FString> RootContentPaths;
-		FPackageName::QueryRootContentPaths(RootContentPaths);
-
-		for (const FString& AssetPath : RootContentPaths)
-		{
-			AddPath(AssetPath);
-		}
-	}
-
 	// Listen for new content paths being added or removed at runtime.  These are usually plugin-specific asset paths that
 	// will be loaded a bit later on.
 	FPackageName::OnContentPathMounted().AddUObject(this, &UAssetRegistryImpl::OnContentPathMounted);
 	FPackageName::OnContentPathDismounted().AddUObject(this, &UAssetRegistryImpl::OnContentPathDismounted);
 
 	// If we were called before engine has fully initialized, refresh classes on initialize. If not this won't do anything as it already happened
-	FCoreDelegates::OnPostEngineInit.AddUObject(this, &UAssetRegistryImpl::RefreshNativeClasses);
-
-	InitRedirectors();
-
-	if (HasAnyFlags(RF_ClassDefaultObject))
-	{
-		check(UE::AssetRegistry::Private::IAssetRegistrySingleton::Singleton == nullptr && IAssetRegistryInterface::Default == nullptr);
-		UE::AssetRegistry::Private::IAssetRegistrySingleton::Singleton = this;
-		IAssetRegistryInterface::Default = &GAssetRegistryInterface;
-	}
+	FCoreDelegates::OnPostEngineInit.AddUObject(this, &UAssetRegistryImpl::OnRefreshNativeClasses);
 
+	IPluginManager& PluginManager = IPluginManager::Get();
 	ELoadingPhase::Type LoadingPhase = PluginManager.GetLastCompletedLoadingPhase();
 	if (LoadingPhase == ELoadingPhase::None || LoadingPhase < ELoadingPhase::PostEngineInit)
 	{
@@ -432,7 +525,16 @@ UAssetRegistryImpl::UAssetRegistryImpl(FVTableHelper& Helper)
 {
 }
 
-bool UAssetRegistryImpl::ResolveRedirect(const FString& InPackageName, FString& OutPackageName)
+bool UAssetRegistryImpl::OnResolveRedirect(const FString& InPackageName, FString& OutPackageName)
+{
+	FReadScopeLock InterfaceScopeLock(InterfaceLock);
+	return GuardedData.ResolveRedirect(InPackageName, OutPackageName);
+}
+
+namespace UE::AssetRegistry
+{
+
+bool FAssetRegistryImpl::ResolveRedirect(const FString& InPackageName, FString& OutPackageName) const
 {
 	int32 DotIndex = InPackageName.Find(TEXT("."), ESearchCase::CaseSensitive);
 
@@ -456,8 +558,10 @@ bool UAssetRegistryImpl::ResolveRedirect(const FString& InPackageName, FString&
 	return false;
 }
 
-void UAssetRegistryImpl::InitRedirectors()
+void FAssetRegistryImpl::InitRedirectors(Impl::FEventContext& EventContext, bool& bOutRedirectorsNeedSubscribe)
 {
+	bOutRedirectorsNeedSubscribe = false;
+
 	// plugins can't initialize redirectors in the editor, it will mess up the saving of content.
 	if ( GIsEditor )
 	{
@@ -465,7 +569,7 @@ void UAssetRegistryImpl::InitRedirectors()
 	}
 
 	TArray<TSharedRef<IPlugin>> EnabledPlugins = IPluginManager::Get().GetEnabledPlugins();
-	for (TSharedRef<IPlugin> Plugin : EnabledPlugins)
+	for (const TSharedRef<IPlugin>& Plugin : EnabledPlugins)
 	{
 		FString PluginConfigFilename = FString::Printf(TEXT("%s%s/%s.ini"), *FPaths::GeneratedConfigDir(), ANSI_TO_TCHAR(FPlatformProperties::PlatformName()), *Plugin->GetName() );
 		
@@ -489,52 +593,47 @@ void UAssetRegistryImpl::InitRedirectors()
 			FString RootPackageName = FString::Printf(TEXT("/%s/"), *Plugin->GetName());
 			PathsToSearch.Add(RootPackageName);
 
-			ScanPathsSynchronous(PathsToSearch, false /* bForceRescan */, false /* bIgnoreBlacklistScanFilters */);
+			Impl::FScanPathContext Context(EventContext, PathsToSearch, TArray<FString>());
+			ScanPathsSynchronous(Context);
 		}
-		
-		FName PluginPackageName = FName(*FString::Printf(TEXT("/%s/"), *Plugin->GetName()));
-		TArray<FAssetData> AssetList;
-		GetAssetsByPath(PluginPackageName, AssetList, true, false);
 
-		for (const FAssetData& Asset : AssetList)
-		{
-			FString NewPackageNameString = Asset.PackageName.ToString();
-			FString RootPackageName = FString::Printf(TEXT("/%s/"), *Plugin->GetName());
-
-			FString OriginalPackageNameString = NewPackageNameString.Replace(*RootPackageName, TEXT("/Game/"));
-
-
-			PackageRedirects.Add(FAssetRegistryPackageRedirect(OriginalPackageNameString, NewPackageNameString));
-
-		}
+		FName PluginPackageName = FName(*FString::Printf(TEXT("/%s/"), *Plugin->GetName()));
+		EnumerateAssetsByPathNoTags(PluginPackageName,
+			[&Plugin, this](const FAssetData& PartialAssetData)
+			{
+				FString NewPackageNameString = PartialAssetData.PackageName.ToString();
+				FString RootPackageName = FString::Printf(TEXT("/%s/"), *Plugin->GetName());
+				FString OriginalPackageNameString = NewPackageNameString.Replace(*RootPackageName, TEXT("/Game/"));
 
+				PackageRedirects.Add(FAssetRegistryPackageRedirect(OriginalPackageNameString, NewPackageNameString));
+				return true;
+			}, true, false);
 
-		FCoreDelegates::FResolvePackageNameDelegate PackageResolveDelegate;
-		PackageResolveDelegate.BindUObject(this, &UAssetRegistryImpl::ResolveRedirect);
-		FCoreDelegates::PackageNameResolvers.Add(PackageResolveDelegate);
+		bOutRedirectorsNeedSubscribe = true;
 	}
 }
 
-void UAssetRegistryImpl::RegisterReadOfScriptPackages()
+}
+
+void UAssetRegistryImpl::OnPluginLoadingPhaseComplete(ELoadingPhase::Type LoadingPhase, bool bPhaseSuccessful)
 {
-	if (!GlobalGatherer.IsValid())
+	if (LoadingPhase != ELoadingPhase::PostEngineInit)
 	{
 		return;
 	}
-	IPluginManager& PluginManager = IPluginManager::Get();
-	ELoadingPhase::Type LoadingPhase = PluginManager.GetLastCompletedLoadingPhase();
-	if (LoadingPhase != ELoadingPhase::None && LoadingPhase >= ELoadingPhase::PostEngineInit)
 	{
-		ReadScriptPackages();
+		FWriteScopeLock InterfaceScopeLock(InterfaceLock);
+		GuardedData.OnPostEngineInit(bPhaseSuccessful);
 	}
+
+	IPluginManager::Get().OnLoadingPhaseComplete().RemoveAll(this);
 }
 
-void UAssetRegistryImpl::OnPluginLoadingPhaseComplete(ELoadingPhase::Type LoadingPhase, bool bPhaseSuccessful)
+namespace UE::AssetRegistry
+{
+
+void FAssetRegistryImpl::OnPostEngineInit(bool bPhaseSuccessful)
 {
-	if (LoadingPhase != ELoadingPhase::PostEngineInit)
-	{
-		return;
-	}
 	// If we have constructed the GlobalGatherer then we need to readscriptpackages,
 	// otherwise we will read them when constructing the gatherer.
 	if (GlobalGatherer.IsValid())
@@ -544,55 +643,69 @@ void UAssetRegistryImpl::OnPluginLoadingPhaseComplete(ELoadingPhase::Type Loadin
 
 	// Reparse the skip classes the next time ShouldSkipAsset is called, since available classes
 	// for the search over all classes may have changed
-#if WITH_EDITOR
-	UE::AssetRegistry::FFiltering::MarkDirty();
+#if WITH_ENGINE && WITH_EDITOR
+	// If we ever need to update the Filtering list outside of the game thread, we will need to defer the update 
+	// of the Filtering namespace to the tick function; UE::AssetRegistry::Filtering can only be used in game thread
+	check(IsInGameThread());
+	Utils::PopulateSkipClasses(SkipUncookedClasses, SkipCookedClasses);
+	UE::AssetRegistry::FFiltering::SetSkipClasses(SkipUncookedClasses, SkipCookedClasses);
 #endif
-	IPluginManager::Get().OnLoadingPhaseComplete().RemoveAll(this);
 }
 
-void UAssetRegistryImpl::ReadScriptPackages()
+void FAssetRegistryImpl::ReadScriptPackages()
 {
-	if (GlobalGatherer.IsValid())
+	GlobalGatherer->SetInitialPluginsLoaded();
+	if (GlobalGatherer->IsGatheringDependencies())
 	{
-		GlobalGatherer->SetInitialPluginsLoaded();
-		if (GlobalGatherer->IsGatheringDependencies())
+		// Now that all scripts have been loaded, we need to create AssetPackageDatas for every script
+		// This is also done whenever scripts are referenced in our gather of existing packages,
+		// but we need to complete it for all scripts that were referenced but not yet loaded for packages
+		// that we already gathered
+		for (TObjectIterator<UPackage> It; It; ++It)
 		{
-			// Now that all scripts have been loaded, we need to create AssetPackageDatas for every script
-			// This is also done whenever scripts are referenced in our gather of existing packages,
-			// but we need to complete it for all scripts that were referenced but not yet loaded for packages
-			// that we already gathered
-			for (TObjectIterator<UPackage> It; It; ++It)
+			UPackage* Package = *It;
+			if (Package)
 			{
-				UPackage* Package = *It;
-				if (Package)
+				if (Package && FPackageName::IsScriptPackage(Package->GetName()))
 				{
-					if (Package && FPackageName::IsScriptPackage(Package->GetName()))
-					{
-						FAssetPackageData* ScriptPackageData = State.CreateOrGetAssetPackageData(Package->GetFName());
-						// Get the guid off the script package, it is updated when script is changed so we need to refresh it every run
-						PRAGMA_DISABLE_DEPRECATION_WARNINGS
-						ScriptPackageData->PackageGuid = Package->GetGuid();
-						PRAGMA_ENABLE_DEPRECATION_WARNINGS
-					}
+					FAssetPackageData* ScriptPackageData = State.CreateOrGetAssetPackageData(Package->GetFName());
+					// Get the guid off the script package, it is updated when script is changed so we need to refresh it every run
+					PRAGMA_DISABLE_DEPRECATION_WARNINGS
+					ScriptPackageData->PackageGuid = Package->GetGuid();
+					PRAGMA_ENABLE_DEPRECATION_WARNINGS
 				}
 			}
 		}
 	}
 }
 
+}
+
 void UAssetRegistryImpl::InitializeSerializationOptions(FAssetRegistrySerializationOptions& Options, const FString& PlatformIniName) const
 {
 	if (PlatformIniName.IsEmpty())
 	{
+		FReadScopeLock InterfaceScopeLock(InterfaceLock);
 		// Use options we already loaded, the first pass for this happens at object creation time so this is always valid when queried externally
-		Options = SerializationOptions;
+		GuardedData.CopySerializationOptions(Options);
 	}
 	else
 	{
-		InitializeSerializationOptionsFromIni(Options, PlatformIniName);
+		UE::AssetRegistry::Utils::InitializeSerializationOptionsFromIni(Options, PlatformIniName);
 	}
 }
 
+namespace UE::AssetRegistry
+{
+
+void FAssetRegistryImpl::CopySerializationOptions(FAssetRegistrySerializationOptions& OutOptions) const
+{
+	OutOptions = SerializationOptions;
+}
+
+namespace Utils
+{
+
 static TSet<FName> MakeNameSet(const TArray<FString>& Strings)
 {
 	TSet<FName> Out;
@@ -605,7 +718,7 @@ static TSet<FName> MakeNameSet(const TArray<FString>& Strings)
 	return Out;
 }
 
-static void InitializeSerializationOptionsFromIni(FAssetRegistrySerializationOptions& Options, const FString& PlatformIniName)
+void InitializeSerializationOptionsFromIni(FAssetRegistrySerializationOptions& Options, const FString& PlatformIniName)
 {
 	FConfigFile* EngineIni = nullptr;
 #if WITH_EDITOR
@@ -715,7 +828,9 @@ static void InitializeSerializationOptionsFromIni(FAssetRegistrySerializationOpt
 	}
 }
 
-void UAssetRegistryImpl::CollectCodeGeneratorClasses() const
+}
+
+void FAssetRegistryImpl::CollectCodeGeneratorClasses()
 {
 	// Only refresh the list if our registered classes have changed
 	if (ClassGeneratorNamesRegisteredClassesVersionNumber != GetRegisteredClassesVersionNumber())
@@ -737,23 +852,31 @@ void UAssetRegistryImpl::CollectCodeGeneratorClasses() const
 	}
 }
 
-void UAssetRegistryImpl::RefreshNativeClasses()
+}
+
+void UAssetRegistryImpl::OnRefreshNativeClasses()
+{
+	FWriteScopeLock InterfaceScopeLock(InterfaceLock);
+	GuardedData.RefreshNativeClasses();
+}
+
+namespace UE::AssetRegistry
+{
+
+void FAssetRegistryImpl::RefreshNativeClasses()
 {
 	// Native classes have changed so reinitialize code generator and serialization options
 	CollectCodeGeneratorClasses();
 
 	// Read default serialization options
-	InitializeSerializationOptionsFromIni(SerializationOptions, FString());
+	Utils::InitializeSerializationOptionsFromIni(SerializationOptions, FString());
+}
+
 }
 
 UAssetRegistryImpl::~UAssetRegistryImpl()
 {
-	if (HasAnyFlags(RF_ClassDefaultObject))
-	{
-		check(UE::AssetRegistry::Private::IAssetRegistrySingleton::Singleton == this && IAssetRegistryInterface::Default == &GAssetRegistryInterface);
-		UE::AssetRegistry::Private::IAssetRegistrySingleton::Singleton = nullptr;
-		IAssetRegistryInterface::Default = nullptr;
-	}
+	FWriteScopeLock InterfaceScopeLock(InterfaceLock);
 
 	// Stop listening for content mount point events
 	FPackageName::OnContentPathMounted().RemoveAll(this);
@@ -784,12 +907,19 @@ UAssetRegistryImpl::~UAssetRegistryImpl()
 		}
 	}
 
-	if (bUpdateDiskCacheAfterLoad)
+	if (GuardedData.IsUpdateDiskCacheAfterLoad())
 	{
 		FCoreUObjectDelegates::OnAssetLoaded.RemoveAll(this);
 	}
 #endif // WITH_EDITOR
 
+	if (HasAnyFlags(RF_ClassDefaultObject))
+	{
+		check(UE::AssetRegistry::Private::IAssetRegistrySingleton::Singleton == this && IAssetRegistryInterface::Default == &GAssetRegistryInterface);
+		UE::AssetRegistry::Private::IAssetRegistrySingleton::Singleton = nullptr;
+		IAssetRegistryInterface::Default = nullptr;
+	}
+
 	// Clear all listeners
 	PathAddedEvent.Clear();
 	PathRemovedEvent.Clear();
@@ -809,7 +939,10 @@ UAssetRegistryImpl& UAssetRegistryImpl::Get()
 	return static_cast<UAssetRegistryImpl&>(*UE::AssetRegistry::Private::IAssetRegistrySingleton::Singleton);
 }
 
-void UAssetRegistryImpl::ConstructGatherer()
+namespace UE::AssetRegistry
+{
+
+void FAssetRegistryImpl::ConstructGatherer()
 {
 	if (GlobalGatherer.IsValid())
 	{
@@ -826,18 +959,45 @@ void UAssetRegistryImpl::ConstructGatherer()
 
 	bool bIsSynchronous = IsRunningGame();
 	GlobalGatherer = MakeUnique<FAssetDataGatherer>(BlacklistPaths, BlacklistContentSubPaths, bIsSynchronous);
-	RegisterReadOfScriptPackages();
+
+	// Read script packages if all initial plugins have been loaded, otherwise do nothing; we wait for the callback.
+	ELoadingPhase::Type LoadingPhase = IPluginManager::Get().GetLastCompletedLoadingPhase();
+	if (LoadingPhase != ELoadingPhase::None && LoadingPhase >= ELoadingPhase::PostEngineInit)
+	{
+		ReadScriptPackages();
+	}
 }
 
-void UAssetRegistryImpl::SearchAllAssetsInitialAsync()
+void FAssetRegistryImpl::SearchAllAssetsInitialAsync(Impl::FEventContext& EventContext)
 {
 	bInitialSearchStarted = true;
 	bInitialSearchCompleted = false;
 	FullSearchStartTime = FPlatformTime::Seconds();
-	SearchAllAssets(false);
+	SearchAllAssets(EventContext, false /* bSynchronousSearch */);
+}
+
 }
 
 void UAssetRegistryImpl::SearchAllAssets(bool bSynchronousSearch)
+{
+	UE::AssetRegistry::Impl::FEventContext EventContext;
+	{
+		FWriteScopeLock InterfaceScopeLock(InterfaceLock);
+		GuardedData.SearchAllAssets(EventContext, bSynchronousSearch);
+	}
+#if WITH_EDITOR
+	if (bSynchronousSearch)
+	{
+		ProcessLoadedAssetsToUpdateCache(EventContext, -1., true /* bIdle */);
+	}
+#endif
+	Broadcast(EventContext);
+}
+
+namespace UE::AssetRegistry
+{
+
+void FAssetRegistryImpl::SearchAllAssets(Impl::FEventContext& EventContext, bool bSynchronousSearch)
 {
 	ConstructGatherer();
 	FAssetDataGatherer& Gatherer = *GlobalGatherer;
@@ -864,7 +1024,9 @@ void UAssetRegistryImpl::SearchAllAssets(bool bSynchronousSearch)
 	if (bSynchronousSearch)
 	{
 		Gatherer.WaitForIdle();
-		TickGatherer(-1.f);
+		bool bUnusedIdle;
+		bool bUnusedInterrupted;
+		TickGatherer(EventContext, -1., bUnusedIdle, bUnusedInterrupted);
 #if WITH_EDITOR
 		if (!bInitialSearchStarted)
 		{
@@ -880,11 +1042,27 @@ void UAssetRegistryImpl::SearchAllAssets(bool bSynchronousSearch)
 	}
 }
 
+}
+
 void UAssetRegistryImpl::WaitForCompletion()
 {
-	while (IsLoadingAssets())
+	for (;;)
 	{
-		Tick(-1.0f);
+		UE::AssetRegistry::Impl::FEventContext EventContext;
+		bool bIsIdle = false;
+		{
+			FWriteScopeLock InterfaceScopeLock(InterfaceLock);
+			if (!GuardedData.IsLoadingAssets())
+			{
+				break;
+			}
+			bool bUnusedInterrupted;
+			GuardedData.TickGatherer(EventContext, -1., bIsIdle, bUnusedInterrupted);
+		}
+#if WITH_EDITOR
+		ProcessLoadedAssetsToUpdateCache(EventContext, -1., bIsIdle);
+#endif
+		Broadcast(EventContext);
 
 		FThreadHeartBeat::Get().HeartBeat();
 		FPlatformProcess::SleepNoStats(0.0001f);
@@ -893,22 +1071,31 @@ void UAssetRegistryImpl::WaitForCompletion()
 
 void UAssetRegistryImpl::WaitForPackage(const FString& PackageName)
 {
-	check(IsInGameThread());
-	if (!IsLoadingAssets() || !GlobalGatherer.IsValid())
-	{
-		return;
-	}
-
-	FString LocalPath;
-	if (!FPackageName::TryConvertLongPackageNameToFilename(PackageName, LocalPath))
+	UE::AssetRegistry::Impl::FEventContext EventContext;
 	{
-		return;
+		FWriteScopeLock InterfaceScopeLock(InterfaceLock);
+		if (GuardedData.IsLoadingAssets())
+		{
+			FString LocalPath;
+			if (FPackageName::TryConvertLongPackageNameToFilename(PackageName, LocalPath))
+			{
+				GuardedData.TickGatherPackage(EventContext, PackageName, LocalPath);
+			}
+		}
 	}
-	GlobalGatherer->WaitOnPath(LocalPath);
-	TickGatherPackage(PackageName);
+	Broadcast(EventContext);
 }
 
 bool UAssetRegistryImpl::HasAssets(const FName PackagePath, const bool bRecursive) const
+{
+	FReadScopeLock InterfaceScopeLock(InterfaceLock);
+	return GuardedData.HasAssets(PackagePath, bRecursive);
+}
+
+namespace UE::AssetRegistry
+{
+
+bool FAssetRegistryImpl::HasAssets(const FName PackagePath, const bool bRecursive) const
 {
 	bool bHasAssets = State.HasAssets(PackagePath, true /*bARFiltering*/);
 
@@ -924,6 +1111,8 @@ bool UAssetRegistryImpl::HasAssets(const FName PackagePath, const bool bRecursiv
 	return bHasAssets;
 }
 
+}
+
 bool UAssetRegistryImpl::GetAssetsByPackageName(FName PackageName, TArray<FAssetData>& OutAssetData, bool bIncludeOnlyOnDiskAssets) const
 {
 	FARFilter Filter;
@@ -941,6 +1130,45 @@ bool UAssetRegistryImpl::GetAssetsByPath(FName PackagePath, TArray<FAssetData>&
 	return GetAssets(Filter, OutAssetData);
 }
 
+namespace UE::AssetRegistry
+{
+
+void FAssetRegistryImpl::EnumerateAssetsByPathNoTags(FName PackagePath,
+	TFunctionRef<bool(const FAssetData&)> Callback, bool bRecursive, bool bIncludeOnlyOnDiskAssets) const
+{
+	if (PackagePath.IsNone())
+	{
+		return;
+	}
+	FARFilter Filter;
+	Filter.bRecursivePaths = bRecursive;
+	Filter.PackagePaths.Add(PackagePath);
+	Filter.bIncludeOnlyOnDiskAssets = bIncludeOnlyOnDiskAssets;
+
+	// CompileFilter takes an inheritance context, but only to handle filters with recursive classes, which we are not using here
+	UE::AssetRegistry::Impl::FClassInheritanceContext EmptyInheritanceContext;
+	FARCompiledFilter CompiledFilter;
+	CompileFilter(EmptyInheritanceContext, Filter, CompiledFilter);
+
+	TSet<FName> PackagesToSkip;
+	if (!bIncludeOnlyOnDiskAssets)
+	{
+		bool bStopIteration;
+		Utils::EnumerateMemoryAssetsHelper(CompiledFilter, PackagesToSkip, bStopIteration,
+			[&Callback](const UObject* Object, FAssetData&& PartialAssetData)
+			{
+				return Callback(PartialAssetData);
+			});
+		if (bStopIteration)
+		{
+			return;
+		}
+	}
+	EnumerateDiskAssets(CompiledFilter, PackagesToSkip, Callback);
+}
+
+}
+
 bool UAssetRegistryImpl::GetAssetsByClass(FName ClassName, TArray<FAssetData>& OutAssetData, bool bSearchSubClasses) const
 {
 	FARFilter Filter;
@@ -971,17 +1199,36 @@ bool UAssetRegistryImpl::GetAssetsByTagValues(const TMultiMap<FName, FString>& A
 
 bool UAssetRegistryImpl::GetAssets(const FARFilter& InFilter, TArray<FAssetData>& OutAssetData) const
 {
-	const double GetAssetsStartTime = FPlatformTime::Seconds();
+	using namespace UE::AssetRegistry::Utils;
 
-	const bool bResult = EnumerateAssets(InFilter, [&OutAssetData](const FAssetData& AssetData)
+	FARCompiledFilter CompiledFilter;
+	CompileFilter(InFilter, CompiledFilter);
+	if (CompiledFilter.IsEmpty() || !IsFilterValid(CompiledFilter))
 	{
-		OutAssetData.Emplace(AssetData);
-		return true;
-	});
+		return false;
+	}
 
-	UE_LOG(LogAssetRegistry, Verbose, TEXT("GetAssets completed in %0.4f seconds"), FPlatformTime::Seconds() - GetAssetsStartTime);
+	TSet<FName> PackagesToSkip;
+	if (!InFilter.bIncludeOnlyOnDiskAssets)
+	{
+		bool bStopIterationUnused;
+		EnumerateMemoryAssets(CompiledFilter, PackagesToSkip, bStopIterationUnused,
+			[&OutAssetData](FAssetData&& AssetData)
+			{
+				OutAssetData.Add(MoveTemp(AssetData));
+				return true;
+			});
+	}
 
-	return bResult;
+	{
+		FReadScopeLock InterfaceScopeLock(InterfaceLock);
+		GuardedData.EnumerateDiskAssets(CompiledFilter, PackagesToSkip, [&OutAssetData](const FAssetData& AssetData)
+			{
+				OutAssetData.Emplace(AssetData);
+				return true;
+			});
+	}
+	return true;
 }
 
 bool UAssetRegistryImpl::EnumerateAssets(const FARFilter& InFilter, TFunctionRef<bool(const FAssetData&)> Callback) const
@@ -993,157 +1240,220 @@ bool UAssetRegistryImpl::EnumerateAssets(const FARFilter& InFilter, TFunctionRef
 
 bool UAssetRegistryImpl::EnumerateAssets(const FARCompiledFilter& InFilter, TFunctionRef<bool(const FAssetData&)> Callback) const
 {
+	using namespace UE::AssetRegistry::Utils;
+
 	// Verify filter input. If all assets are needed, use EnumerateAllAssets() instead.
-	if (InFilter.IsEmpty())
+	if (InFilter.IsEmpty() || !IsFilterValid(InFilter))
 	{
 		return false;
 	}
 
-	if (!FAssetRegistryState::IsFilterValid(InFilter))
+	TSet<FName> PackagesToSkip;
+	if (!InFilter.bIncludeOnlyOnDiskAssets)
 	{
-		return false;
+		bool bStopIteration;
+		EnumerateMemoryAssets(InFilter, PackagesToSkip, bStopIteration,
+			[&Callback](FAssetData&& AssetData)
+			{
+				return Callback(AssetData);
+			});
+		if (bStopIteration)
+		{
+			return true;
+		}
 	}
 
-	// Start with in memory assets
-	TSet<FName> PackagesToSkip = CachedEmptyPackages;
-	if (!InFilter.bIncludeOnlyOnDiskAssets)
+	TArray<FAssetData, TInlineAllocator<128>> FoundAssets;
 	{
-		// Skip assets that were loaded for diffing
-		const uint32 FilterWithoutPackageFlags = InFilter.WithoutPackageFlags | PKG_ForDiffing;
-		const uint32 FilterWithPackageFlags = InFilter.WithPackageFlags;
-
-		auto FilterInMemoryObjectLambda = [&](const UObject* Obj, bool& OutContinue)
-		{
-			if (Obj->IsAsset())
+		FReadScopeLock InterfaceScopeLock(InterfaceLock);
+		GuardedData.EnumerateDiskAssets(InFilter, PackagesToSkip, [&FoundAssets](const FAssetData& AssetData)
 			{
-				// Skip assets that are currently loading
-				if (Obj->HasAnyFlags(RF_NeedLoad))
-				{
-					return;
-				}
+				FoundAssets.Emplace(AssetData);
+				return true;
+			});
+	}
+	for (const FAssetData& AssetData : FoundAssets)
+	{
+		if (!Callback(AssetData))
+		{
+			break;
+		}
+	}
+	return true;
+}
 
-				UPackage* InMemoryPackage = Obj->GetOutermost();
+namespace UE::AssetRegistry
+{
 
-				// Skip assets with any of the specified 'without' package flags 
-				if (InMemoryPackage->HasAnyPackageFlags(FilterWithoutPackageFlags))
-				{
-					return;
-				}
+namespace Utils
+{
 
-				// Skip assets without any the specified 'with' packages flags
-				if (!InMemoryPackage->HasAllPackagesFlags(FilterWithPackageFlags))
-				{
-					return;
-				}
+void EnumerateMemoryAssetsHelper(const FARCompiledFilter& InFilter, TSet<FName>& OutPackageNamesWithAssets,
+	bool& bOutStopIteration, TFunctionRef<bool(const UObject* Object, FAssetData&& PartialAssetData)> Callback)
+{
+	checkf(IsInGameThread(), TEXT("Enumerating in-memory assets can only be done on the game thread; it uses non-threadsafe UE::AssetRegistry::Filtering globals."));
+	bOutStopIteration = false;
 
-				// Skip classes that report themselves as assets but that the editor AssetRegistry is currently not counting as assets
-				if (UE::AssetRegistry::FFiltering::ShouldSkipAsset(Obj))
-				{
-					return;
-				}
+	// Skip assets that were loaded for diffing
+	const uint32 FilterWithoutPackageFlags = InFilter.WithoutPackageFlags | PKG_ForDiffing;
+	const uint32 FilterWithPackageFlags = InFilter.WithPackageFlags;
+
+	auto FilterInMemoryObjectLambda = [&](const UObject* Obj, bool& OutContinue)
+	{
+		if (Obj->IsAsset())
+		{
+			// Skip assets that are currently loading
+			if (Obj->HasAnyFlags(RF_NeedLoad))
+			{
+				return;
+			}
 
-				// Package name
-				const FName PackageName = InMemoryPackage->GetFName();
+			UPackage* InMemoryPackage = Obj->GetOutermost();
 
-				PackagesToSkip.Add(PackageName);
+			// Skip assets with any of the specified 'without' package flags 
+			if (InMemoryPackage->HasAnyPackageFlags(FilterWithoutPackageFlags))
+			{
+				return;
+			}
 
-				if (InFilter.PackageNames.Num() && !InFilter.PackageNames.Contains(PackageName))
-				{
-					return;
-				}
+			// Skip assets without any the specified 'with' packages flags
+			if (!InMemoryPackage->HasAllPackagesFlags(FilterWithPackageFlags))
+			{
+				return;
+			}
 
-				// Object Path
-				const FString ObjectPathStr = Obj->GetPathName();
-				if (InFilter.ObjectPaths.Num() > 0)
-				{
-					const FName ObjectPath = FName(*ObjectPathStr, FNAME_Find);
-					if (!InFilter.ObjectPaths.Contains(ObjectPath))
-					{
-						return;
-					}
-				}
+			// Skip classes that report themselves as assets but that the editor AssetRegistry is currently not counting as assets
+			if (UE::AssetRegistry::FFiltering::ShouldSkipAsset(Obj))
+			{
+				return;
+			}
 
-				// Package path
-				const FString PackageNameStr = InMemoryPackage->GetName();
-				const FName PackagePath = FName(*FPackageName::GetLongPackagePath(PackageNameStr));
-				if (InFilter.PackagePaths.Num() > 0 && !InFilter.PackagePaths.Contains(PackagePath))
+			// Package name
+			const FName PackageName = InMemoryPackage->GetFName();
+
+			OutPackageNamesWithAssets.Add(PackageName);
+
+			if (InFilter.PackageNames.Num() && !InFilter.PackageNames.Contains(PackageName))
+			{
+				return;
+			}
+
+			// Object Path
+			const FString ObjectPathStr = Obj->GetPathName();
+			if (InFilter.ObjectPaths.Num() > 0)
+			{
+				const FName ObjectPath = FName(*ObjectPathStr, FNAME_Find);
+				if (!InFilter.ObjectPaths.Contains(ObjectPath))
 				{
 					return;
 				}
-			
-				// Could perhaps save some FName -> String conversions by creating this a bit earlier using the UObject constructor
-				// to get package name and path.
-				FAssetData AssetData(PackageNameStr, ObjectPathStr, Obj->GetClass()->GetFName(), FAssetDataTagMap(), InMemoryPackage->GetChunkIDs(), InMemoryPackage->GetPackageFlags());
-				Obj->GetAssetRegistryTags(AssetData);
+			}
 
-				// Tags and values
-				if (InFilter.TagsAndValues.Num() > 0)
-				{
-					bool bMatch = false;
-					for (const TPair<FName, TOptional<FString>>& FilterPair : InFilter.TagsAndValues)
-					{
-						FAssetTagValueRef RegistryValue = AssetData.TagsAndValues.FindTag(FilterPair.Key);
+			// Package path
+			const FString PackageNameStr = InMemoryPackage->GetName();
+			const FName PackagePath = FName(*FPackageName::GetLongPackagePath(PackageNameStr));
+			if (InFilter.PackagePaths.Num() > 0 && !InFilter.PackagePaths.Contains(PackagePath))
+			{
+				return;
+			}
 
-						if (RegistryValue.IsSet() && (!FilterPair.Value.IsSet() || RegistryValue == FilterPair.Value.GetValue()))
-						{
-							bMatch = true;
-							break;
-						}
-					}
+			// Could perhaps save some FName -> String conversions by creating this a bit earlier using the UObject constructor
+			// to get package name and path.
+			FAssetData PartialAssetData(PackageNameStr, ObjectPathStr, Obj->GetClass()->GetFName(), FAssetDataTagMap(),
+				InMemoryPackage->GetChunkIDs(), InMemoryPackage->GetPackageFlags());
 
-					if (!bMatch)
-					{
-						return;
-					}
-				}
+			// All filters passed, except for AssetRegistry filter; caller must check that one
+			OutContinue = Callback(Obj, MoveTemp(PartialAssetData));
+		}
+	};
 
-				// All filters passed
-				OutContinue = Callback(AssetData);
+	// Iterate over all in-memory assets to find the ones that pass the filter components
+	if (InFilter.ClassNames.Num() > 0)
+	{
+		TArray<UObject*> InMemoryObjects;
+		for (FName ClassName : InFilter.ClassNames)
+		{
+			UClass* Class = FindObjectFast<UClass>(nullptr, ClassName, false, true, RF_NoFlags);
+			if (Class != nullptr)
+			{
+				GetObjectsOfClass(Class, InMemoryObjects, false, RF_NoFlags);
 			}
-		};
+		}
 
-		// Iterate over all in-memory assets to find the ones that pass the filter components
-		if (InFilter.ClassNames.Num() > 0)
+		for (UObject* Object : InMemoryObjects)
 		{
-			TArray<UObject*> InMemoryObjects;
-			for (FName ClassName : InFilter.ClassNames)
+			bool bContinue = true;
+			FilterInMemoryObjectLambda(Object, bContinue);
+			if (!bContinue)
 			{
-				UClass* Class = FindObjectFast<UClass>(nullptr, ClassName, false, true, RF_NoFlags);
-				if (Class != nullptr)
-				{
-					GetObjectsOfClass(Class, InMemoryObjects, false, RF_NoFlags);
-				}
+				bOutStopIteration = true;
+				return;
 			}
-
-			for (UObject* Object : InMemoryObjects)
+		}
+	}
+	else
+	{
+		for (FThreadSafeObjectIterator ObjIt; ObjIt; ++ObjIt)
+		{
+			bool bContinue = true;
+			FilterInMemoryObjectLambda(*ObjIt, bContinue);
+			if (!bContinue)
 			{
-				bool bContinue = true;
-				FilterInMemoryObjectLambda(Object, bContinue);
-				if (!bContinue)
-				{
-					return true;
-				}
+				bOutStopIteration = true;
+				return;
 			}
+
+			FPlatformMisc::PumpEssentialAppMessages();
 		}
-		else
+	}
+}
+
+void EnumerateMemoryAssets(const FARCompiledFilter& InFilter, TSet<FName>& OutPackageNamesWithAssets,
+	bool& bOutStopIteration, TFunctionRef<bool(FAssetData&&)> Callback)
+{
+	check(!InFilter.IsEmpty() && Utils::IsFilterValid(InFilter));
+	EnumerateMemoryAssetsHelper(InFilter, OutPackageNamesWithAssets, bOutStopIteration,
+		[&InFilter, &Callback](const UObject* Object, FAssetData&& PartialAssetData)
 		{
-			for (FThreadSafeObjectIterator ObjIt; ObjIt; ++ObjIt)
+			Object->GetAssetRegistryTags(PartialAssetData);
+			// After adding tags, PartialAssetData is now a full AssetData
+
+			// Tags and values
+			if (InFilter.TagsAndValues.Num() > 0)
 			{
-				bool bContinue = true;
-				FilterInMemoryObjectLambda(*ObjIt, bContinue);
-				if (!bContinue)
+				bool bMatch = false;
+				for (const TPair<FName, TOptional<FString>>& FilterPair : InFilter.TagsAndValues)
 				{
-					return true;
+					FAssetTagValueRef RegistryValue = PartialAssetData.TagsAndValues.FindTag(FilterPair.Key);
+
+					if (RegistryValue.IsSet() && (!FilterPair.Value.IsSet() || RegistryValue == FilterPair.Value.GetValue()))
+					{
+						bMatch = true;
+						break;
+					}
 				}
 
-				FPlatformMisc::PumpEssentialAppMessages();
+				if (!bMatch)
+				{
+					return true;
+				}
 			}
-		}
-	}
 
+			// All filters passed
+			return Callback(MoveTemp(PartialAssetData));
+		});
+}
+
+}
+
+void FAssetRegistryImpl::EnumerateDiskAssets(const FARCompiledFilter& InFilter, TSet<FName>& PackagesToSkip,
+	TFunctionRef<bool(const FAssetData&)> Callback) const
+{
+	check(!InFilter.IsEmpty() && Utils::IsFilterValid(InFilter));
+	PackagesToSkip.Append(CachedEmptyPackages);
 	State.EnumerateAssets(InFilter, PackagesToSkip, Callback, true /*bARFiltering*/);
+}
 
-	return true;
 }
 
 FAssetData UAssetRegistryImpl::GetAssetByObjectPath(const FName ObjectPath, bool bIncludeOnlyOnDiskAssets) const
@@ -1154,73 +1464,151 @@ FAssetData UAssetRegistryImpl::GetAssetByObjectPath(const FName ObjectPath, bool
 
 		if (Asset)
 		{
-			return FAssetData(Asset);
+			return FAssetData(Asset, false /* bAllowBlueprintClass */);
 		}
 	}
 
-	const FAssetData* FoundData = State.GetAssetByObjectPath(ObjectPath);
-
-	return (FoundData && !UE::AssetRegistry::FFiltering::ShouldSkipAsset(FoundData->AssetClass, FoundData->PackageFlags)) ? *FoundData : FAssetData();
+	{
+		FReadScopeLock InterfaceScopeLock(InterfaceLock);
+		const FAssetRegistryState& State = GuardedData.GetState();
+		const FAssetData* FoundData = State.GetAssetByObjectPath(ObjectPath);
+		return (FoundData && !GuardedData.ShouldSkipAsset(FoundData->AssetClass, FoundData->PackageFlags)) ? *FoundData : FAssetData();
+	}
 }
 
 bool UAssetRegistryImpl::GetAllAssets(TArray<FAssetData>& OutAssetData, bool bIncludeOnlyOnDiskAssets) const
 {
 	const double GetAllAssetsStartTime = FPlatformTime::Seconds();
+	TSet<FName> PackageNamesToSkip;
 
-	const bool bResult = EnumerateAllAssets([&OutAssetData](const FAssetData& AssetData)
+	// All in memory assets
+	if (!bIncludeOnlyOnDiskAssets)
 	{
-		OutAssetData.Emplace(AssetData);
-		return true;
-	}, bIncludeOnlyOnDiskAssets);
+		bool bStopIterationUnused;
+		UE::AssetRegistry::Utils::EnumerateAllMemoryAssets(PackageNamesToSkip, bStopIterationUnused,
+			[&OutAssetData](FAssetData&& AssetData)
+			{
+				OutAssetData.Add(MoveTemp(AssetData));
+				return true;
+			});
+	}
+
+	{
+		FReadScopeLock InterfaceScopeLock(InterfaceLock);
+		GuardedData.EnumerateAllDiskAssets(PackageNamesToSkip,
+			[&OutAssetData](const FAssetData& AssetData)
+			{
+				OutAssetData.Add(AssetData);
+				return true;
+			});
+	}
 
 	UE_LOG(LogAssetRegistry, VeryVerbose, TEXT("GetAllAssets completed in %0.4f seconds"), FPlatformTime::Seconds() - GetAllAssetsStartTime);
-	
-	return bResult;
+	return true;
 }
 
 bool UAssetRegistryImpl::EnumerateAllAssets(TFunctionRef<bool(const FAssetData&)> Callback, bool bIncludeOnlyOnDiskAssets) const
 {
-	TSet<FName> PackageNamesToSkip = CachedEmptyPackages;
+	const double GetAllAssetsStartTime = FPlatformTime::Seconds();
+	TSet<FName> PackageNamesToSkip;
 
 	// All in memory assets
 	if (!bIncludeOnlyOnDiskAssets)
 	{
-		for (FThreadSafeObjectIterator ObjIt; ObjIt; ++ObjIt)
+		bool bStopIteration;
+		UE::AssetRegistry::Utils::EnumerateAllMemoryAssets(PackageNamesToSkip, bStopIteration,
+			[&Callback](FAssetData&& AssetData)
+			{
+				return Callback(AssetData);
+			});
+		if (bStopIteration)
+		{
+			return true;
+		}
+	}
+
+	// We have to call the callback on a copy rather than a reference since the callback may reenter the lock
+	TArray<FAssetData, TInlineAllocator<128>> OnDiskAssetDatas;
+	{
+		FReadScopeLock InterfaceScopeLock(InterfaceLock);
+		GuardedData.EnumerateAllDiskAssets(PackageNamesToSkip,
+			[&OnDiskAssetDatas](const FAssetData& AssetData)
+			{
+				OnDiskAssetDatas.Add(AssetData);
+				return true;
+			});
+	}
+
+	for (const FAssetData& AssetData : OnDiskAssetDatas)
+	{
+		if (!Callback(AssetData))
 		{
-			if (ObjIt->IsAsset())
+			return true;
+		}
+	}
+	return true;
+}
+
+namespace UE::AssetRegistry
+{
+
+namespace Utils
+{
+
+void EnumerateAllMemoryAssets(TSet<FName>& OutPackageNamesWithAssets, bool& bOutStopIteration,
+	TFunctionRef<bool(FAssetData&&)> Callback)
+{
+	checkf(IsInGameThread(), TEXT("Enumerating memory assets can only be done on the game thread; it uses non-threadsafe UE::AssetRegistry::Filtering globals."));
+	bOutStopIteration = false;
+	for (FThreadSafeObjectIterator ObjIt; ObjIt; ++ObjIt)
+	{
+		if (ObjIt->IsAsset() && !UE::AssetRegistry::FFiltering::ShouldSkipAsset(*ObjIt))
+		{
+			FAssetData AssetData(*ObjIt, true /* bAllowBlueprintClass */);
+			OutPackageNamesWithAssets.Add(AssetData.PackageName);
+			if (!Callback(MoveTemp(AssetData)))
 			{
-				const FAssetData AssetData(*ObjIt);
-				if (!UE::AssetRegistry::FFiltering::ShouldSkipAsset(*ObjIt))
-				{
-					if (!Callback(AssetData))
-					{
-						return true;
-					}
-				}
-				PackageNamesToSkip.Add(AssetData.PackageName);
+				bOutStopIteration = true;
+				return;
 			}
 		}
 	}
+}
 
+}
+
+void FAssetRegistryImpl::EnumerateAllDiskAssets(TSet<FName>& PackageNamesToSkip, TFunctionRef<bool(const FAssetData&)> Callback) const
+{
+	PackageNamesToSkip.Append(CachedEmptyPackages);
 	State.EnumerateAllAssets(PackageNamesToSkip, Callback, true /*bARFiltering*/);
+}
 
-	return true;
 }
 
 void UAssetRegistryImpl::GetPackagesByName(FStringView PackageName, TArray<FName>& OutPackageNames) const
 {
-	UE_CLOG(bInitialSearchStarted && !bInitialSearchCompleted, LogAssetRegistry, Warning,
+	FReadScopeLock InterfaceScopeLock(InterfaceLock);
+	const FAssetRegistryState& State = GuardedData.GetState();
+	UE_CLOG(GuardedData.IsInitialSearchStarted() && !GuardedData.IsInitialSearchCompleted(), LogAssetRegistry, Warning,
 		TEXT("GetPackagesByName has been called before AssetRegistry gather is complete and it does not wait. ")
 		TEXT("The search may return incomplete results."));
 	State.GetPackagesByName(PackageName, OutPackageNames);
+
 }
 
 FName UAssetRegistryImpl::GetFirstPackageByName(FStringView PackageName) const
 {
-	UE_CLOG(bInitialSearchStarted && !bInitialSearchCompleted, LogAssetRegistry, Warning,
-		TEXT("GetFirstPackageByName has been called before AssetRegistry gather is complete and it does not wait. ")
-		TEXT("The search may fail to find the package."));
-	FName LongPackageName = State.GetFirstPackageByName(PackageName);
+	FName LongPackageName;
+	bool bSearchAllAssets;
+	{
+		FReadScopeLock InterfaceScopeLock(InterfaceLock);
+		const FAssetRegistryState& State = GuardedData.GetState();
+		UE_CLOG(GuardedData.IsInitialSearchStarted() && !GuardedData.IsInitialSearchCompleted(), LogAssetRegistry, Warning,
+			TEXT("GetFirstPackageByName has been called before AssetRegistry gather is complete and it does not wait. ")
+			TEXT("The search may fail to find the package."));
+		LongPackageName = State.GetFirstPackageByName(PackageName);
+		bSearchAllAssets = GuardedData.IsSearchAllAssets();
+	}
 #if WITH_EDITOR
 	if (!GIsEditor && !bSearchAllAssets)
 	{
@@ -1243,19 +1631,22 @@ FName UAssetRegistryImpl::GetFirstPackageByName(FStringView PackageName) const
 
 bool UAssetRegistryImpl::GetDependencies(const FAssetIdentifier& AssetIdentifier, TArray<FAssetIdentifier>& OutDependencies, EAssetRegistryDependencyType::Type InDependencyType) const
 {
+	FReadScopeLock InterfaceScopeLock(InterfaceLock);
 	PRAGMA_DISABLE_DEPRECATION_WARNINGS
-	return State.GetDependencies(AssetIdentifier, OutDependencies, InDependencyType);
+	return GuardedData.GetState().GetDependencies(AssetIdentifier, OutDependencies, InDependencyType);
 	PRAGMA_ENABLE_DEPRECATION_WARNINGS
 }
 
 bool UAssetRegistryImpl::GetDependencies(const FAssetIdentifier& AssetIdentifier, TArray<FAssetIdentifier>& OutDependencies, UE::AssetRegistry::EDependencyCategory Category, const UE::AssetRegistry::FDependencyQuery& Flags) const
 {
-	return State.GetDependencies(AssetIdentifier, OutDependencies, Category, Flags);
+	FReadScopeLock InterfaceScopeLock(InterfaceLock);
+	return GuardedData.GetState().GetDependencies(AssetIdentifier, OutDependencies, Category, Flags);
 }
 
 bool UAssetRegistryImpl::GetDependencies(const FAssetIdentifier& AssetIdentifier, TArray<FAssetDependency>& OutDependencies, UE::AssetRegistry::EDependencyCategory Category, const UE::AssetRegistry::FDependencyQuery& Flags) const
 {
-	return State.GetDependencies(AssetIdentifier, OutDependencies, Category, Flags);
+	FReadScopeLock InterfaceScopeLock(InterfaceLock);
+	return GuardedData.GetState().GetDependencies(AssetIdentifier, OutDependencies, Category, Flags);
 }
 
 static void ConvertAssetIdentifiersToPackageNames(const TArray<FAssetIdentifier>& AssetIdentifiers, TArray<FName>& OutPackageNames)
@@ -1314,19 +1705,22 @@ bool IAssetRegistry::K2_GetDependencies(FName PackageName, const FAssetRegistryD
 
 bool UAssetRegistryImpl::GetReferencers(const FAssetIdentifier& AssetIdentifier, TArray<FAssetIdentifier>& OutReferencers, EAssetRegistryDependencyType::Type InReferenceType) const
 {
+	FReadScopeLock InterfaceScopeLock(InterfaceLock);
 	PRAGMA_DISABLE_DEPRECATION_WARNINGS
-	return State.GetReferencers(AssetIdentifier, OutReferencers, InReferenceType);
+	return GuardedData.GetState().GetReferencers(AssetIdentifier, OutReferencers, InReferenceType);
 	PRAGMA_ENABLE_DEPRECATION_WARNINGS
 }
 
 bool UAssetRegistryImpl::GetReferencers(const FAssetIdentifier& AssetIdentifier, TArray<FAssetIdentifier>& OutReferencers, UE::AssetRegistry::EDependencyCategory Category, const UE::AssetRegistry::FDependencyQuery& Flags) const
 {
-	return State.GetReferencers(AssetIdentifier, OutReferencers, Category, Flags);
+	FReadScopeLock InterfaceScopeLock(InterfaceLock);
+	return GuardedData.GetState().GetReferencers(AssetIdentifier, OutReferencers, Category, Flags);
 }
 
 bool UAssetRegistryImpl::GetReferencers(const FAssetIdentifier& AssetIdentifier, TArray<FAssetDependency>& OutReferencers, UE::AssetRegistry::EDependencyCategory Category, const UE::AssetRegistry::FDependencyQuery& Flags) const
 {
-	return State.GetReferencers(AssetIdentifier, OutReferencers, Category, Flags);
+	FReadScopeLock InterfaceScopeLock(InterfaceLock);
+	return GuardedData.GetState().GetReferencers(AssetIdentifier, OutReferencers, Category, Flags);
 }
 
 bool UAssetRegistryImpl::GetReferencers(FName PackageName, TArray<FName>& OutReferencers, EAssetRegistryDependencyType::Type InReferenceType) const
@@ -1377,51 +1771,131 @@ bool IAssetRegistry::K2_GetReferencers(FName PackageName, const FAssetRegistryDe
 
 const FAssetPackageData* UAssetRegistryImpl::GetAssetPackageData(FName PackageName) const
 {
-	return State.GetAssetPackageData(PackageName);
+	FWriteScopeLock InterfaceScopeLock(InterfaceLock);
+	PRAGMA_DISABLE_DEPRECATION_WARNINGS
+	return const_cast<UE::AssetRegistry::FAssetRegistryImpl&>(GuardedData).GetAssetPackageData(PackageName);
+	PRAGMA_ENABLE_DEPRECATION_WARNINGS
+}
+
+namespace UE::AssetRegistry
+{
+
+const FAssetPackageData* FAssetRegistryImpl::GetAssetPackageData(FName PackageName)
+{
+	const FAssetPackageData* AssetPackageData = State.GetAssetPackageData(PackageName);
+	if (!AssetPackageData)
+	{
+		return nullptr;
+	}
+	FAssetPackageData* Result = new FAssetPackageData(*AssetPackageData);
+	PRAGMA_DISABLE_DEPRECATION_WARNINGS
+	DeleteActions.Add([Result]() { delete Result; });
+	PRAGMA_ENABLE_DEPRECATION_WARNINGS
+	return Result;
+}
+
+}
+
+TOptional<FAssetPackageData> UAssetRegistryImpl::GetAssetPackageDataCopy(FName PackageName) const
+{
+	FReadScopeLock InterfaceScopeLock(InterfaceLock);
+	const FAssetPackageData* AssetPackageData = GuardedData.GetState().GetAssetPackageData(PackageName);
+	return AssetPackageData ? *AssetPackageData : TOptional<FAssetPackageData>();
 }
 
 FName UAssetRegistryImpl::GetRedirectedObjectPath(const FName ObjectPath) const
+{
+	FReadScopeLock InterfaceScopeLock(InterfaceLock);
+	return GuardedData.GetRedirectedObjectPath(ObjectPath);
+}
+
+namespace UE::AssetRegistry
+{
+
+FName FAssetRegistryImpl::GetRedirectedObjectPath(const FName ObjectPath) const
 {
 	FString RedirectedPath = ObjectPath.ToString();
-	FAssetData DestinationData = GetAssetByObjectPath(ObjectPath);
+
+	// For legacy behavior, for the first object pointed to, we look up the object in memory
+	// before checking the on-disk assets
+	UObject* Asset = FindObject<UObject>(nullptr, *ObjectPath.ToString());
+	const FAssetData* AssetData = nullptr;
+	if (!Asset)
+	{
+		AssetData = State.GetAssetByObjectPath(ObjectPath);
+	}
+
 	TSet<FString> SeenPaths;
 	SeenPaths.Add(RedirectedPath);
 
-	// Need to follow chain of redirectors
-	while (DestinationData.IsRedirector())
+	auto TryGetRedirectedPath = [](UObject* InAsset, const FAssetData* InAssetData, FString& OutRedirectedPath)
 	{
-		if (DestinationData.GetTagValue("DestinationObject", RedirectedPath))
+		if (InAsset)
 		{
-			ConstructorHelpers::StripObjectClass(RedirectedPath);
-			if (SeenPaths.Contains(RedirectedPath))
+			UObjectRedirector* Redirector = Cast<UObjectRedirector>(InAsset);
+			if (Redirector && Redirector->DestinationObject)
 			{
-				// Recursive, bail
-				DestinationData = FAssetData();
+				OutRedirectedPath = Redirector->DestinationObject->GetPathName();
+				return true;
 			}
-			else
+		}
+		else if (InAssetData && InAssetData->IsRedirector())
+		{
+			if (InAssetData->GetTagValue("DestinationObject", OutRedirectedPath))
 			{
-				SeenPaths.Add(RedirectedPath);
-				DestinationData = GetAssetByObjectPath(FName(*RedirectedPath), true);
+				ConstructorHelpers::StripObjectClass(OutRedirectedPath);
+				return true;
 			}
 		}
+		return false;
+	};
+
+	// Need to follow chain of redirectors
+	while (TryGetRedirectedPath(Asset, AssetData, RedirectedPath))
+	{
+		if (SeenPaths.Contains(RedirectedPath))
+		{
+			// Recursive, bail
+			break;
+		}
 		else
 		{
-			// Can't extract
-			DestinationData = FAssetData();
+			SeenPaths.Add(RedirectedPath);
+			// For legacy behavior, for all redirects after the initial request, we only check on-disk assets
+			Asset = nullptr;
+			AssetData = State.GetAssetByObjectPath(FName(*RedirectedPath));
 		}
 	}
 
 	return FName(*RedirectedPath);
 }
 
+}
+
 bool UAssetRegistryImpl::GetAncestorClassNames(FName ClassName, TArray<FName>& OutAncestorClassNames) const
+{
+	UE::AssetRegistry::Impl::FClassInheritanceContext InheritanceContext;
+	UE::AssetRegistry::Impl::FClassInheritanceBuffer InheritanceBuffer;
+	FRWScopeLock InterfaceScopeLock(InterfaceLock, SLT_ReadOnly);
+	const_cast<UAssetRegistryImpl*>(this)->GetInheritanceContextWithRequiredLock(
+		InterfaceScopeLock, InheritanceContext, InheritanceBuffer);
+	return GuardedData.GetAncestorClassNames(InheritanceContext, ClassName, OutAncestorClassNames);
+}
+
+namespace UE::AssetRegistry
+{
+
+bool FAssetRegistryImpl::GetAncestorClassNames(Impl::FClassInheritanceContext& InheritanceContext, FName ClassName,
+	TArray<FName>& OutAncestorClassNames) const
 {
 	// Assume we found the class unless there is an error
 	bool bFoundClass = true;
-	UpdateTemporaryCaches();
+
+	InheritanceContext.ConditionalUpdate();
+	const TMap<FName, FName>& InheritanceMap = InheritanceContext.Buffer->InheritanceMap;
 
 	// Make sure the requested class is in the inheritance map
-	if (!TempCachedInheritanceMap.Contains(ClassName))
+	if (!InheritanceMap.Contains(ClassName))
 	{
 		bFoundClass = false;
 	}
@@ -1433,7 +1907,7 @@ bool UAssetRegistryImpl::GetAncestorClassNames(FName ClassName, TArray<FName>& O
 		uint32 CurrentInheritanceDepth = 0;
 		while (CurrentInheritanceDepth < MaxInheritanceDepth && CurrentClassName != nullptr)
 		{
-			CurrentClassName = TempCachedInheritanceMap.Find(*CurrentClassName);
+			CurrentClassName = InheritanceMap.Find(*CurrentClassName);
 
 			if (CurrentClassName)
 			{
@@ -1457,18 +1931,28 @@ bool UAssetRegistryImpl::GetAncestorClassNames(FName ClassName, TArray<FName>& O
 		}
 	}
 
-	ClearTemporaryCaches();
 	return bFoundClass;
 }
 
-void UAssetRegistryImpl::GetDerivedClassNames(const TArray<FName>& ClassNames, const TSet<FName>& ExcludedClassNames, TSet<FName>& OutDerivedClassNames) const
+}
+
+void UAssetRegistryImpl::GetDerivedClassNames(const TArray<FName>& ClassNames, const TSet<FName>& ExcludedClassNames,
+	TSet<FName>& OutDerivedClassNames) const
 {
-	GetSubClasses(ClassNames, ExcludedClassNames, OutDerivedClassNames);
+	UE::AssetRegistry::Impl::FClassInheritanceContext InheritanceContext;
+	UE::AssetRegistry::Impl::FClassInheritanceBuffer InheritanceBuffer;
+	FRWScopeLock InterfaceScopeLock(InterfaceLock, SLT_ReadOnly);
+	const_cast<UAssetRegistryImpl*>(this)->GetInheritanceContextWithRequiredLock(
+		InterfaceScopeLock, InheritanceContext, InheritanceBuffer);
+	GuardedData.GetSubClasses(InheritanceContext, ClassNames, ExcludedClassNames, OutDerivedClassNames);
 }
 
 void UAssetRegistryImpl::GetAllCachedPaths(TArray<FString>& OutPathList) const
 {
-	EnumerateAllCachedPaths([&OutPathList](FName Path)
+	FReadScopeLock InterfaceScopeLock(InterfaceLock);
+	const FPathTree& CachedPathTree = GuardedData.GetCachedPathTree();
+	OutPathList.Reserve(OutPathList.Num() + CachedPathTree.NumPaths());
+	CachedPathTree.EnumerateAllPaths([&OutPathList](FName Path)
 	{
 		OutPathList.Emplace(Path.ToString());
 		return true;
@@ -1485,12 +1969,31 @@ void UAssetRegistryImpl::EnumerateAllCachedPaths(TFunctionRef<bool(FString)> Cal
 
 void UAssetRegistryImpl::EnumerateAllCachedPaths(TFunctionRef<bool(FName)> Callback) const
 {
-	CachedPathTree.EnumerateAllPaths(Callback);
+	TArray<FName> FoundPaths;
+	{
+		FReadScopeLock InterfaceScopeLock(InterfaceLock);
+		const FPathTree& CachedPathTree = GuardedData.GetCachedPathTree();
+		FoundPaths.Reserve(CachedPathTree.NumPaths());
+		CachedPathTree.EnumerateAllPaths([&FoundPaths](FName Path)
+		{
+			FoundPaths.Add(Path);
+			return true;
+		});
+	}
+	for (FName Path : FoundPaths)
+	{
+		if (!Callback(Path))
+		{
+			return;
+		}
+	}
 }
 
 void UAssetRegistryImpl::GetSubPaths(const FString& InBasePath, TArray<FString>& OutPathList, bool bInRecurse) const
 {
-	EnumerateSubPaths(*InBasePath, [&OutPathList](FName Path)
+	FReadScopeLock InterfaceScopeLock(InterfaceLock);
+	const FPathTree& CachedPathTree = GuardedData.GetCachedPathTree();
+	CachedPathTree.EnumerateSubPaths(*InBasePath, [&OutPathList](FName Path)
 	{
 		OutPathList.Emplace(Path.ToString());
 		return true;
@@ -1499,50 +2002,94 @@ void UAssetRegistryImpl::GetSubPaths(const FString& InBasePath, TArray<FString>&
 
 void UAssetRegistryImpl::EnumerateSubPaths(const FString& InBasePath, TFunctionRef<bool(FString)> Callback, bool bInRecurse) const
 {
-	EnumerateSubPaths(*InBasePath, [&Callback](FName Path)
+	TArray<FName, TInlineAllocator<64>> SubPaths;
 	{
-		return Callback(Path.ToString());
-	}, bInRecurse);
+		FReadScopeLock InterfaceScopeLock(InterfaceLock);
+		const FPathTree& CachedPathTree = GuardedData.GetCachedPathTree();
+		CachedPathTree.EnumerateSubPaths(FName(*InBasePath), [&SubPaths](FName PathName)
+		{
+			SubPaths.Add(PathName);
+			return true;
+		}, bInRecurse);
+	}
+	for (FName PathName : SubPaths)
+	{
+		if (!Callback(PathName.ToString()))
+		{
+			break;
+		}
+	}
 }
 
 void UAssetRegistryImpl::EnumerateSubPaths(const FName InBasePath, TFunctionRef<bool(FName)> Callback, bool bInRecurse) const
 {
-	CachedPathTree.EnumerateSubPaths(InBasePath, Callback, bInRecurse);
+	TArray<FName, TInlineAllocator<64>> SubPaths;
+	{
+		FReadScopeLock InterfaceScopeLock(InterfaceLock);
+		const FPathTree& CachedPathTree = GuardedData.GetCachedPathTree();
+		CachedPathTree.EnumerateSubPaths(InBasePath, [&SubPaths](FName PathName)
+		{
+			SubPaths.Add(PathName);
+			return true;
+		}, bInRecurse);
+	}
+	for (FName PathName : SubPaths)
+	{
+		if (!Callback(PathName))
+		{
+			break;
+		}
+	}
 }
 
 void UAssetRegistryImpl::RunAssetsThroughFilter(TArray<FAssetData>& AssetDataList, const FARFilter& Filter) const
 {
-	RunAssetsThroughFilterImpl(AssetDataList, Filter, EARFilterMode::Inclusive);
+	if (Filter.IsEmpty())
+	{
+		return;
+	}
+	FARCompiledFilter CompiledFilter;
+	CompileFilter(Filter, CompiledFilter);
+	UE::AssetRegistry::Utils::RunAssetsThroughFilter(AssetDataList, CompiledFilter, UE::AssetRegistry::Utils::EFilterMode::Inclusive);
 }
 
 void UAssetRegistryImpl::UseFilterToExcludeAssets(TArray<FAssetData>& AssetDataList, const FARFilter& Filter) const
 {
-	RunAssetsThroughFilterImpl(AssetDataList, Filter, EARFilterMode::Exclusive);
+	if (Filter.IsEmpty())
+	{
+		return;
+	}
+	FARCompiledFilter CompiledFilter;
+	CompileFilter(Filter, CompiledFilter);
+	UE::AssetRegistry::Utils::RunAssetsThroughFilter(AssetDataList, CompiledFilter, UE::AssetRegistry::Utils::EFilterMode::Exclusive);
 }
 
 bool UAssetRegistryImpl::IsAssetIncludedByFilter(const FAssetData& AssetData, const FARCompiledFilter& Filter) const
 {
-	return RunAssetThroughFilterImpl(AssetData, Filter, EARFilterMode::Inclusive);
+	return UE::AssetRegistry::Utils::RunAssetThroughFilter(AssetData, Filter, UE::AssetRegistry::Utils::EFilterMode::Inclusive);
 }
 
 bool UAssetRegistryImpl::IsAssetExcludedByFilter(const FAssetData& AssetData, const FARCompiledFilter& Filter) const
 {
-	return RunAssetThroughFilterImpl(AssetData, Filter, EARFilterMode::Exclusive);
+	return UE::AssetRegistry::Utils::RunAssetThroughFilter(AssetData, Filter, UE::AssetRegistry::Utils::EFilterMode::Exclusive);
 }
 
-bool UAssetRegistryImpl::RunAssetThroughFilterImpl(const FAssetData& AssetData, const FARCompiledFilter& Filter, const EARFilterMode FilterMode) const
+namespace UE::AssetRegistry::Utils
 {
-	const bool bPassFilterValue = FilterMode == EARFilterMode::Inclusive;
+
+bool RunAssetThroughFilter(const FAssetData& AssetData, const FARCompiledFilter& Filter, const EFilterMode FilterMode)
+{
+	const bool bPassFilterValue = FilterMode == EFilterMode::Inclusive;
 	if (Filter.IsEmpty())
 	{
 		return bPassFilterValue;
 	}
 
-	const bool bFilterResult = RunAssetThroughFilterImpl_Unchecked(AssetData, Filter, bPassFilterValue);
+	const bool bFilterResult = RunAssetThroughFilter_Unchecked(AssetData, Filter, bPassFilterValue);
 	return bFilterResult == bPassFilterValue;
 }
 
-bool UAssetRegistryImpl::RunAssetThroughFilterImpl_Unchecked(const FAssetData& AssetData, const FARCompiledFilter& Filter, const bool bPassFilterValue) const
+bool RunAssetThroughFilter_Unchecked(const FAssetData& AssetData, const FARCompiledFilter& Filter, const bool bPassFilterValue)
 {
 	// Package Names
 	if (Filter.PackageNames.Num() > 0)
@@ -1607,27 +2154,20 @@ bool UAssetRegistryImpl::RunAssetThroughFilterImpl_Unchecked(const FAssetData& A
 	return bPassFilterValue;
 }
 
-void UAssetRegistryImpl::RunAssetsThroughFilterImpl(TArray<FAssetData>& AssetDataList, const FARFilter& Filter, const EARFilterMode FilterMode) const
+void RunAssetsThroughFilter(TArray<FAssetData>& AssetDataList, const FARCompiledFilter& CompiledFilter, const EFilterMode FilterMode)
 {
-	if (Filter.IsEmpty())
-	{
-		return;
-	}
-
-	FARCompiledFilter CompiledFilter;
-	CompileFilter(Filter, CompiledFilter);
-	if (!FAssetRegistryState::IsFilterValid(CompiledFilter))
+	if (!IsFilterValid(CompiledFilter))
 	{
 		return;
 	}
 
 	const int32 OriginalArrayCount = AssetDataList.Num();
-	const bool bPassFilterValue = FilterMode == EARFilterMode::Inclusive;
+	const bool bPassFilterValue = FilterMode == EFilterMode::Inclusive;
 
 	// Spin the array backwards to minimize the number of elements that are repeatedly moved down
 	for (int32 AssetDataIndex = AssetDataList.Num() - 1; AssetDataIndex >= 0; --AssetDataIndex)
 	{
-		const bool bFilterResult = RunAssetThroughFilterImpl_Unchecked(AssetDataList[AssetDataIndex], CompiledFilter, bPassFilterValue);
+		const bool bFilterResult = RunAssetThroughFilter_Unchecked(AssetDataList[AssetDataIndex], CompiledFilter, bPassFilterValue);
 		if (bFilterResult != bPassFilterValue)
 		{
 			AssetDataList.RemoveAt(AssetDataIndex, 1, /*bAllowShrinking*/false);
@@ -1640,6 +2180,8 @@ void UAssetRegistryImpl::RunAssetsThroughFilterImpl(TArray<FAssetData>& AssetDat
 	}
 }
 
+}
+
 void UAssetRegistryImpl::ExpandRecursiveFilter(const FARFilter& InFilter, FARFilter& ExpandedFilter) const
 {
 	FARCompiledFilter CompiledFilter;
@@ -1657,6 +2199,28 @@ void UAssetRegistryImpl::ExpandRecursiveFilter(const FARFilter& InFilter, FARFil
 }
 
 void UAssetRegistryImpl::CompileFilter(const FARFilter& InFilter, FARCompiledFilter& OutCompiledFilter) const
+{
+	UE::AssetRegistry::Impl::FClassInheritanceContext InheritanceContext;
+	UE::AssetRegistry::Impl::FClassInheritanceBuffer InheritanceBuffer;
+	FRWScopeLock InterfaceScopeLock(InterfaceLock, SLT_ReadOnly);
+	if (InFilter.bRecursiveClasses)
+	{
+		const_cast<UAssetRegistryImpl*>(this)->GetInheritanceContextWithRequiredLock(
+			InterfaceScopeLock, InheritanceContext, InheritanceBuffer);
+	}
+	else
+	{
+		// CompileFilter takes an inheritance context, but only to handle filters with recursive classes
+		// which we are not using here, so leave the InheritanceContext empty
+	}
+	GuardedData.CompileFilter(InheritanceContext, InFilter, OutCompiledFilter);
+}
+
+namespace UE::AssetRegistry
+{
+
+void FAssetRegistryImpl::CompileFilter(Impl::FClassInheritanceContext& InheritanceContext, const FARFilter& InFilter,
+	FARCompiledFilter& OutCompiledFilter) const
 {
 	OutCompiledFilter.Clear();
 	OutCompiledFilter.PackageNames.Append(InFilter.PackageNames);
@@ -1685,16 +2249,26 @@ void UAssetRegistryImpl::CompileFilter(const FARFilter& InFilter, FARCompiledFil
 			TArray<FName> ClassNamesObject;
 			ClassNamesObject.Add(UObject::StaticClass()->GetFName());
 
-			GetSubClasses(ClassNamesObject, InFilter.RecursiveClassesExclusionSet, OutCompiledFilter.ClassNames);
+			GetSubClasses(InheritanceContext, ClassNamesObject, InFilter.RecursiveClassesExclusionSet, OutCompiledFilter.ClassNames);
 		}
 		else
 		{
-			GetSubClasses(InFilter.ClassNames, InFilter.RecursiveClassesExclusionSet, OutCompiledFilter.ClassNames);
+			GetSubClasses(InheritanceContext, InFilter.ClassNames, InFilter.RecursiveClassesExclusionSet, OutCompiledFilter.ClassNames);
 		}
 	}
 }
 
+}
+
 EAssetAvailability::Type UAssetRegistryImpl::GetAssetAvailability(const FAssetData& AssetData) const
+{
+	return UE::AssetRegistry::Utils::GetAssetAvailability(AssetData);
+}
+
+namespace UE::AssetRegistry::Utils
+{
+
+EAssetAvailability::Type GetAssetAvailability(const FAssetData& AssetData)
 {
 	IPlatformChunkInstall* ChunkInstall = FPlatformMisc::GetPlatformChunkInstall();
 
@@ -1734,7 +2308,17 @@ EAssetAvailability::Type UAssetRegistryImpl::GetAssetAvailability(const FAssetDa
 	}
 }
 
+}
+
 float UAssetRegistryImpl::GetAssetAvailabilityProgress(const FAssetData& AssetData, EAssetAvailabilityProgressReportingType::Type ReportType) const
+{
+	return UE::AssetRegistry::Utils::GetAssetAvailabilityProgress(AssetData, ReportType);
+}
+
+namespace UE::AssetRegistry::Utils
+{
+
+float GetAssetAvailabilityProgress(const FAssetData& AssetData, EAssetAvailabilityProgressReportingType::Type ReportType)
 {
 	IPlatformChunkInstall* ChunkInstall = FPlatformMisc::GetPlatformChunkInstall();
 	EChunkProgressReportingType::Type ChunkReportType = GetChunkAvailabilityProgressType(ReportType);
@@ -1775,13 +2359,33 @@ float UAssetRegistryImpl::GetAssetAvailabilityProgress(const FAssetData& AssetDa
 	return BestProgress;
 }
 
+}
+
 bool UAssetRegistryImpl::GetAssetAvailabilityProgressTypeSupported(EAssetAvailabilityProgressReportingType::Type ReportType) const
+{
+	return UE::AssetRegistry::Utils::GetAssetAvailabilityProgressTypeSupported(ReportType);
+}
+
+namespace UE::AssetRegistry::Utils
+{
+
+bool GetAssetAvailabilityProgressTypeSupported(EAssetAvailabilityProgressReportingType::Type ReportType)
 {
 	IPlatformChunkInstall* ChunkInstall = FPlatformMisc::GetPlatformChunkInstall();
 	return ChunkInstall->GetProgressReportingTypeSupported(GetChunkAvailabilityProgressType(ReportType));
 }
 
+}
+
 void UAssetRegistryImpl::PrioritizeAssetInstall(const FAssetData& AssetData) const
+{
+	UE::AssetRegistry::Utils::PrioritizeAssetInstall(AssetData);
+}
+
+namespace UE::AssetRegistry::Utils
+{
+
+void PrioritizeAssetInstall(const FAssetData& AssetData)
 {
 	IPlatformChunkInstall* ChunkInstall = FPlatformMisc::GetPlatformChunkInstall();
 
@@ -1793,7 +2397,24 @@ void UAssetRegistryImpl::PrioritizeAssetInstall(const FAssetData& AssetData) con
 	ChunkInstall->PrioritizePakchunk(AssetData.ChunkIDs[0], EChunkPriority::Immediate);
 }
 
+}
+
 bool UAssetRegistryImpl::AddPath(const FString& PathToAdd)
+{
+	UE::AssetRegistry::Impl::FEventContext EventContext;
+	bool bResult;
+	{
+		FWriteScopeLock InterfaceScopeLock(InterfaceLock);
+		bResult = GuardedData.AddPath(EventContext, PathToAdd);
+	}
+	Broadcast(EventContext);
+	return bResult;
+}
+
+namespace UE::AssetRegistry
+{
+
+bool FAssetRegistryImpl::AddPath(Impl::FEventContext& EventContext, const FString& PathToAdd)
 {
 	bool bBlacklisted = false;
 	// If no GlobalGatherer, then we are in the game or non-cook commandlet and we do not implement blacklisting
@@ -1809,12 +2430,21 @@ bool UAssetRegistryImpl::AddPath(const FString& PathToAdd)
 	{
 		return false;
 	}
-	return AddAssetPath(FName(*PathToAdd));
+	return AddAssetPath(EventContext, FName(*PathToAdd));
+}
+
 }
 
 bool UAssetRegistryImpl::RemovePath(const FString& PathToRemove)
 {
-	return RemoveAssetPath(FName(*PathToRemove));
+	UE::AssetRegistry::Impl::FEventContext EventContext;
+	bool bResult;
+	{
+		FWriteScopeLock InterfaceScopeLock(InterfaceLock);
+		bResult = GuardedData.RemoveAssetPath(EventContext, FName(*PathToRemove));
+	}
+	Broadcast(EventContext);
+	return bResult;
 }
 
 bool UAssetRegistryImpl::PathExists(const FString& PathToTest) const
@@ -1824,20 +2454,67 @@ bool UAssetRegistryImpl::PathExists(const FString& PathToTest) const
 
 bool UAssetRegistryImpl::PathExists(const FName PathToTest) const
 {
-	return CachedPathTree.PathExists(PathToTest);
+	FReadScopeLock InterfaceScopeLock(InterfaceLock);
+	return GuardedData.GetCachedPathTree().PathExists(PathToTest);
 }
 
-void UAssetRegistryImpl::ScanPathsSynchronous(const TArray<FString>& InPaths, bool bForceRescan, bool bIgnoreBlackListScanFilters)
+void UAssetRegistryImpl::ScanPathsSynchronous(const TArray<FString>& InPaths, bool bForceRescan,
+	bool bIgnoreBlackListScanFilters)
 {
-	ScanPathsSynchronousInternal(InPaths, TArray<FString>(), bForceRescan, bIgnoreBlackListScanFilters, nullptr /* OutFindAssets */);
+	ScanPathsSynchronousInternal(InPaths, TArray<FString>(), bForceRescan, bIgnoreBlackListScanFilters);
 }
 
 void UAssetRegistryImpl::ScanFilesSynchronous(const TArray<FString>& InFilePaths, bool bForceRescan)
 {
-	return ScanPathsSynchronousInternal(TArray<FString>(), InFilePaths, bForceRescan, false /* bIgnoreBlackListScanFilters */, nullptr /* OutFindAssets */);
+	ScanPathsSynchronousInternal(TArray<FString>(), InFilePaths, bForceRescan,
+		false /* bIgnoreBlackListScanFilters */);
+}
+
+void UAssetRegistryImpl::ScanPathsSynchronousInternal(const TArray<FString>& InDirs, const TArray<FString>& InFiles,
+	bool bInForceRescan, bool bInIgnoreBlackListScanFilters)
+{
+	TRACE_CPUPROFILER_EVENT_SCOPE(UAssetRegistryImpl::ScanPathsSynchronousInternal);
+	const double SearchStartTime = FPlatformTime::Seconds();
+
+	UE::AssetRegistry::Impl::FEventContext EventContext;
+	UE::AssetRegistry::Impl::FScanPathContext Context(EventContext, InDirs, InFiles,
+		bInForceRescan, bInIgnoreBlackListScanFilters, nullptr /* OutFindAssets */);
+
+	{
+		FWriteScopeLock InterfaceScopeLock(InterfaceLock);
+		GuardedData.ScanPathsSynchronous(Context);
+	}
+
+#if WITH_EDITOR
+	ProcessLoadedAssetsToUpdateCache(EventContext, -1., Context.bIdle);
+#endif
+	Broadcast(EventContext);
+
+	// Log stats
+	FString PathsString;
+	if (Context.LocalPaths.Num() > 1)
+	{
+		PathsString = FString::Printf(TEXT("'%s' and %d other paths"), *Context.LocalPaths[0], Context.LocalPaths.Num() - 1);
+	}
+	else
+	{
+		PathsString = FString::Printf(TEXT("'%s'"), *Context.LocalPaths[0]);
+	}
+
+	UE_LOG(LogAssetRegistry, Verbose, TEXT("ScanPathsSynchronous completed scanning %s to find %d assets in %0.4f seconds"), *PathsString,
+		Context.NumFoundAssets, FPlatformTime::Seconds() - SearchStartTime);
 }
 
 void UAssetRegistryImpl::PrioritizeSearchPath(const FString& PathToPrioritize)
+{
+	FWriteScopeLock InterfaceScopeLock(InterfaceLock);
+	GuardedData.PrioritizeSearchPath(PathToPrioritize);
+}
+
+namespace UE::AssetRegistry
+{
+
+void FAssetRegistryImpl::PrioritizeSearchPath(const FString& PathToPrioritize)
 {
 	if (!GlobalGatherer.IsValid())
 	{
@@ -1866,6 +2543,8 @@ void UAssetRegistryImpl::PrioritizeSearchPath(const FString& PathToPrioritize)
 	}
 }
 
+}
+
 void UAssetRegistryImpl::AssetCreated(UObject* NewAsset)
 {
 	if (ensure(NewAsset) && NewAsset->IsAsset())
@@ -1880,14 +2559,22 @@ void UAssetRegistryImpl::AssetCreated(UObject* NewAsset)
 
 		const FString NewPackageName = NewPackage->GetName();
 
-		// This package not empty, in case it ever was
-		RemoveEmptyPackage(NewPackage->GetFName());
+		bool bShouldSkipAssset;
+		UE::AssetRegistry::Impl::FEventContext EventContext;
+		{
+			FWriteScopeLock InterfaceScopeLock(InterfaceLock);
+			// If this package was marked as an empty package before, it is no longer empty, so remove it from the list
+			GuardedData.RemoveEmptyPackage(NewPackage->GetFName());
 
-		// Add the path to the Path Tree, in case it wasn't already there
-		AddAssetPath(*FPackageName::GetLongPackagePath(NewPackageName));
+			// Add the path to the Path Tree, in case it wasn't already there
+			GuardedData.AddAssetPath(EventContext, *FPackageName::GetLongPackagePath(NewPackageName));
+			bShouldSkipAssset = GuardedData.ShouldSkipAsset(NewAsset);
+		}
 
-		if (!UE::AssetRegistry::FFiltering::ShouldSkipAsset(NewAsset))
+		Broadcast(EventContext);
+		if (!bShouldSkipAssset)
 		{
+			checkf(IsInGameThread(), TEXT("AssetCreated is not yet implemented as callable from other threads"));
 			// Let subscribers know that the new asset was added to the registry
 			AssetAddedEvent.Broadcast(FAssetData(NewAsset, true /* bAllowBlueprintClass */));
 
@@ -1903,17 +2590,20 @@ void UAssetRegistryImpl::AssetDeleted(UObject* DeletedAsset)
 	if (ensure(DeletedAsset) && DeletedAsset->IsAsset())
 	{
 		UPackage* DeletedObjectPackage = DeletedAsset->GetOutermost();
-		if (DeletedObjectPackage != nullptr)
+		bool bIsEmptyPackage = DeletedObjectPackage != nullptr && UPackage::IsEmptyPackage(DeletedObjectPackage, DeletedAsset);
+		bool bInitialSearchCompleted = false;
+
+		bool bShouldSkipAsset;
 		{
-			const FString PackageName = DeletedObjectPackage->GetName();
+			FWriteScopeLock InterfaceScopeLock(InterfaceLock);
 
 			// Deleting the last asset in a package causes the package to be garbage collected.
 			// If the UPackage object is GCed, it will be considered 'Unloaded' which will cause it to
 			// be fully loaded from disk when save is invoked.
 			// We want to keep the package around so we can save it empty or delete the file.
-			if (UPackage::IsEmptyPackage(DeletedObjectPackage, DeletedAsset))
+			if (bIsEmptyPackage)
 			{
-				AddEmptyPackage(DeletedObjectPackage->GetFName());
+				GuardedData.AddEmptyPackage(DeletedObjectPackage->GetFName());
 
 				// If there is a package metadata object, clear the standalone flag so the package can be truly emptied upon GC
 				if (UMetaData* MetaData = DeletedObjectPackage->GetMetaData())
@@ -1921,20 +2611,24 @@ void UAssetRegistryImpl::AssetDeleted(UObject* DeletedAsset)
 					MetaData->ClearFlags(RF_Standalone);
 				}
 			}
+			bInitialSearchCompleted = GuardedData.IsInitialSearchCompleted();
+			bShouldSkipAsset = GuardedData.ShouldSkipAsset(DeletedAsset);
 		}
 
-		FAssetData AssetDataDeleted = FAssetData(DeletedAsset, true /* bAllowBlueprintClass */);
-
 #if WITH_EDITOR
-		if (bInitialSearchCompleted && AssetDataDeleted.IsRedirector())
+		if (bInitialSearchCompleted && FAssetData::IsRedirector(DeletedAsset))
+
 		{
 			// Need to remove from GRedirectCollector
-			GRedirectCollector.RemoveAssetPathRedirection(AssetDataDeleted.ObjectPath);
+			GRedirectCollector.RemoveAssetPathRedirection(FName(*DeletedAsset->GetPathName()));
 		}
 #endif
 
-		if (!UE::AssetRegistry::FFiltering::ShouldSkipAsset(DeletedAsset))
+		if (!bShouldSkipAsset)
 		{
+			FAssetData AssetDataDeleted = FAssetData(DeletedAsset, true /* bAllowBlueprintClass */);
+
+			checkf(IsInGameThread(), TEXT("AssetDeleted is not yet implemented as callable from other threads"));
 			// Let subscribers know that the asset was removed from the registry
 			AssetRemovedEvent.Broadcast(AssetDataDeleted);
 
@@ -1956,26 +2650,35 @@ void UAssetRegistryImpl::AssetRenamed(const UObject* RenamedAsset, const FString
 		const FString NewPackageName = NewPackage->GetName();
 		const FString Filename = FPackageName::LongPackageNameToFilename(NewPackageName, FPackageName::GetAssetPackageExtension());
 
-		RemoveEmptyPackage(NewPackage->GetFName());
-
 		// We want to keep track of empty packages so we can properly merge cached assets with in-memory assets
+		UPackage* OldPackage = nullptr;
 		FString OldPackageName;
 		FString OldAssetName;
 		if (OldObjectPath.Split(TEXT("."), &OldPackageName, &OldAssetName))
 		{
-			UPackage* OldPackage = FindPackage(nullptr, *OldPackageName);
+			OldPackage = FindPackage(nullptr, *OldPackageName);
+		}
+
+		bool bShouldSkipAsset;
+		UE::AssetRegistry::Impl::FEventContext EventContext;
+		{
+			FWriteScopeLock InterfaceScopeLock(InterfaceLock);
+			GuardedData.RemoveEmptyPackage(NewPackage->GetFName());
 
-			if (UPackage::IsEmptyPackage(OldPackage))
+			if (OldPackage && UPackage::IsEmptyPackage(OldPackage))
 			{
-				AddEmptyPackage(OldPackage->GetFName());
+				GuardedData.AddEmptyPackage(OldPackage->GetFName());
 			}
-		}
 
-		// Add the path to the Path Tree, in case it wasn't already there
-		AddAssetPath(*FPackageName::GetLongPackagePath(NewPackageName));
+			// Add the path to the Path Tree, in case it wasn't already there
+			GuardedData.AddAssetPath(EventContext, *FPackageName::GetLongPackagePath(NewPackageName));
+			bShouldSkipAsset = GuardedData.ShouldSkipAsset(RenamedAsset);
+		}
 
-		if (!UE::AssetRegistry::FFiltering::ShouldSkipAsset(RenamedAsset))
+		Broadcast(EventContext);
+		if (!bShouldSkipAsset)
 		{
+			checkf(IsInGameThread(), TEXT("AssetRenamed is not yet implemented as callable from other threads"));
 			AssetRenamedEvent.Broadcast(FAssetData(RenamedAsset, true /* bAllowBlueprintClass */), OldObjectPath);
 		}
 	}
@@ -1984,36 +2687,100 @@ void UAssetRegistryImpl::AssetRenamed(const UObject* RenamedAsset, const FString
 void UAssetRegistryImpl::PackageDeleted(UPackage* DeletedPackage)
 {
 	checkf(GIsEditor, TEXT("Updating the AssetRegistry is only available in editor"));
+	UE::AssetRegistry::Impl::FEventContext EventContext;
 	if (ensure(DeletedPackage))
 	{
-		RemovePackageData(*DeletedPackage->GetName());
+		FWriteScopeLock InterfaceScopeLock(InterfaceLock);
+		GuardedData.RemovePackageData(EventContext, DeletedPackage->GetFName());
 	}
+	Broadcast(EventContext);
 }
 
 bool UAssetRegistryImpl::IsLoadingAssets() const
+{
+	FReadScopeLock InterfaceScopeLock(InterfaceLock);
+	return GuardedData.IsLoadingAssets();
+}
+
+namespace UE::AssetRegistry
+{
+
+bool FAssetRegistryImpl::IsLoadingAssets() const
 {
 	return !bInitialSearchCompleted;
 }
 
+}
+
 void UAssetRegistryImpl::Tick(float DeltaTime)
 {
-	TickGatherer(DeltaTime);
+	checkf(IsInGameThread(), TEXT("The tick function executes deferred loads and events and must be on the game thread to do so."));
+
+	bool bIdle = false;
+	double TickStartTime = -1; // Force a full flush if DeltaTime < 0
+	if (DeltaTime >= 0)
+	{
+		TickStartTime = FPlatformTime::Seconds();
+	}
+
+
+	bool bInterrupted;
+	do
+	{
+		bInterrupted = false;
+		UE::AssetRegistry::Impl::FEventContext EventContext;
+		{
+			FWriteScopeLock InterfaceScopeLock(InterfaceLock);
+
+			// Process any deferred events and deletes
+			PRAGMA_DISABLE_DEPRECATION_WARNINGS
+			GuardedData.TickDeletes();
+			PRAGMA_ENABLE_DEPRECATION_WARNINGS
+			EventContext = MoveTemp(DeferredEvents);
+			DeferredEvents.Clear();
+			if (EventContext.IsEmpty())
+			{
+				// Tick the Gatherer
+				GuardedData.TickGatherer(EventContext, TickStartTime, bIdle, bInterrupted);
+			}
+			else
+			{
+				// Skip the TickGather to deal with the DeferredEvents first
+				bInterrupted = true;
+			}
+		}
+
+#if WITH_EDITOR
+		if (!bInterrupted)
+		{
+			ProcessLoadedAssetsToUpdateCache(EventContext, TickStartTime, bIdle);
+		}
+#endif
+		Broadcast(EventContext);
+	} while (bInterrupted && (TickStartTime < 0 || (FPlatformTime::Seconds() - TickStartTime) <= UE::AssetRegistry::Impl::MaxSecondsPerFrame));
 }
 
-void UAssetRegistryImpl::TickGatherer(float DeltaTime, TOptional<UAssetRegistryImpl::FAssetsFoundCallback> AssetsFoundCallback)
+namespace UE::AssetRegistry
 {
-	// Note this function can be reentered due to arbitrary code execution in construction of FAssetData
-	if (!GlobalGatherer.IsValid())
+
+void FAssetRegistryImpl::TickDeletes()
+{
+	PRAGMA_DISABLE_DEPRECATION_WARNINGS
+	for (TUniqueFunction<void()>& Action : DeleteActions)
 	{
-		return;
+		Action();
 	}
+	DeleteActions.Empty();
+	PRAGMA_ENABLE_DEPRECATION_WARNINGS
+}
 
-	double TickStartTime = FPlatformTime::Seconds();
-
-	if (DeltaTime < 0)
+void FAssetRegistryImpl::TickGatherer(Impl::FEventContext& EventContext, const double TickStartTime, bool& bOutIdle, bool& bOutInterrupted, TOptional<FAssetsFoundCallback> AssetsFoundCallback)
+{
+	bOutIdle = true;
+	bOutInterrupted = false;
+	if (!GlobalGatherer.IsValid())
 	{
-		// Force a full flush
-		TickStartTime = -1;
+		return;
 	}
 
 	// Gather results from the background search
@@ -2028,20 +2795,37 @@ void UAssetRegistryImpl::TickGatherer(float DeltaTime, TOptional<UAssetRegistryI
 	{
 		UE_LOG(LogAssetRegistry, Verbose, TEXT("### Background search completed in %0.4f seconds"), SearchTimes[SearchTimeIdx]);
 	}
+	const bool bHadAssetsToProcess = BackgroundAssetResults.Num() > 0 || BackgroundDependencyResults.Num() > 0;
+	int32 NumPending = 0;
+	auto UpdateStatus = [&NumPending, &bIsSearching, &bHadAssetsToProcess, &bIsDiscoveringFiles, &EventContext, &bOutIdle, &NumFilesToSearch, &NumPathsToSearch, this](bool bInIdle)
+	{
+		// Compute total pending, plus highest pending for this run so we can show a good progress bar
+		NumPending = NumFilesToSearch + NumPathsToSearch + BackgroundPathResults.Num() + BackgroundAssetResults.Num() + BackgroundDependencyResults.Num() + BackgroundCookedPackageNamesWithoutAssetDataResults.Num();
+		HighestPending = FMath::Max(this->HighestPending, NumPending);
+
+		// Notify the status change
+		if (bIsSearching || bHadAssetsToProcess)
+		{
+			EventContext.ProgressUpdateData.Emplace(
+				HighestPending,					// NumTotalAssets
+				HighestPending - NumPending,	// NumAssetsProcessedByAssetRegistry
+				NumPending / 2,					// NumAssetsPendingDataLoad, divided by 2 because assets are double counted due to dependencies
+				bIsDiscoveringFiles				// bIsDiscoveringAssetFiles
+			);
+		}
+
+		bOutIdle = bInIdle;
+		this->bGatherIdle = bInIdle;
+	};
 
-	// Go into temporary caching mode to speed up class queries while adding assets
-	bool bOldTemporaryCachingMode = GetTemporaryCachingMode();
-	SetTemporaryCachingMode(true);
-	ON_SCOPE_EXIT{ SetTemporaryCachingMode(bOldTemporaryCachingMode); };
 
 	// Add discovered paths
 	if (BackgroundPathResults.Num())
 	{
-		PathDataGathered(TickStartTime, BackgroundPathResults);
+		PathDataGathered(EventContext, TickStartTime, BackgroundPathResults);
 	}
 
 	// Process the asset results
-	const bool bHadAssetsToProcess = BackgroundAssetResults.Num() > 0 || BackgroundDependencyResults.Num() > 0;
 	if (BackgroundAssetResults.Num())
 	{
 		// Mark the first amortize time
@@ -2054,7 +2838,7 @@ void UAssetRegistryImpl::TickGatherer(float DeltaTime, TOptional<UAssetRegistryI
 			AssetsFoundCallback.GetValue()(BackgroundAssetResults);
 		}
 
-		AssetSearchDataGathered(TickStartTime, BackgroundAssetResults);
+		AssetSearchDataGathered(EventContext, TickStartTime, BackgroundAssetResults);
 
 		if (BackgroundAssetResults.Num() == 0)
 		{
@@ -2072,31 +2856,18 @@ void UAssetRegistryImpl::TickGatherer(float DeltaTime, TOptional<UAssetRegistryI
 	// Load cooked packages that do not have asset data
 	if (BackgroundCookedPackageNamesWithoutAssetDataResults.Num())
 	{
-		CookedPackageNamesWithoutAssetDataGathered(TickStartTime, BackgroundCookedPackageNamesWithoutAssetDataResults);
-	}
-
-	// Compute total pending, plus highest pending for this run so we can show a good progress bar
-	static int32 HighestPending = 0;
-	const int32 NumPending = NumFilesToSearch + NumPathsToSearch + BackgroundPathResults.Num() + BackgroundAssetResults.Num() + BackgroundDependencyResults.Num() + BackgroundCookedPackageNamesWithoutAssetDataResults.Num();
-
-	HighestPending = FMath::Max(HighestPending, NumPending);
-
-	// Notify the status change
-	if (bIsSearching || bHadAssetsToProcess)
-	{
-		const FFileLoadProgressUpdateData ProgressUpdateData(
-			HighestPending,					// NumTotalAssets
-			HighestPending - NumPending,	// NumAssetsProcessedByAssetRegistry
-			NumPending / 2,					// NumAssetsPendingDataLoad, divided by 2 because assets are double counted due to dependencies
-			bIsDiscoveringFiles				// bIsDiscoveringAssetFiles
-		);
-		FileLoadProgressUpdatedEvent.Broadcast(ProgressUpdateData);
+		CookedPackageNamesWithoutAssetDataGathered(EventContext, TickStartTime, BackgroundCookedPackageNamesWithoutAssetDataResults, bOutInterrupted);
+		if (bOutInterrupted)
+		{
+			UpdateStatus(false /* bInIdle */);
+			return;
+		}
 	}
 
 	// If completing an initial search, refresh the content browser
 	bool bIsIdle = !bIsSearching && NumPending == 0;
-	bool bBecameIdle = !bGatherIdle && bIsIdle;
-	bGatherIdle = bIsIdle;
+	UpdateStatus(bIsIdle);
+
 	if (bIsIdle)
 	{
 		HighestPending = 0;
@@ -2112,24 +2883,19 @@ void UAssetRegistryImpl::TickGatherer(float DeltaTime, TOptional<UAssetRegistryI
 
 			bInitialSearchCompleted = true;
 
-			FileLoadedEvent.Broadcast();
+			EventContext.bFileLoadedEventBroadcast = true;
 		}
 	}
-#if WITH_EDITOR
-	if (bUpdateDiskCacheAfterLoad && (bIsIdle || TickStartTime < 0))
-	{
-		ProcessLoadedAssetsToUpdateCache(TickStartTime, bBecameIdle);
-	}
-#endif
 }
 
-void UAssetRegistryImpl::TickGatherPackage(const FString& PackageName)
+void FAssetRegistryImpl::TickGatherPackage(Impl::FEventContext& EventContext, const FString& PackageName, const FString& LocalPath)
 {
-	check(IsInGameThread()); // AssetRegistry is not threadsafe
 	if (!GlobalGatherer.IsValid())
 	{
 		return;
 	}
+	GlobalGatherer->WaitOnPath(LocalPath);
+
 	FName PackageFName(PackageName);
 
 	// Gather results from the background search
@@ -2159,7 +2925,7 @@ void UAssetRegistryImpl::TickGatherPackage(const FString& PackageName)
 	}
 	if (PackageAssets.Num() > 0)
 	{
-		AssetSearchDataGathered(-1., PackageAssets);
+		AssetSearchDataGathered(EventContext, -1., PackageAssets);
 	}
 	if (PackageDependencyDatas.Num() > 0)
 	{
@@ -2167,7 +2933,22 @@ void UAssetRegistryImpl::TickGatherPackage(const FString& PackageName)
 	}
 }
 
-void UAssetRegistryImpl::Serialize(FArchive& Ar)
+}
+
+void UAssetRegistryImpl::Serialize(FArchive& Ar)
+{
+	UE::AssetRegistry::Impl::FEventContext EventContext;
+	{
+		FWriteScopeLock InterfaceScopeLock(InterfaceLock);
+		GuardedData.Serialize(Ar, EventContext);
+	}
+	Broadcast(EventContext);
+}
+
+namespace UE::AssetRegistry
+{
+
+void FAssetRegistryImpl::Serialize(FArchive& Ar, Impl::FEventContext& EventContext)
 {
 	if (Ar.IsObjectReferenceCollector())
 	{
@@ -2177,7 +2958,7 @@ void UAssetRegistryImpl::Serialize(FArchive& Ar)
 	else if (Ar.IsLoading())
 	{
 		State.Load(Ar);
-		CachePathsFromState(State);
+		CachePathsFromState(EventContext, State);
 	}
 	else if (Ar.IsSaving())
 	{
@@ -2185,12 +2966,19 @@ void UAssetRegistryImpl::Serialize(FArchive& Ar)
 	}
 }
 
+}
+
 /** Append the assets from the incoming state into our own */
 void UAssetRegistryImpl::AppendState(const FAssetRegistryState& InState)
 {
-	State.InitializeFromExisting(InState, SerializationOptions, FAssetRegistryState::EInitializationMode::Append);
-	CachePathsFromState(InState);
+	UE::AssetRegistry::Impl::FEventContext EventContext;
+	{
+		FWriteScopeLock InterfaceScopeLock(InterfaceLock);
+		GuardedData.AppendState(EventContext, InState);
+	}
 
+	Broadcast(EventContext);
+	checkf(IsInGameThread(), TEXT("AppendState is not yet implemented as callable from other threads"));
 	InState.EnumerateAllAssets(TSet<FName>(), [this](const FAssetData& AssetData)
 	{
 		// Let subscribers know that the new asset was added to the registry
@@ -2199,7 +2987,16 @@ void UAssetRegistryImpl::AppendState(const FAssetRegistryState& InState)
 	}, true /*bARFiltering*/);
 }
 
-void UAssetRegistryImpl::CachePathsFromState(const FAssetRegistryState& InState)
+namespace UE::AssetRegistry
+{
+
+void FAssetRegistryImpl::AppendState(Impl::FEventContext& EventContext, const FAssetRegistryState& InState)
+{
+	State.InitializeFromExisting(InState, SerializationOptions, FAssetRegistryState::EInitializationMode::Append);
+	CachePathsFromState(EventContext, InState);
+}
+
+void FAssetRegistryImpl::CachePathsFromState(Impl::FEventContext& EventContext, const FAssetRegistryState& InState)
 {
 	LLM_SCOPE(ELLMTag::AssetRegistry);
 
@@ -2213,7 +3010,7 @@ void UAssetRegistryImpl::CachePathsFromState(const FAssetRegistryState& InState)
 
 		if (AssetData != nullptr)
 		{
-			AddAssetPath(AssetData->PackagePath);
+			AddAssetPath(EventContext, AssetData->PackagePath);
 
 			// Populate the class map if adding blueprint
 			if (ClassGeneratorNames.Contains(AssetData->AssetClass))
@@ -2226,23 +3023,51 @@ void UAssetRegistryImpl::CachePathsFromState(const FAssetRegistryState& InState)
 					CachedBPInheritanceMap.Add(GeneratedClass.Object, ParentClass.Object);
 
 					// Invalidate caching because CachedBPInheritanceMap got modified
-					bIsTempCachingUpToDate = false;
+					TempCachedInheritanceBuffer.bDirty = true;
 				}
 			}
 		}
 	}
 }
 
+}
+
 uint32 UAssetRegistryImpl::GetAllocatedSize(bool bLogDetailed) const
 {
-	uint32 StateSize = State.GetAllocatedSize(bLogDetailed);
+	uint32 StateSize = 0;
+	uint32 StaticSize = 0;
+	uint32 SearchSize = 0;
+	{
+		FReadScopeLock InterfaceScopeLock(InterfaceLock);
+		GuardedData.GetAllocatedSize(bLogDetailed, StateSize, StaticSize, SearchSize);
+		StaticSize += sizeof(UAssetRegistryImpl);
+#if WITH_EDITOR
+		StaticSize += OnDirectoryChangedDelegateHandles.GetAllocatedSize();
+#endif
+	}
+
+	if (bLogDetailed)
+	{
+		UE_LOG(LogAssetRegistry, Log, TEXT("AssetRegistry Static Size: %dk"), StaticSize / 1024);
+		UE_LOG(LogAssetRegistry, Log, TEXT("AssetRegistry Search Size: %dk"), SearchSize / 1024);
+	}
+
+	return StateSize + StaticSize + SearchSize;
+}
+
+namespace UE::AssetRegistry
+{
+
+void FAssetRegistryImpl::GetAllocatedSize(bool bLogDetailed, uint32& StateSize, uint32& StaticSize, uint32& SearchSize) const
+{
+	StateSize = State.GetAllocatedSize(bLogDetailed);
 
-	uint32 StaticSize = sizeof(UAssetRegistryImpl) + CachedEmptyPackages.GetAllocatedSize() + CachedBPInheritanceMap.GetAllocatedSize()  + ClassGeneratorNames.GetAllocatedSize() + OnDirectoryChangedDelegateHandles.GetAllocatedSize();
-	uint32 SearchSize = BackgroundAssetResults.GetAllocatedSize() + BackgroundPathResults.GetAllocatedSize() + BackgroundDependencyResults.GetAllocatedSize() + BackgroundCookedPackageNamesWithoutAssetDataResults.GetAllocatedSize() + CachedPathTree.GetAllocatedSize();
+	StaticSize = CachedEmptyPackages.GetAllocatedSize() + CachedBPInheritanceMap.GetAllocatedSize() + ClassGeneratorNames.GetAllocatedSize();
+	SearchSize = BackgroundAssetResults.GetAllocatedSize() + BackgroundPathResults.GetAllocatedSize() + BackgroundDependencyResults.GetAllocatedSize() + BackgroundCookedPackageNamesWithoutAssetDataResults.GetAllocatedSize() + CachedPathTree.GetAllocatedSize();
 
 	if (bIsTempCachingEnabled && !bIsTempCachingAlwaysEnabled)
 	{
-		uint32 TempCacheMem = TempCachedInheritanceMap.GetAllocatedSize() + TempReverseInheritanceMap.GetAllocatedSize();
+		uint32 TempCacheMem = TempCachedInheritanceBuffer.GetAllocatedSize();
 		StaticSize += TempCacheMem;
 		UE_LOG(LogAssetRegistry, Warning, TEXT("Asset Registry Temp caching enabled, wasting memory: %dk"), TempCacheMem / 1024);
 	}
@@ -2258,14 +3083,8 @@ uint32 UAssetRegistryImpl::GetAllocatedSize(bool bLogDetailed) const
 	{
 		StaticSize += Pair.Value.GetAllocatedSize();
 	}
+}
 
-	if (bLogDetailed)
-	{
-		UE_LOG(LogAssetRegistry, Log, TEXT("AssetRegistry Static Size: %dk"), StaticSize / 1024);
-		UE_LOG(LogAssetRegistry, Log, TEXT("AssetRegistry Search Size: %dk"), SearchSize / 1024);
-	}
-
-	return StateSize + StaticSize + SearchSize;
 }
 
 void UAssetRegistryImpl::LoadPackageRegistryData(FArchive& Ar, FLoadPackageRegistryData& InOutData) const
@@ -2273,7 +3092,7 @@ void UAssetRegistryImpl::LoadPackageRegistryData(FArchive& Ar, FLoadPackageRegis
 	FPackageReader Reader;
 	if (Reader.OpenPackageFile(&Ar))
 	{
-		ReadAssetFile(Reader, InOutData);
+		UE::AssetRegistry::Utils::ReadAssetFile(Reader, InOutData);
 	}
 }
 
@@ -2282,11 +3101,14 @@ void UAssetRegistryImpl::LoadPackageRegistryData(const FString& PackageFilename,
 	FPackageReader Reader;
 	if (Reader.OpenPackageFile(PackageFilename))
 	{
-		ReadAssetFile(Reader, InOutData);
+		UE::AssetRegistry::Utils::ReadAssetFile(Reader, InOutData);
 	}
 }
 
-bool UAssetRegistryImpl::ReadAssetFile(FPackageReader& PackageReader, FLoadPackageRegistryData& InOutData) const
+namespace UE::AssetRegistry::Utils
+{
+
+bool ReadAssetFile(FPackageReader& PackageReader, IAssetRegistry::FLoadPackageRegistryData& InOutData)
 {
 	TArray<FAssetData*> AssetDataList;
 	TArray<FString> CookedPackageNamesWithoutAssetDataGathered;
@@ -2326,34 +3148,71 @@ bool UAssetRegistryImpl::ReadAssetFile(FPackageReader& PackageReader, FLoadPacka
 	return bReadOk;
 }
 
+}
+
 void UAssetRegistryImpl::InitializeTemporaryAssetRegistryState(FAssetRegistryState& OutState, const FAssetRegistrySerializationOptions& Options, bool bRefreshExisting, const TMap<FName, FAssetData*>& OverrideData) const
 {
+	FReadScopeLock InterfaceScopeLock(InterfaceLock);
+	const FAssetRegistryState& State = GuardedData.GetState();
 	const TMap<FName, FAssetData*>& DataToUse = OverrideData.Num() > 0 ? OverrideData : State.CachedAssetsByObjectPath;
 
 	OutState.InitializeFromExisting(DataToUse, State.CachedDependsNodes, State.CachedPackageData, Options, bRefreshExisting ? FAssetRegistryState::EInitializationMode::OnlyUpdateExisting : FAssetRegistryState::EInitializationMode::Rebuild);
 }
 
+#if ASSET_REGISTRY_STATE_DUMPING_ENABLED
+void UAssetRegistryImpl::DumpState(const TArray<FString>& Arguments, TArray<FString>& OutPages, int32 LinesPerPage) const
+{
+	FReadScopeLock InterfaceScopeLock(InterfaceLock);
+	GuardedData.GetState().Dump(Arguments, OutPages, LinesPerPage);
+}
+#endif
+
 const FAssetRegistryState* UAssetRegistryImpl::GetAssetRegistryState() const
 {
-	return &State;
+	FReadScopeLock InterfaceScopeLock(InterfaceLock);
+	PRAGMA_DISABLE_DEPRECATION_WARNINGS
+	return &GuardedData.GetState();
+	PRAGMA_ENABLE_DEPRECATION_WARNINGS
+}
+
+TSet<FName> UAssetRegistryImpl::GetCachedEmptyPackagesCopy() const
+{
+	FReadScopeLock InterfaceScopeLock(InterfaceLock);
+	return GuardedData.GetCachedEmptyPackages();
 }
 
 const TSet<FName>& UAssetRegistryImpl::GetCachedEmptyPackages() const
 {
-	return CachedEmptyPackages;
+	FReadScopeLock InterfaceScopeLock(InterfaceLock);
+	PRAGMA_DISABLE_DEPRECATION_WARNINGS
+	return GuardedData.GetCachedEmptyPackages();
+	PRAGMA_ENABLE_DEPRECATION_WARNINGS
 }
 
-void UAssetRegistryImpl::ScanPathsSynchronousInternal(const TArray<FString>& InDirs, const TArray<FString>& InFiles, bool bForceRescan, bool bIgnoreBlackListScanFilters, TArray<FName>* OutFoundAssets)
+bool UAssetRegistryImpl::ContainsTag(FName TagName) const
 {
-	TRACE_CPUPROFILER_EVENT_SCOPE(UAssetRegistryImpl::ScanPathsSynchronousInternal);
+	FReadScopeLock InterfaceScopeLock(InterfaceLock);
+	return GuardedData.GetState().GetTagToAssetDatasMap().Contains(TagName);
+}
 
-	// Note this function can be reentered due to arbitrary code execution in construction of FAssetData
-	LLM_SCOPE(ELLMTag::AssetRegistry);
-	const double SearchStartTime = FPlatformTime::Seconds();
+namespace UE::AssetRegistry
+{
+
+namespace Impl
+{
+
+FScanPathContext::FScanPathContext(FEventContext& InEventContext, const TArray<FString>& InDirs, const TArray<FString>& InFiles,
+	bool bInForceRescan, bool bInIgnoreBlackListScanFilters, TArray<FName>* FoundAssets)
+	: EventContext(InEventContext)
+	, OutFoundAssets(FoundAssets)
+	, bForceRescan(bInForceRescan)
+	, bIgnoreBlackListScanFilters(bInIgnoreBlackListScanFilters)
+{
 	if (OutFoundAssets)
 	{
 		OutFoundAssets->Empty();
 	}
+
 	if (bIgnoreBlackListScanFilters && !bForceRescan)
 	{
 		// This restriction is necessary because we have not yet implemented some of the required behavior to handle bIgnoreBlackListScanFilters without bForceRescan;
@@ -2363,10 +3222,6 @@ void UAssetRegistryImpl::ScanPathsSynchronousInternal(const TArray<FString>& InD
 		bForceRescan = true;
 	}
 
-	TArray<FString> PackageDirs;
-	TArray<FString> LocalDirs;
-	TArray<FString> PackageFiles;
-	TArray<FString> LocalFiles;
 	FString LocalPath;
 	FString PackageName;
 	FString Extension;
@@ -2410,31 +3265,37 @@ void UAssetRegistryImpl::ScanPathsSynchronousInternal(const TArray<FString>& InD
 		LocalDirs.Add(LocalPath + Extension);
 		PackageDirs.Add(PackageName + Extension);
 	}
+}
+
+}
+
+void FAssetRegistryImpl::ScanPathsSynchronous(Impl::FScanPathContext& Context)
+{
+	LLM_SCOPE(ELLMTag::AssetRegistry);
 
 	ConstructGatherer();
 	FAssetDataGatherer& Gatherer = *GlobalGatherer;
 
 	// Add a cache file for any not-yet-scanned dirs
 	TArray<FString> CacheFilePackagePaths;
-	if (!bForceRescan && Gatherer.IsCacheEnabled())
+	if (!Context.bForceRescan && Gatherer.IsCacheEnabled())
 	{
-		for (int n = 0; n < LocalDirs.Num(); ++n)
+		for (int n = 0; n < Context.LocalDirs.Num(); ++n)
 		{
-			if (!Gatherer.IsWhitelisted(LocalDirs[n]))
+			if (!Gatherer.IsWhitelisted(Context.LocalDirs[n]))
 			{
-				CacheFilePackagePaths.Add(PackageDirs[n]);
+				CacheFilePackagePaths.Add(Context.PackageDirs[n]);
 			}
 		}
 	}
-	TArray<FString> LocalPaths;
-	LocalPaths.Reserve(LocalFiles.Num() + LocalDirs.Num());
-	LocalPaths.Append(MoveTemp(LocalDirs));
-	LocalPaths.Append(MoveTemp(LocalFiles));
-	if (LocalPaths.IsEmpty())
+	Context.LocalPaths.Reserve(Context.LocalFiles.Num() + Context.LocalDirs.Num());
+	Context.LocalPaths.Append(MoveTemp(Context.LocalDirs));
+	Context.LocalPaths.Append(MoveTemp(Context.LocalFiles));
+	if (Context.LocalPaths.IsEmpty())
 	{
 		return;
 	}
-	Gatherer.AddRequiredMountPoints(LocalPaths);
+	Gatherer.AddRequiredMountPoints(Context.LocalPaths);
 
 	FString CacheFilename;
 	if (!CacheFilePackagePaths.IsEmpty())
@@ -2443,20 +3304,19 @@ void UAssetRegistryImpl::ScanPathsSynchronousInternal(const TArray<FString>& InD
 		Gatherer.LoadCacheFile(CacheFilename);
 	}
 
-	Gatherer.ScanPathsSynchronous(LocalPaths, bForceRescan, bIgnoreBlackListScanFilters, CacheFilename, PackageDirs);
+	Gatherer.ScanPathsSynchronous(Context.LocalPaths, Context.bForceRescan, Context.bIgnoreBlackListScanFilters, CacheFilename, Context.PackageDirs);
 
-	int32 NumFoundAssets = 0;
-	auto AssetsFoundCallback = [&PackageFiles, &PackageDirs, &OutFoundAssets, &NumFoundAssets](const TRingBuffer<FAssetData*>& InFoundAssets)
+	auto AssetsFoundCallback = [&Context](const TRingBuffer<FAssetData*>& InFoundAssets)
 	{
-		NumFoundAssets = InFoundAssets.Num();
+		Context.NumFoundAssets = InFoundAssets.Num();
 
-		if (!OutFoundAssets)
+		if (!Context.OutFoundAssets)
 		{
 			return;
 		}
 
-		OutFoundAssets->Reset();
-		OutFoundAssets->Reserve(NumFoundAssets);
+		Context.OutFoundAssets->Reset();
+		Context.OutFoundAssets->Reserve(Context.NumFoundAssets);
 
 		// The gatherer may have added other assets that were scanned as part of the ongoing background scan; remove any assets that were not in the requested paths
 		for (FAssetData* AssetData : InFoundAssets)
@@ -2467,7 +3327,7 @@ void UAssetRegistryImpl::ScanPathsSynchronousInternal(const TArray<FString>& InD
 			AssetData->PackageName.ToString(PackageNameStr);
 			FStringView PackageName(PackageNameStr.ToString(), PackageNameStr.Len());
 
-			for (const FString& RequestedPackageDir : PackageDirs)
+			for (const FString& RequestedPackageDir : Context.PackageDirs)
 			{
 				if (FPathViews::IsParentPathOf(RequestedPackageDir, PackageName))
 				{
@@ -2478,7 +3338,7 @@ void UAssetRegistryImpl::ScanPathsSynchronousInternal(const TArray<FString>& InD
 
 			if (!bIsInRequestedPaths)
 			{
-				for (const FString& RequestedPackageFile : PackageFiles)
+				for (const FString& RequestedPackageFile : Context.PackageFiles)
 				{
 					if (PackageName.Equals(RequestedPackageFile, ESearchCase::IgnoreCase))
 					{
@@ -2490,28 +3350,19 @@ void UAssetRegistryImpl::ScanPathsSynchronousInternal(const TArray<FString>& InD
 
 			if (bIsInRequestedPaths)
 			{
-				OutFoundAssets->Add(AssetData->ObjectPath);
+				Context.OutFoundAssets->Add(AssetData->ObjectPath);
 			}
 		}
 	};
 
-	TickGatherer(-1.f, FAssetsFoundCallback(AssetsFoundCallback));
-
-	// Log stats
-	FString PathsString;
-	if (LocalPaths.Num() > 1)
-	{
-		PathsString = FString::Printf(TEXT("'%s' and %d other paths"), *LocalPaths[0], LocalPaths.Num() - 1);
-	}
-	else
-	{
-		PathsString = FString::Printf(TEXT("'%s'"), *LocalPaths[0]);
-	}
-
-	UE_LOG(LogAssetRegistry, Verbose, TEXT("ScanPathsSynchronous completed scanning %s to find %d assets in %0.4f seconds"), *PathsString, NumFoundAssets, FPlatformTime::Seconds() - SearchStartTime);
+	bool bUnusedInterrupted;
+	TickGatherer(Context.EventContext, -1., Context.bIdle, bUnusedInterrupted, FAssetsFoundCallback(AssetsFoundCallback));
 }
 
-bool UAssetRegistryImpl::IsPathMounted(const FString& Path, const TSet<FString>& MountPointsNoTrailingSlashes, FString& StringBuffer) const
+namespace Utils
+{
+
+bool IsPathMounted(const FString& Path, const TSet<FString>& MountPointsNoTrailingSlashes, FString& StringBuffer)
 {
 	const int32 SecondSlash = Path.Len() > 1 ? Path.Find(TEXT("/"), ESearchCase::CaseSensitive, ESearchDir::FromStart, 1) : INDEX_NONE;
 	if (SecondSlash != INDEX_NONE)
@@ -2534,7 +3385,9 @@ bool UAssetRegistryImpl::IsPathMounted(const FString& Path, const TSet<FString>&
 	return false;
 }
 
-void UAssetRegistryImpl::AssetSearchDataGathered(const double TickStartTime, TRingBuffer<FAssetData*>& AssetResults)
+}
+
+void FAssetRegistryImpl::AssetSearchDataGathered(Impl::FEventContext& EventContext, const double TickStartTime, TRingBuffer<FAssetData*>& AssetResults)
 {
 	const bool bFlushFullBuffer = TickStartTime < 0;
 
@@ -2568,7 +3421,7 @@ void UAssetRegistryImpl::AssetSearchDataGathered(const double TickStartTime, TRi
 		if (bVerifyMountPointAfterGather)
 		{
 			PackagePath.ToString(PackagePathString);
-			if (!IsPathMounted(PackagePathString, MountPoints, PackageRoot))
+			if (!Utils::IsPathMounted(PackagePathString, MountPoints, PackageRoot))
 			{
 				bPathIsMounted = false;
 			}
@@ -2585,7 +3438,7 @@ void UAssetRegistryImpl::AssetSearchDataGathered(const double TickStartTime, TRi
 #endif
 				{
 					// The asset exists in the cache from disk and has not yet been loaded into memory, update it with the new background data
-					UpdateAssetData(AssetData, *BackgroundResult);
+					UpdateAssetData(EventContext, AssetData, *BackgroundResult);
 				}
 			}
 		}
@@ -2594,25 +3447,25 @@ void UAssetRegistryImpl::AssetSearchDataGathered(const double TickStartTime, TRi
 			// The asset isn't in the cache yet, add it and notify subscribers
 			if (bPathIsMounted)
 			{
-				AddAssetData(BackgroundResult.Release());
+				AddAssetData(EventContext, BackgroundResult.Release());
 			}
 		}
 
 		if (bPathIsMounted)
 		{
 			// Populate the path tree
-			AddAssetPath(PackagePath);
+			AddAssetPath(EventContext, PackagePath);
 		}
 
 		// Check to see if we have run out of time in this tick
-		if (!bFlushFullBuffer && (FPlatformTime::Seconds() - TickStartTime) > MaxSecondsPerFrame)
+		if (!bFlushFullBuffer && (FPlatformTime::Seconds() - TickStartTime) > Impl::MaxSecondsPerFrame)
 		{
 			return;
 		}
 	}
 }
 
-void UAssetRegistryImpl::PathDataGathered(const double TickStartTime, TRingBuffer<FString>& PathResults)
+void FAssetRegistryImpl::PathDataGathered(Impl::FEventContext& EventContext, const double TickStartTime, TRingBuffer<FString>& PathResults)
 {
 	const bool bFlushFullBuffer = TickStartTime < 0;
 
@@ -2630,20 +3483,20 @@ void UAssetRegistryImpl::PathDataGathered(const double TickStartTime, TRingBuffe
 		FString Path = PathResults.PopFrontValue();
 
 		// Skip stale results caused by mount then unmount of a path within short period.
-		if (!bVerifyMountPointAfterGather || IsPathMounted(Path, MountPoints, PackageRoot))
+		if (!bVerifyMountPointAfterGather || Utils::IsPathMounted(Path, MountPoints, PackageRoot))
 		{
-			AddAssetPath(FName(*Path));
+			AddAssetPath(EventContext, FName(*Path));
 		}
 
 		// Check to see if we have run out of time in this tick
-		if (!bFlushFullBuffer && (FPlatformTime::Seconds() - TickStartTime) > MaxSecondsPerFrame)
+		if (!bFlushFullBuffer && (FPlatformTime::Seconds() - TickStartTime) > Impl::MaxSecondsPerFrame)
 		{
 			return;
 		}
 	}
 }
 
-void UAssetRegistryImpl::DependencyDataGathered(const double TickStartTime, TRingBuffer<FPackageDependencyData>& DependsResults)
+void FAssetRegistryImpl::DependencyDataGathered(const double TickStartTime, TRingBuffer<FPackageDependencyData>& DependsResults)
 {
 	using namespace UE::AssetRegistry;
 	const bool bFlushFullBuffer = TickStartTime < 0;
@@ -2783,15 +3636,17 @@ void UAssetRegistryImpl::DependencyDataGathered(const double TickStartTime, TRin
 		}
 
 		// Check to see if we have run out of time in this tick
-		if (!bFlushFullBuffer && (FPlatformTime::Seconds() - TickStartTime) > MaxSecondsPerFrame)
+		if (!bFlushFullBuffer && (FPlatformTime::Seconds() - TickStartTime) > Impl::MaxSecondsPerFrame)
 		{
 			return;
 		}
 	}
 }
 
-void UAssetRegistryImpl::CookedPackageNamesWithoutAssetDataGathered(const double TickStartTime, TRingBuffer<FString>& CookedPackageNamesWithoutAssetDataResults)
+void FAssetRegistryImpl::CookedPackageNamesWithoutAssetDataGathered(Impl::FEventContext& EventContext,
+	const double TickStartTime, TRingBuffer<FString>& CookedPackageNamesWithoutAssetDataResults, bool& bOutInterrupted)
 {
+	bOutInterrupted = false;
 	const bool bFlushFullBuffer = TickStartTime < 0;
 
 	struct FConfigValue
@@ -2814,15 +3669,14 @@ void UAssetRegistryImpl::CookedPackageNamesWithoutAssetDataGathered(const double
 		while (CookedPackageNamesWithoutAssetDataResults.Num() > 0)
 		{
 			// If this data is cooked and it we couldn't find any asset in its export table then try to load the entire package 
-			// Loading the entire package will make all of its assets searchable through the in-memory scanning performed by GetAsseets
-			FString BackgroundResult = CookedPackageNamesWithoutAssetDataResults.PopFrontValue();
-			LoadPackage(nullptr, *BackgroundResult, 0);
-
-			// Check to see if we have run out of time in this tick
-			if (!bFlushFullBuffer && (FPlatformTime::Seconds() - TickStartTime) > MaxSecondsPerFrame)
-			{
-				return;
-			}
+			// Loading the entire package will make all of its assets searchable through the in-memory scanning performed by GetAssets
+			EventContext.RequiredLoads.Add(CookedPackageNamesWithoutAssetDataResults.PopFrontValue());
+		}
+		if (TickStartTime >= 0)
+		{
+			// If the tick is time-limited, signal an interruption now to process the loads before proceeding with any other tick steps
+			bOutInterrupted = true;
+			return;
 		}
 	}
 	else
@@ -2834,58 +3688,66 @@ void UAssetRegistryImpl::CookedPackageNamesWithoutAssetDataGathered(const double
 	}
 }
 
-void UAssetRegistryImpl::AddEmptyPackage(FName PackageName)
+void FAssetRegistryImpl::AddEmptyPackage(FName PackageName)
 {
 	CachedEmptyPackages.Add(PackageName);
 }
 
-bool UAssetRegistryImpl::RemoveEmptyPackage(FName PackageName)
+bool FAssetRegistryImpl::RemoveEmptyPackage(FName PackageName)
 {
 	return CachedEmptyPackages.Remove(PackageName) > 0;
 }
 
-bool UAssetRegistryImpl::AddAssetPath(FName PathToAdd)
+bool FAssetRegistryImpl::AddAssetPath(Impl::FEventContext& EventContext, FName PathToAdd)
 {
-	return CachedPathTree.CachePath(PathToAdd, [this](FName AddedPath)
+	return CachedPathTree.CachePath(PathToAdd, [this, &EventContext](FName AddedPath)
 	{
-		PathAddedEvent.Broadcast(AddedPath.ToString());
+		EventContext.PathEvents.Emplace(AddedPath.ToString(), Impl::FEventContext::EEvent::Added);
 	});
 }
 
-bool UAssetRegistryImpl::RemoveAssetPath(FName PathToRemove, bool bEvenIfAssetsStillExist)
+bool FAssetRegistryImpl::RemoveAssetPath(Impl::FEventContext& EventContext, FName PathToRemove, bool bEvenIfAssetsStillExist)
 {
 	if (!bEvenIfAssetsStillExist)
 	{
 		// Check if there were assets in the specified folder. You can not remove paths that still contain assets
-		TArray<FAssetData> AssetsInPath;
-		GetAssetsByPath(PathToRemove, AssetsInPath, true);
-		if (AssetsInPath.Num() > 0)
+		bool bHasAsset = false;
+		EnumerateAssetsByPathNoTags(PathToRemove, [&bHasAsset](const FAssetData&)
+			{
+				bHasAsset = true;
+				return false;
+			}, true /* bRecursive */, false /* bIncludeOnlyOnDiskAssets */);
+		if (bHasAsset)
 		{
 			// At least one asset still exists in the path. Fail the remove.
 			return false;
 		}
 	}
 
-	return CachedPathTree.RemovePath(PathToRemove, [this](FName RemovedPath)
+	return CachedPathTree.RemovePath(PathToRemove, [this, &EventContext](FName RemovedPath)
 	{
-		PathRemovedEvent.Broadcast(RemovedPath.ToString());
+		EventContext.PathEvents.Emplace(RemovedPath.ToString(), Impl::FEventContext::EEvent::Removed);
 	});
 }
 
-FString UAssetRegistryImpl::ExportTextPathToObjectName(const FString& InExportTextPath) const
+namespace Utils
+{
+
+FString ExportTextPathToObjectName(const FString& InExportTextPath)
 {
 	const FString ObjectPath = FPackageName::ExportTextPathToObjectPath(InExportTextPath);
 	return FPackageName::ObjectPathToObjectName(ObjectPath);
 }
 
-void UAssetRegistryImpl::AddAssetData(FAssetData* AssetData)
+}
+
+void FAssetRegistryImpl::AddAssetData(Impl::FEventContext& EventContext, FAssetData* AssetData)
 {
 	State.AddAssetData(AssetData);
 
-	if (!UE::AssetRegistry::FFiltering::ShouldSkipAsset(AssetData->AssetClass, AssetData->PackageFlags))
+	if (!ShouldSkipAsset(AssetData->AssetClass, AssetData->PackageFlags))
 	{
-		// Notify subscribers
-		AssetAddedEvent.Broadcast(*AssetData);
+		EventContext.AssetEvents.Emplace(*AssetData, Impl::FEventContext::EEvent::Added);
 	}
 
 	// Populate the class map if adding blueprint
@@ -2895,17 +3757,17 @@ void UAssetRegistryImpl::AddAssetData(FAssetData* AssetData)
 		const FString ParentClass = AssetData->GetTagValueRef<FString>(FBlueprintTags::ParentClassPath);
 		if (!GeneratedClass.IsEmpty() && !ParentClass.IsEmpty())
 		{
-			const FName GeneratedClassFName = *ExportTextPathToObjectName(GeneratedClass);
-			const FName ParentClassFName = *ExportTextPathToObjectName(ParentClass);
+			const FName GeneratedClassFName = *Utils::ExportTextPathToObjectName(GeneratedClass);
+			const FName ParentClassFName = *Utils::ExportTextPathToObjectName(ParentClass);
 			CachedBPInheritanceMap.Add(GeneratedClassFName, ParentClassFName);
 			
 			// Invalidate caching because CachedBPInheritanceMap got modified
-			bIsTempCachingUpToDate = false;
+			TempCachedInheritanceBuffer.bDirty = true;
 		}
 	}
 }
 
-void UAssetRegistryImpl::UpdateAssetData(FAssetData* AssetData, const FAssetData& NewAssetData)
+void FAssetRegistryImpl::UpdateAssetData(Impl::FEventContext& EventContext, FAssetData* AssetData, const FAssetData& NewAssetData)
 {
 	// Update the class map if updating a blueprint
 	if (ClassGeneratorNames.Contains(AssetData->AssetClass))
@@ -2913,44 +3775,43 @@ void UAssetRegistryImpl::UpdateAssetData(FAssetData* AssetData, const FAssetData
 		const FString OldGeneratedClass = AssetData->GetTagValueRef<FString>(FBlueprintTags::GeneratedClassPath);
 		if (!OldGeneratedClass.IsEmpty())
 		{
-			const FName OldGeneratedClassFName = *ExportTextPathToObjectName(OldGeneratedClass);
+			const FName OldGeneratedClassFName = *Utils::ExportTextPathToObjectName(OldGeneratedClass);
 			CachedBPInheritanceMap.Remove(OldGeneratedClassFName);
 
 			// Invalidate caching because CachedBPInheritanceMap got modified
-			bIsTempCachingUpToDate = false;
+			TempCachedInheritanceBuffer.bDirty = true;
 		}
 
 		const FString NewGeneratedClass = NewAssetData.GetTagValueRef<FString>(FBlueprintTags::GeneratedClassPath);
 		const FString NewParentClass = NewAssetData.GetTagValueRef<FString>(FBlueprintTags::ParentClassPath);
 		if (!NewGeneratedClass.IsEmpty() && !NewParentClass.IsEmpty())
 		{
-			const FName NewGeneratedClassFName = *ExportTextPathToObjectName(*NewGeneratedClass);
-			const FName NewParentClassFName = *ExportTextPathToObjectName(*NewParentClass);
+			const FName NewGeneratedClassFName = *Utils::ExportTextPathToObjectName(*NewGeneratedClass);
+			const FName NewParentClassFName = *Utils::ExportTextPathToObjectName(*NewParentClass);
 			CachedBPInheritanceMap.Add(NewGeneratedClassFName, NewParentClassFName);
 
 			// Invalidate caching because CachedBPInheritanceMap got modified
-			bIsTempCachingUpToDate = false;
+			TempCachedInheritanceBuffer.bDirty = true;
 		}
 	}
 
 	State.UpdateAssetData(AssetData, NewAssetData);
 	
-	if (!UE::AssetRegistry::FFiltering::ShouldSkipAsset(AssetData->AssetClass, AssetData->PackageFlags))
+	if (!ShouldSkipAsset(AssetData->AssetClass, AssetData->PackageFlags))
 	{
-		AssetUpdatedEvent.Broadcast(*AssetData);
+		EventContext.AssetEvents.Emplace(*AssetData, Impl::FEventContext::EEvent::Updated);
 	}
 }
 
-bool UAssetRegistryImpl::RemoveAssetData(FAssetData* AssetData)
+bool FAssetRegistryImpl::RemoveAssetData(Impl::FEventContext& EventContext, FAssetData* AssetData)
 {
 	bool bRemoved = false;
 
 	if (ensure(AssetData))
 	{
-		if (!UE::AssetRegistry::FFiltering::ShouldSkipAsset(AssetData->AssetClass, AssetData->PackageFlags))
+		if (!ShouldSkipAsset(AssetData->AssetClass, AssetData->PackageFlags))
 		{
-			// Notify subscribers
-			AssetRemovedEvent.Broadcast(*AssetData);
+			EventContext.AssetEvents.Emplace(*AssetData, Impl::FEventContext::EEvent::Removed);
 		}
 
 		// Remove from the class map if removing a blueprint
@@ -2959,11 +3820,11 @@ bool UAssetRegistryImpl::RemoveAssetData(FAssetData* AssetData)
 			const FString OldGeneratedClass = AssetData->GetTagValueRef<FString>(FBlueprintTags::GeneratedClassPath);
 			if (!OldGeneratedClass.IsEmpty())
 			{
-				const FName OldGeneratedClassFName = *ExportTextPathToObjectName(OldGeneratedClass);
+				const FName OldGeneratedClassFName = *Utils::ExportTextPathToObjectName(OldGeneratedClass);
 				CachedBPInheritanceMap.Remove(OldGeneratedClassFName);
 
 				// Invalidate caching because CachedBPInheritanceMap got modified
-				bIsTempCachingUpToDate = false;
+				TempCachedInheritanceBuffer.bDirty = true;
 			}
 		}
 
@@ -2974,7 +3835,7 @@ bool UAssetRegistryImpl::RemoveAssetData(FAssetData* AssetData)
 	return bRemoved;
 }
 
-void UAssetRegistryImpl::RemovePackageData(const FName PackageName)
+void FAssetRegistryImpl::RemovePackageData(Impl::FEventContext& EventContext, const FName PackageName)
 {
 	TArray<FAssetData*, TInlineAllocator<1>>* PackageAssetsPtr = State.CachedAssetsByPackageName.Find(PackageName);
 	if (PackageAssetsPtr && PackageAssetsPtr->Num() > 0)
@@ -2996,7 +3857,7 @@ void UAssetRegistryImpl::RemovePackageData(const FName PackageName)
 		TArray<FAssetData*, TInlineAllocator<1>> PackageAssets = *PackageAssetsPtr;
 		for (FAssetData* PackageAsset : PackageAssets)
 		{
-			RemoveAssetData(PackageAsset);
+			RemoveAssetData(EventContext, PackageAsset);
 		}
 
 		// Readd any referencers, creating an empty DependsNode to hold them
@@ -3016,6 +3877,8 @@ void UAssetRegistryImpl::RemovePackageData(const FName PackageName)
 	}
 }
 
+}
+
 #if WITH_EDITOR
 
 void UAssetRegistryImpl::OnDirectoryChanged(const TArray<FFileChangeData>& FileChanges)
@@ -3049,10 +3912,22 @@ void UAssetRegistryImpl::OnDirectoryChanged(const TArray<FFileChangeData>& FileC
 		}
 	}
 
+	UE::AssetRegistry::Impl::FEventContext EventContext;
+	{
+		FWriteScopeLock InterfaceScopeLock(InterfaceLock);
+		GuardedData.OnDirectoryChanged(EventContext, FileChangesProcessed);
+	}
+	Broadcast(EventContext);
+}
+
+namespace UE::AssetRegistry
+{
+
+void FAssetRegistryImpl::OnDirectoryChanged(Impl::FEventContext& EventContext, TArray<FFileChangeData>& FileChangesProcessed)
+{
 	TArray<FString> NewDirs;
 	TArray<FString> NewFiles;
 	TArray<FString> ModifiedFiles;
-
 	for (int32 FileIdx = 0; FileIdx < FileChangesProcessed.Num(); ++FileIdx)
 	{
 		FString LongPackageName;
@@ -3079,7 +3954,7 @@ void UAssetRegistryImpl::OnDirectoryChanged(const TArray<FFileChangeData>& FileC
 
 			case FFileChangeData::FCA_Removed:
 				// This file was deleted. Remove all assets in the package from the registry.
-				RemovePackageData(*LongPackageName);
+				RemovePackageData(EventContext, *LongPackageName);
 				UE_LOG(LogAssetRegistry, Verbose, TEXT("File was removed from content directory: %s"), *File);
 				break;
 			}
@@ -3101,7 +3976,7 @@ void UAssetRegistryImpl::OnDirectoryChanged(const TArray<FFileChangeData>& FileC
 			}
 			case FFileChangeData::FCA_Removed:
 			{
-				RemoveAssetPath(*LongPackageName);
+				RemoveAssetPath(EventContext, *LongPackageName);
 				UE_LOG(LogAssetRegistry, Verbose, TEXT("Directory was removed from content directory: %s"), *File);
 				break;
 			}
@@ -3122,34 +3997,109 @@ void UAssetRegistryImpl::OnDirectoryChanged(const TArray<FFileChangeData>& FileC
 			GlobalGatherer->OnFilesCreated(NewFiles);
 			if (GlobalGatherer->IsSynchronous())
 			{
-				ScanPathsSynchronousInternal(NewDirs, NewFiles, false /* bForceRescan */, false /* bIgnoreBlacklistScanFilters */, nullptr /* OutFoundAssets */);
+				Impl::FScanPathContext Context(EventContext, NewDirs, NewFiles,
+					false /* bForceRescan */, false /* bIgnoreBlacklistScanFilters */, nullptr /* OutFoundAssets */);
+				ScanPathsSynchronous(Context);
 			}
 		}
 	}
+	ScanModifiedAssetFiles(EventContext, ModifiedFiles);
+}
 
-	ScanModifiedAssetFiles(ModifiedFiles);
 }
 
 void UAssetRegistryImpl::OnAssetLoaded(UObject *AssetLoaded)
 {
-	LoadedAssetsToProcess.Add(AssetLoaded);
+	FWriteScopeLock InterfaceScopeLock(InterfaceLock);
+	GuardedData.OnAssetLoaded(AssetLoaded);
 }
 
-void UAssetRegistryImpl::ProcessLoadedAssetsToUpdateCache(const double TickStartTime, bool bBecameIdle)
+void UAssetRegistryImpl::ProcessLoadedAssetsToUpdateCache(UE::AssetRegistry::Impl::FEventContext& EventContext,
+	const double TickStartTime, bool bIdle)
 {
 	// Note this function can be reentered due to arbitrary code execution in construction of FAssetData
-	LLM_SCOPE(ELLMTag::AssetRegistry);
+	if (!IsInGameThread())
+	{
+		// Calls to GetAssetRegistryTags are only allowed on the GameThread
+		return;
+	}
 
 	const bool bFlushFullBuffer = TickStartTime < 0;
+	if (!bIdle && !bFlushFullBuffer)
+	{
+		return;
+	}
 
-	// Refreshes ClassGeneratorNames if out of date due to module load
-	CollectCodeGeneratorClasses();
+	LLM_SCOPE(ELLMTag::AssetRegistry);
+
+
+	constexpr int32 BatchSize = 16;
+	TArray<UObject*> BatchObjects;
+	TArray<FAssetData, TInlineAllocator<BatchSize>> BatchAssetDatas;
 
-	// Add the found assets
-	while (!LoadedAssetsToProcess.IsEmpty())
 	{
-		UObject* LoadedAsset = LoadedAssetsToProcess.PopFrontValue().Get();
+		FWriteScopeLock InterfaceScopeLock(InterfaceLock);
+		GuardedData.GetProcessLoadedAssetsBatch(BatchObjects, BatchSize);
+		if (BatchObjects.Num() == 0)
+		{
+			return;
+		}
 
+		// Refreshes ClassGeneratorNames if out of date due to module load
+		GuardedData.CollectCodeGeneratorClasses();
+	}
+
+	while (BatchObjects.Num() > 0)
+	{
+		bool bTimedOut = false;
+		int32 CurrentBatchSize = BatchObjects.Num();
+		BatchAssetDatas.Reset(CurrentBatchSize);
+		int32 Index = 0;
+		while (Index < CurrentBatchSize)
+		{
+			UObject* LoadedObject = BatchObjects[Index++];
+			BatchAssetDatas.Add(FAssetData(LoadedObject, true /* bAllowBlueprintClass */));
+
+			// Check to see if we have run out of time in this tick
+			if (!bFlushFullBuffer &&
+				(FPlatformTime::Seconds() - TickStartTime) > UE::AssetRegistry::Impl::MaxSecondsPerFrame)
+			{
+				bTimedOut = true;
+				break;
+			}
+		}
+
+		FWriteScopeLock InterfaceScopeLock(InterfaceLock);
+		GuardedData.PushProcessLoadedAssetsBatch(EventContext, BatchAssetDatas,
+			TArrayView<UObject*>(BatchObjects).Slice(Index, CurrentBatchSize-Index));
+		if (bTimedOut)
+		{
+			break;
+		}
+		GuardedData.GetProcessLoadedAssetsBatch(BatchObjects, BatchSize);
+	}
+}
+
+namespace UE::AssetRegistry
+{
+
+void FAssetRegistryImpl::OnAssetLoaded(UObject* AssetLoaded)
+{
+	LoadedAssetsToProcess.Add(AssetLoaded);
+}
+
+void FAssetRegistryImpl::GetProcessLoadedAssetsBatch(TArray<UObject*>& OutLoadedAssets, uint32 BatchSize)
+{
+	if (!GlobalGatherer.IsValid() || !bUpdateDiskCacheAfterLoad)
+	{
+		OutLoadedAssets.Reset();
+		return;
+	}
+
+	OutLoadedAssets.Reset(BatchSize);
+	while (!LoadedAssetsToProcess.IsEmpty() && OutLoadedAssets.Num() < static_cast<int32>(BatchSize))
+	{
+		UObject* LoadedAsset = LoadedAssetsToProcess.PopFrontValue().Get();
 		if (!LoadedAsset)
 		{
 			// This could be null, in which case it already got freed, ignore
@@ -3176,22 +4126,31 @@ void UAssetRegistryImpl::ProcessLoadedAssetsToUpdateCache(const double TickStart
 			continue;
 		}
 
-		FAssetData** DataFromGather = State.CachedAssetsByObjectPath.Find(ObjectPath);
+		OutLoadedAssets.Add(LoadedAsset);
+	}
+}
+
+void FAssetRegistryImpl::PushProcessLoadedAssetsBatch(Impl::FEventContext& EventContext,
+	TArrayView<FAssetData> LoadedAssetDatas, TArrayView<UObject*> UnprocessedFromBatch)
+{
+	// Add or update existing for all of the AssetDatas created by the batch
+	for (FAssetData& NewAssetData : LoadedAssetDatas)
+	{
+		FAssetData** DataFromGather = State.CachedAssetsByObjectPath.Find(NewAssetData.ObjectPath);
 
-		AssetDataObjectPathsUpdatedOnLoad.Add(ObjectPath);
+		AssetDataObjectPathsUpdatedOnLoad.Add(NewAssetData.ObjectPath);
 
 		if (!DataFromGather)
 		{
-			FAssetData* NewAssetData = new FAssetData(LoadedAsset, true /* bAllowBlueprintClass */);
-			AddAssetData(NewAssetData);
+			FAssetData* ClonedAssetData = new FAssetData(MoveTemp(NewAssetData));
+			AddAssetData(EventContext, ClonedAssetData);
 		}
 		else
 		{
-			FAssetData NewAssetData = FAssetData(LoadedAsset, true /* bAllowBlueprintClass */);
 			if (NewAssetData.TagsAndValues != (*DataFromGather)->TagsAndValues)
 			{
 				// We need to actually update disk cache
-				UpdateAssetData(*DataFromGather, NewAssetData);
+				UpdateAssetData(EventContext, *DataFromGather, NewAssetData);
 			}
 			else
 			{
@@ -3199,16 +4158,17 @@ void UAssetRegistryImpl::ProcessLoadedAssetsToUpdateCache(const double TickStart
 				(*DataFromGather)->TaggedAssetBundles = NewAssetData.TaggedAssetBundles;
 			}
 		}
+	}
 
-		// Check to see if we have run out of time in this tick
-		if (!bFlushFullBuffer && (FPlatformTime::Seconds() - TickStartTime) > MaxSecondsPerFrame)
-		{
-			break;
-		}
+	// Push back any objects from the batch that were not processed due to timing out
+	for (int32 Index = UnprocessedFromBatch.Num() - 1; Index >= 0; --Index)
+	{
+		LoadedAssetsToProcess.EmplaceFront(UnprocessedFromBatch[Index]);
 	}
 }
 
-void UAssetRegistryImpl::UpdateRedirectCollector()
+
+void FAssetRegistryImpl::UpdateRedirectCollector()
 {
 	// Look for all redirectors in list
 	const TArray<const FAssetData*>& RedirectorAssets = State.GetAssetsByClassName(UObjectRedirector::StaticClass()->GetFName());
@@ -3224,9 +4184,24 @@ void UAssetRegistryImpl::UpdateRedirectCollector()
 	}
 }
 
+}
+
 #endif // WITH_EDITOR
 
 void UAssetRegistryImpl::ScanModifiedAssetFiles(const TArray<FString>& InFilePaths)
+{
+	UE::AssetRegistry::Impl::FEventContext EventContext;
+	{
+		FWriteScopeLock InterfaceScopeLock(InterfaceLock);
+		GuardedData.ScanModifiedAssetFiles(EventContext, InFilePaths);
+	}
+	Broadcast(EventContext);
+}
+
+namespace UE::AssetRegistry
+{
+
+void FAssetRegistryImpl::ScanModifiedAssetFiles(Impl::FEventContext& EventContext, const TArray<FString>& InFilePaths)
 {
 	if (InFilePaths.Num() > 0)
 	{
@@ -3256,7 +4231,9 @@ void UAssetRegistryImpl::ScanModifiedAssetFiles(const TArray<FString>& InFilePat
 
 		// Re-scan and update the asset registry with the new asset data
 		TArray<FName> FoundAssets;
-		ScanPathsSynchronousInternal(TArray<FString>(), InFilePaths, true /* bForceRescan */, false /* bIgnoreBlackListScanFilters */, &FoundAssets);
+		Impl::FScanPathContext Context(EventContext, TArray<FString>(), InFilePaths,
+			true /* bForceRescan */, false /* bIgnoreBlacklistScanFilters */, &FoundAssets);
+		ScanPathsSynchronous(Context);
 
 		// Remove any assets that are no longer present in the package
 		for (const TArray<FAssetData*, TInlineAllocator<1>>& OldPackageAssets : ExistingFilesAssetData)
@@ -3265,13 +4242,15 @@ void UAssetRegistryImpl::ScanModifiedAssetFiles(const TArray<FString>& InFilePat
 			{
 				if (!FoundAssets.Contains(OldPackageAsset->ObjectPath))
 				{
-					RemoveAssetData(OldPackageAsset);
+					RemoveAssetData(EventContext, OldPackageAsset);
 				}
 			}
 		}
 	}
 }
 
+}
+
 void UAssetRegistryImpl::OnContentPathMounted(const FString& InAssetPath, const FString& FileSystemPath)
 {
 	// Sanitize
@@ -3286,14 +4265,63 @@ void UAssetRegistryImpl::OnContentPathMounted(const FString& InAssetPath, const
 		AssetPathWithTrailingSlash = InAssetPath;
 	}
 
+#if WITH_EDITOR
+	FDirectoryWatcherModule& DirectoryWatcherModule = FModuleManager::LoadModuleChecked<FDirectoryWatcherModule>(TEXT("DirectoryWatcher"));
+	IDirectoryWatcher* DirectoryWatcher = nullptr;
+	if (GIsEditor) 	// In-game doesn't listen for directory changes
+	{
+		DirectoryWatcher = DirectoryWatcherModule.Get();
+		if (DirectoryWatcher)
+		{
+			// If the path doesn't exist on disk, make it so the watcher will work.
+			IFileManager::Get().MakeDirectory(*FileSystemPath);
+		}
+	}
+		
+#endif
+
+	UE::AssetRegistry::Impl::FEventContext EventContext;
+	{
+		FWriteScopeLock InterfaceScopeLock(InterfaceLock);
+		GuardedData.OnContentPathMounted(EventContext, InAssetPath, AssetPathWithTrailingSlash, FileSystemPath);
+
+		// Listen for directory changes in this content path
+#if WITH_EDITOR
+		// In-game doesn't listen for directory changes
+		if (DirectoryWatcher)
+		{
+			if (!OnDirectoryChangedDelegateHandles.Contains(AssetPathWithTrailingSlash))
+			{
+				FDelegateHandle NewHandle;
+				DirectoryWatcher->RegisterDirectoryChangedCallback_Handle(
+					FileSystemPath, 
+					IDirectoryWatcher::FDirectoryChanged::CreateUObject(this, &UAssetRegistryImpl::OnDirectoryChanged),
+					NewHandle, 
+					IDirectoryWatcher::WatchOptions::IncludeDirectoryChanges);
+
+				OnDirectoryChangedDelegateHandles.Add(AssetPathWithTrailingSlash, NewHandle);
+			}
+		}
+#endif // WITH_EDITOR
+	}
+
+	Broadcast(EventContext);
+}
+
+namespace UE::AssetRegistry
+{
+
+void FAssetRegistryImpl::OnContentPathMounted(Impl::FEventContext& EventContext, const FString& InAssetPath, const FString& AssetPathWithTrailingSlash, const FString& FileSystemPath)
+{
 	// Content roots always exist
-	AddPath(AssetPathWithTrailingSlash);
+	AddPath(EventContext, AssetPathWithTrailingSlash);
 
 	if (GlobalGatherer.IsValid() && bSearchAllAssets)
 	{
 		if (GlobalGatherer->IsSynchronous())
 		{
-			ScanPathsSynchronous({ FileSystemPath });
+			Impl::FScanPathContext Context(EventContext, { FileSystemPath }, TArray<FString>());
+			ScanPathsSynchronous(Context);
 		}
 		else
 		{
@@ -3301,46 +4329,65 @@ void UAssetRegistryImpl::OnContentPathMounted(const FString& InAssetPath, const
 			GlobalGatherer->SetIsWhitelisted(FileSystemPath, true);
 		}
 	}
+}
+
+}
+
+void UAssetRegistryImpl::OnContentPathDismounted(const FString& InAssetPath, const FString& FileSystemPath)
+{
+	// Sanitize
+	FString AssetPathNoTrailingSlash = InAssetPath;
+	if (AssetPathNoTrailingSlash.EndsWith(TEXT("/"), ESearchCase::CaseSensitive))
+	{
+		// We don't want a trailing slash here as it could interfere with RemoveAssetPath
+		AssetPathNoTrailingSlash.LeftChopInline(1, false);
+	}
 
-	// Listen for directory changes in this content path
 #if WITH_EDITOR
-	// In-game doesn't listen for directory changes
-	if (GIsEditor)
+	FDirectoryWatcherModule& DirectoryWatcherModule = FModuleManager::LoadModuleChecked<FDirectoryWatcherModule>(TEXT("DirectoryWatcher"));
+	IDirectoryWatcher* DirectoryWatcher = nullptr;
+	if (GIsEditor) 	// In-game doesn't listen for directory changes
 	{
-		FDirectoryWatcherModule& DirectoryWatcherModule = FModuleManager::LoadModuleChecked<FDirectoryWatcherModule>(TEXT("DirectoryWatcher"));
-		IDirectoryWatcher* DirectoryWatcher = DirectoryWatcherModule.Get();
+		DirectoryWatcher = DirectoryWatcherModule.Get();
+	}
+#endif
+
+	UE::AssetRegistry::Impl::FEventContext EventContext;
+	{
+		FWriteScopeLock InterfaceScopeLock(InterfaceLock);
+		GuardedData.OnContentPathDismounted(EventContext, InAssetPath, AssetPathNoTrailingSlash, FileSystemPath);
+
+		// Stop listening for directory changes in this content path
+#if WITH_EDITOR
 		if (DirectoryWatcher)
 		{
-			// If the path doesn't exist on disk, make it so the watcher will work.
-			IFileManager::Get().MakeDirectory(*FileSystemPath);
-
-			if (!OnDirectoryChangedDelegateHandles.Contains(AssetPathWithTrailingSlash))
+			// Make sure OnDirectoryChangedDelegateHandles key is symmetrical with the one used in OnContentPathMounted
+			FString AssetPathWithTrailingSlash;
+			if (!InAssetPath.EndsWith(TEXT("/"), ESearchCase::CaseSensitive))
 			{
-				FDelegateHandle NewHandle;
-				DirectoryWatcher->RegisterDirectoryChangedCallback_Handle(
-					FileSystemPath, 
-					IDirectoryWatcher::FDirectoryChanged::CreateUObject(this, &UAssetRegistryImpl::OnDirectoryChanged),
-					NewHandle, 
-					IDirectoryWatcher::WatchOptions::IncludeDirectoryChanges);
+				AssetPathWithTrailingSlash = InAssetPath + TEXT("/");
+			}
+			else
+			{
+				AssetPathWithTrailingSlash = InAssetPath;
+			}
 
-				OnDirectoryChangedDelegateHandles.Add(AssetPathWithTrailingSlash, NewHandle);
+			FDelegateHandle DirectoryChangedHandle;
+			if (ensure(OnDirectoryChangedDelegateHandles.RemoveAndCopyValue(AssetPathWithTrailingSlash, DirectoryChangedHandle)))
+			{
+				DirectoryWatcher->UnregisterDirectoryChangedCallback_Handle(FileSystemPath, DirectoryChangedHandle);
 			}
 		}
-	}
 #endif // WITH_EDITOR
-
+	}
+	Broadcast(EventContext);
 }
 
-void UAssetRegistryImpl::OnContentPathDismounted(const FString& InAssetPath, const FString& FileSystemPath)
+namespace UE::AssetRegistry
 {
-	// Sanitize
-	FString AssetPathNoTrailingSlash = InAssetPath;
-	if (AssetPathNoTrailingSlash.EndsWith(TEXT("/"), ESearchCase::CaseSensitive))
-	{
-		// We don't want a trailing slash here as it could interfere with RemoveAssetPath
-		AssetPathNoTrailingSlash.LeftChopInline(1, false);
-	}
 
+void FAssetRegistryImpl::OnContentPathDismounted(Impl::FEventContext& EventContext, const FString& InAssetPath, const FString& AssetPathNoTrailingSlash, const FString& FileSystemPath)
+{
 	if (GlobalGatherer.IsValid())
 	{
 		GlobalGatherer->RemoveMountPoint(FileSystemPath);
@@ -3348,14 +4395,15 @@ void UAssetRegistryImpl::OnContentPathDismounted(const FString& InAssetPath, con
 
 	// Remove all cached assets found at this location
 	{
+		FName AssetPathNoTrailingSlashFName(*AssetPathNoTrailingSlash);
 		TArray<FAssetData*> AllAssetDataToRemove;
-		TArray<FString> PathList;
+		TSet<FName> PathList;
 		const bool bRecurse = true;
-		GetSubPaths(AssetPathNoTrailingSlash, PathList, bRecurse);
-		PathList.Add(AssetPathNoTrailingSlash);
-		for (const FString& Path : PathList)
+		CachedPathTree.GetSubPaths(AssetPathNoTrailingSlashFName, PathList, bRecurse);
+		PathList.Add(AssetPathNoTrailingSlashFName);
+		for (FName PathName : PathList)
 		{
-			TArray<FAssetData*>* AssetsInPath = State.CachedAssetsByPath.Find(FName(*Path));
+			TArray<FAssetData*>* AssetsInPath = State.CachedAssetsByPath.Find(PathName);
 			if (AssetsInPath)
 			{
 				AllAssetDataToRemove.Append(*AssetsInPath);
@@ -3364,102 +4412,78 @@ void UAssetRegistryImpl::OnContentPathDismounted(const FString& InAssetPath, con
 
 		for (FAssetData* AssetData : AllAssetDataToRemove)
 		{
-			RemoveAssetData(AssetData);
+			RemoveAssetData(EventContext, AssetData);
 		}
 	}
 
 	// Remove the root path
 	{
 		const bool bEvenIfAssetsStillExist = true;
-		RemoveAssetPath(FName(*AssetPathNoTrailingSlash), bEvenIfAssetsStillExist);
-	}
-
-	// Stop listening for directory changes in this content path
-#if WITH_EDITOR
-	// In-game doesn't listen for directory changes
-	if (GIsEditor)
-	{
-		FDirectoryWatcherModule& DirectoryWatcherModule = FModuleManager::LoadModuleChecked<FDirectoryWatcherModule>(TEXT("DirectoryWatcher"));
-		IDirectoryWatcher* DirectoryWatcher = DirectoryWatcherModule.Get();
-		if (DirectoryWatcher)
-		{
-			// Make sure OnDirectoryChangedDelegateHandles key is symmetrical with the one used in OnContentPathMounted
-			FString AssetPathWithTrailingSlash;
-			if (!InAssetPath.EndsWith(TEXT("/"), ESearchCase::CaseSensitive))
-			{
-				AssetPathWithTrailingSlash = InAssetPath + TEXT("/");
-			}
-			else
-			{
-				AssetPathWithTrailingSlash = InAssetPath;
-			}
-
-			FDelegateHandle DirectoryChangedHandle;
-			if (ensure(OnDirectoryChangedDelegateHandles.RemoveAndCopyValue(AssetPathWithTrailingSlash, DirectoryChangedHandle)))
-			{
-				DirectoryWatcher->UnregisterDirectoryChangedCallback_Handle(FileSystemPath, DirectoryChangedHandle);
-			}
-		}
+		RemoveAssetPath(EventContext, FName(*AssetPathNoTrailingSlash), bEvenIfAssetsStillExist);
 	}
-#endif // WITH_EDITOR
+}
 
 }
 
 void UAssetRegistryImpl::SetTemporaryCachingMode(bool bEnable)
+{
+	checkf(IsInGameThread(), TEXT("Changing Caching mode is only available on the game thread because it affects behavior on all threads"));
+	FWriteScopeLock InterfaceScopeLock(InterfaceLock);
+	GuardedData.SetTemporaryCachingMode(bEnable);
+}
+
+namespace UE::AssetRegistry
+{
+
+void FAssetRegistryImpl::SetTemporaryCachingMode(bool bEnable)
 {
 	if (bIsTempCachingAlwaysEnabled || bEnable == bIsTempCachingEnabled)
 	{
 		return;
 	}
 
-	if (bEnable)
-	{
-		bIsTempCachingEnabled = true;
-		bIsTempCachingUpToDate = false;
-	}
-	else
+	bIsTempCachingEnabled = bEnable;
+	TempCachedInheritanceBuffer.bDirty = true;
+	if (!bEnable)
 	{
-		bIsTempCachingEnabled = false;
-		ClearTemporaryCaches();
+		TempCachedInheritanceBuffer.Clear();
 	}
 }
 
+}
+
 void UAssetRegistryImpl::SetTemporaryCachingModeInvalidated()
 {
-	bIsTempCachingUpToDate = false;
+	checkf(IsInGameThread(), TEXT("Invalidating temporary cache is only available on the game thread because it affects behavior on all threads"));
+	FWriteScopeLock InterfaceScopeLock(InterfaceLock);
+	GuardedData.SetTemporaryCachingModeInvalidated();
 }
 
-bool UAssetRegistryImpl::GetTemporaryCachingMode() const
+namespace UE::AssetRegistry
 {
-	return bIsTempCachingEnabled;
+
+void FAssetRegistryImpl::SetTemporaryCachingModeInvalidated()
+{
+	TempCachedInheritanceBuffer.bDirty = true;
+}
+
 }
 
-void UAssetRegistryImpl::ClearTemporaryCaches() const
+bool UAssetRegistryImpl::GetTemporaryCachingMode() const
 {
-	if (!bIsTempCachingEnabled && !bIsTempCachingAlwaysEnabled)
-	{
-		// We clear these as much as possible to get back memory
-		TempCachedInheritanceMap.Empty();
-		TempReverseInheritanceMap.Empty();
-		bIsTempCachingUpToDate = false;
-	}
+	FReadScopeLock InterfaceScopeLock(InterfaceLock);
+	return GuardedData.IsTempCachingEnabled();
 }
 
 COREUOBJECT_API TMap<UClass*, TSet<UClass*>> GetAllDerivedClasses();
 
-void UAssetRegistryImpl::UpdateTemporaryCaches() const
+namespace UE::AssetRegistry
 {
-	if (bIsTempCachingEnabled && bIsTempCachingUpToDate && TempCachingRegisteredClassesVersionNumber == GetRegisteredClassesVersionNumber())
-	{
-		return;
-	}
 
+void FAssetRegistryImpl::UpdateInheritanceBuffer(Impl::FClassInheritanceBuffer& OutBuffer) const
+{
 	TRACE_CPUPROFILER_EVENT_SCOPE(UAssetRegistryImpl::UpdateTemporaryCaches)
 
-	// Refreshes ClassGeneratorNames if out of date due to module load
-	CollectCodeGeneratorClasses();
-
-	TempCachingRegisteredClassesVersionNumber = GetRegisteredClassesVersionNumber();
 	TMap<UClass*, TSet<UClass*>> NativeSubclasses = GetAllDerivedClasses();
 
 	uint32 NumNativeClasses = 1; // UObject has no superclass
@@ -3467,21 +4491,21 @@ void UAssetRegistryImpl::UpdateTemporaryCaches() const
 	{
 		NumNativeClasses += Pair.Value.Num();
 	}
-	TempCachedInheritanceMap.Reserve(NumNativeClasses + CachedBPInheritanceMap.Num());
-	TempCachedInheritanceMap = CachedBPInheritanceMap;
-	TempCachedInheritanceMap.Add(FName("Object"), FName());
+	OutBuffer.InheritanceMap.Reserve(NumNativeClasses + CachedBPInheritanceMap.Num());
+	OutBuffer.InheritanceMap = CachedBPInheritanceMap;
+	OutBuffer.InheritanceMap.Add(FName("Object"), FName());
 
-	for (TPair<FName, TArray<FName>>& Pair : TempReverseInheritanceMap)
+	for (TPair<FName, TArray<FName>>& Pair : OutBuffer.ReverseInheritanceMap)
 	{
 		Pair.Value.Reset();
 	}
-	TempReverseInheritanceMap.Reserve(NativeSubclasses.Num());
+	OutBuffer.ReverseInheritanceMap.Reserve(NativeSubclasses.Num());
 
 	for (const TPair<UClass*, TSet<UClass*>>& Pair : NativeSubclasses)
 	{
 		FName SuperclassName = Pair.Key->GetFName();
 
-		TArray<FName>& TempCachedSubclasses = TempReverseInheritanceMap.FindOrAdd(SuperclassName);
+		TArray<FName>& TempCachedSubclasses = OutBuffer.ReverseInheritanceMap.FindOrAdd(SuperclassName);
 		TempCachedSubclasses.Reserve(Pair.Value.Num());
 		for (UClass* Subclass : Pair.Value)
 		{
@@ -3489,14 +4513,14 @@ void UAssetRegistryImpl::UpdateTemporaryCaches() const
 			{
 				FName SubclassName = Subclass->GetFName();
 				TempCachedSubclasses.Add(SubclassName);
-				TempCachedInheritanceMap.Add(SubclassName, SuperclassName);
+				OutBuffer.InheritanceMap.Add(SubclassName, SuperclassName);
 
 				// Add any implemented interfaces to the reverse inheritance map, but not to the forward map
 				for (const FImplementedInterface& Interface : Subclass->Interfaces)
 				{
 					if (UClass* InterfaceClass = Interface.Class) // could be nulled out by ForceDelete of a blueprint interface
 					{
-						TArray<FName>& Implementations = TempReverseInheritanceMap.FindOrAdd(InterfaceClass->GetFName());
+						TArray<FName>& Implementations = OutBuffer.ReverseInheritanceMap.FindOrAdd(InterfaceClass->GetFName());
 						Implementations.Add(SubclassName);
 					}
 				}
@@ -3510,29 +4534,125 @@ void UAssetRegistryImpl::UpdateTemporaryCaches() const
 		FName ParentClassName = Kvp.Value;
 		if (!ParentClassName.IsNone())
 		{
-			TArray<FName>& ChildClasses = TempReverseInheritanceMap.FindOrAdd(ParentClassName);
+			TArray<FName>& ChildClasses = OutBuffer.ReverseInheritanceMap.FindOrAdd(ParentClassName);
 			ChildClasses.Add(Kvp.Key);
 		}
+
+	}
+
+	OutBuffer.RegisteredClassesVersionNumber = GetRegisteredClassesVersionNumber();
+	OutBuffer.bDirty = false;
+}
+
+}
+
+void UAssetRegistryImpl::GetInheritanceContextWithRequiredLock(FRWScopeLock& InOutScopeLock,
+	UE::AssetRegistry::Impl::FClassInheritanceContext& InheritanceContext,
+	UE::AssetRegistry::Impl::FClassInheritanceBuffer& StackBuffer)
+{
+	uint64 CurrentClassesVersionNumber = GetRegisteredClassesVersionNumber();
+	// If bIsTempCachingAlwaysEnabled, then we are guaranteed that bIsTempCachingEnabled=true.
+	// We rely on this to simplify logic and only check bIsTempCachingEnabled
+	check(!GuardedData.IsTempCachingAlwaysEnabled() || GuardedData.IsTempCachingEnabled());
+
+	bool bNeedsWriteLock = false;
+	if (GuardedData.GetClassGeneratorNamesRegisteredClassesVersionNumber() != CurrentClassesVersionNumber)
+	{
+		// ConditionalUpdate writes to protected data in CollectCodeGeneratorClasses, so we cannot proceed under a read lock
+		bNeedsWriteLock = true;
+	}
+	if (GuardedData.IsTempCachingEnabled() && !GuardedData.GetTempCachedInheritanceBuffer().IsUpToDate(CurrentClassesVersionNumber))
+	{
+		// Temp caching is enabled, so we will be reading the protected data in TempCachedInheritanceBuffer
+		// It's out of date, so we need to write to it first, so we cannot proceed under a read lock
+		bNeedsWriteLock = true;
+	}
+	if (bNeedsWriteLock)
+	{
+		InOutScopeLock.ReleaseReadOnlyLockAndAcquireWriteLock_USE_WITH_CAUTION();
+	}
+
+	// Note that we have to reread all data since we may have dropped the lock
+	bool bCodeGeneratorClassesUpToDate = GuardedData.GetClassGeneratorNamesRegisteredClassesVersionNumber() == CurrentClassesVersionNumber;
+	if (GuardedData.IsTempCachingEnabled())
+	{
+		// Use the persistent buffer
+		UE::AssetRegistry::Impl::FClassInheritanceBuffer& TempCachedInheritanceBuffer = GuardedData.GetTempCachedInheritanceBuffer();
+		bool bInheritanceMapUpToDate = TempCachedInheritanceBuffer.IsUpToDate(CurrentClassesVersionNumber);
+		InheritanceContext.BindToBuffer(TempCachedInheritanceBuffer, GuardedData, bInheritanceMapUpToDate, bCodeGeneratorClassesUpToDate);
+	}
+	else
+	{
+		// Use the StackBuffer for the duration of the caller
+		InheritanceContext.BindToBuffer(StackBuffer, GuardedData, false /* bInInheritanceMapUpToDate */, bCodeGeneratorClassesUpToDate);
+	}
+}
+
+namespace UE::AssetRegistry
+{
+
+namespace Impl
+{
+
+void FClassInheritanceBuffer::Clear()
+{
+	InheritanceMap.Empty();
+	ReverseInheritanceMap.Empty();
+}
+
+bool FClassInheritanceBuffer::IsUpToDate(uint64 CurrentClassesVersionNumber) const
+{
+	return !bDirty && RegisteredClassesVersionNumber == CurrentClassesVersionNumber;
+}
+
+uint32 FClassInheritanceBuffer::GetAllocatedSize() const
+{
+	return InheritanceMap.GetAllocatedSize() + ReverseInheritanceMap.GetAllocatedSize();
+}
+
+void FClassInheritanceContext::BindToBuffer(FClassInheritanceBuffer& InBuffer, FAssetRegistryImpl& InAssetRegistryImpl, 
+	bool bInInheritanceMapUpToDate, bool bInCodeGeneratorClassesUpToDate)
+{
+	AssetRegistryImpl = &InAssetRegistryImpl;
+	Buffer = &InBuffer;
+	bInheritanceMapUpToDate = bInInheritanceMapUpToDate;
+	bCodeGeneratorClassesUpToDate = bInCodeGeneratorClassesUpToDate;
+}
+
+void FClassInheritanceContext::ConditionalUpdate()
+{
+	check(Buffer != nullptr); // It is not valid to call ConditionalUpdate with an empty FClassInheritanceContext
+	if (bInheritanceMapUpToDate)
+	{
+		return;
 	}
 
-	bIsTempCachingUpToDate = true;
+	if (!bCodeGeneratorClassesUpToDate)
+	{
+		AssetRegistryImpl->CollectCodeGeneratorClasses();
+		bCodeGeneratorClassesUpToDate = true;
+	}
+	AssetRegistryImpl->UpdateInheritanceBuffer(*Buffer);
+	bInheritanceMapUpToDate = true;
+}
+
 }
 
-void UAssetRegistryImpl::GetSubClasses(const TArray<FName>& InClassNames, const TSet<FName>& ExcludedClassNames, TSet<FName>& SubClassNames) const
+void FAssetRegistryImpl::GetSubClasses(Impl::FClassInheritanceContext& InheritanceContext,
+	const TArray<FName>& InClassNames, const TSet<FName>& ExcludedClassNames, TSet<FName>& SubClassNames) const
 {
-	UpdateTemporaryCaches();
+	InheritanceContext.ConditionalUpdate();
 
 	TSet<FName> ProcessedClassNames;
 	for (FName ClassName : InClassNames)
 	{
 		// Now find all subclass names
-		GetSubClasses_Recursive(ClassName, SubClassNames, ProcessedClassNames, TempReverseInheritanceMap, ExcludedClassNames);
+		GetSubClasses_Recursive(InheritanceContext, ClassName, SubClassNames, ProcessedClassNames, ExcludedClassNames);
 	}
-
-	ClearTemporaryCaches();
 }
 
-void UAssetRegistryImpl::GetSubClasses_Recursive(FName InClassName, TSet<FName>& SubClassNames, TSet<FName>& ProcessedClassNames, const TMap<FName, TArray<FName>>& ReverseInheritanceMap, const TSet<FName>& ExcludedClassNames) const
+void FAssetRegistryImpl::GetSubClasses_Recursive(Impl::FClassInheritanceContext& InheritanceContext, FName InClassName,
+	TSet<FName>& SubClassNames, TSet<FName>& ProcessedClassNames, const TSet<FName>& ExcludedClassNames) const
 {
 	if (ExcludedClassNames.Contains(InClassName))
 	{
@@ -3547,17 +4667,20 @@ void UAssetRegistryImpl::GetSubClasses_Recursive(FName InClassName, TSet<FName>&
 		SubClassNames.Add(InClassName);
 		ProcessedClassNames.Add(InClassName);
 
-		const TArray<FName>* FoundSubClassNames = ReverseInheritanceMap.Find(InClassName);
+		const TArray<FName>* FoundSubClassNames = InheritanceContext.Buffer->ReverseInheritanceMap.Find(InClassName);
 		if (FoundSubClassNames)
 		{
 			for (FName ClassName : (*FoundSubClassNames))
 			{
-				GetSubClasses_Recursive(ClassName, SubClassNames, ProcessedClassNames, ReverseInheritanceMap, ExcludedClassNames);
+				GetSubClasses_Recursive(InheritanceContext, ClassName, SubClassNames, ProcessedClassNames,
+					ExcludedClassNames);
 			}
 		}
 	}
 }
 
+}
+
 #if WITH_EDITOR
 static FString GAssetRegistryManagementPathsPackageDebugName;
 static FAutoConsoleVariableRef CVarAssetRegistryManagementPathsPackageDebugName(
@@ -3591,10 +4714,23 @@ void PrintAssetRegistryManagementPathsPackageDebugInfo(FDependsNode* Node, const
 }
 #endif // WITH_EDITOR
 
-void UAssetRegistryImpl::SetManageReferences(const TMultiMap<FAssetIdentifier, FAssetIdentifier>& ManagerMap, bool bClearExisting, UE::AssetRegistry::EDependencyCategory RecurseType, TSet<FDependsNode*>& ExistingManagedNodes, ShouldSetManagerPredicate ShouldSetManager)
+void UAssetRegistryImpl::SetManageReferences(const TMultiMap<FAssetIdentifier, FAssetIdentifier>& ManagerMap,
+	bool bClearExisting, UE::AssetRegistry::EDependencyCategory RecurseType, TSet<FDependsNode*>& ExistingManagedNodes,
+	ShouldSetManagerPredicate ShouldSetManager)
+{
+	// For performance reasons we call the ShouldSetManager callback when inside the lock. Licensee UAssetManagers
+	// are responsible for not calling AssetRegistry functions from ShouldSetManager as that would create a deadlock
+	FWriteScopeLock InterfaceScopeLock(InterfaceLock);
+	GuardedData.SetManageReferences(ManagerMap, bClearExisting, RecurseType, ExistingManagedNodes, ShouldSetManager);
+}
+
+namespace UE::AssetRegistry
 {
-	using namespace UE::AssetRegistry;
 
+void FAssetRegistryImpl::SetManageReferences(const TMultiMap<FAssetIdentifier, FAssetIdentifier>& ManagerMap,
+	bool bClearExisting, UE::AssetRegistry::EDependencyCategory RecurseType, TSet<FDependsNode*>& ExistingManagedNodes,
+	IAssetRegistry::ShouldSetManagerPredicate ShouldSetManager)
+{
 	// Set default predicate if needed
 	if (!ShouldSetManager)
 	{
@@ -3765,7 +4901,24 @@ void UAssetRegistryImpl::SetManageReferences(const TMultiMap<FAssetIdentifier, F
 #endif
 }
 
+}
+
 bool UAssetRegistryImpl::SetPrimaryAssetIdForObjectPath(const FName ObjectPath, FPrimaryAssetId PrimaryAssetId)
+{
+	UE::AssetRegistry::Impl::FEventContext EventContext;
+	bool bResult;
+	{
+		FWriteScopeLock InterfaceScopeLock(InterfaceLock);
+		bResult = GuardedData.SetPrimaryAssetIdForObjectPath(EventContext, ObjectPath, PrimaryAssetId);
+	}
+	Broadcast(EventContext);
+	return bResult;
+}
+
+namespace UE::AssetRegistry
+{
+
+bool FAssetRegistryImpl::SetPrimaryAssetIdForObjectPath(Impl::FEventContext& EventContext, const FName ObjectPath, FPrimaryAssetId PrimaryAssetId)
 {
 	FAssetData** FoundAssetData = State.CachedAssetsByObjectPath.Find(ObjectPath);
 
@@ -3782,14 +4935,11 @@ bool UAssetRegistryImpl::SetPrimaryAssetIdForObjectPath(const FName ObjectPath,
 
 	FAssetData NewAssetData = FAssetData(AssetData->PackageName, AssetData->PackagePath, AssetData->AssetName, AssetData->AssetClass, TagsAndValues, AssetData->ChunkIDs, AssetData->PackageFlags);
 	NewAssetData.TaggedAssetBundles = AssetData->TaggedAssetBundles;
-	UpdateAssetData(AssetData, NewAssetData);
+	UpdateAssetData(EventContext, AssetData, NewAssetData);
 
 	return true;
 }
 
-const FAssetData* UAssetRegistryImpl::GetCachedAssetDataForObjectPath(const FName ObjectPath) const
-{
-	return State.GetAssetByObjectPath(ObjectPath);
 }
 
 void FAssetRegistryDependencyOptions::SetFromFlags(const EAssetRegistryDependencyType::Type InFlags)
@@ -3841,3 +4991,228 @@ bool FAssetRegistryDependencyOptions::GetManageQuery(UE::AssetRegistry::FDepende
 	}
 	return false;
 }
+
+namespace UE::AssetRegistry
+{
+
+const FAssetRegistryState& FAssetRegistryImpl::GetState() const
+{
+	return State;
+}
+
+const FPathTree& FAssetRegistryImpl::GetCachedPathTree() const
+{
+	return CachedPathTree;
+}
+
+const TSet<FName>& FAssetRegistryImpl::GetCachedEmptyPackages() const
+{
+	return CachedEmptyPackages;
+}
+
+bool FAssetRegistryImpl::ShouldSkipAsset(FName AssetClass, uint32 PackageFlags) const
+{
+#if WITH_ENGINE && WITH_EDITOR
+	return Utils::ShouldSkipAsset(AssetClass, PackageFlags, SkipUncookedClasses, SkipCookedClasses);
+#else
+	return false;
+#endif
+}
+
+bool FAssetRegistryImpl::ShouldSkipAsset(const UObject* InAsset) const
+{
+#if WITH_ENGINE && WITH_EDITOR
+	return Utils::ShouldSkipAsset(InAsset, SkipUncookedClasses, SkipCookedClasses);
+#else
+	return false;
+#endif
+}
+
+namespace Impl
+{
+
+void FEventContext::Clear()
+{
+	bFileLoadedEventBroadcast = false;
+	ProgressUpdateData.Reset();
+	PathEvents.Empty();
+	AssetEvents.Empty();
+	RequiredLoads.Empty();
+}
+
+bool FEventContext::IsEmpty() const
+{
+	return !bFileLoadedEventBroadcast &&
+		!ProgressUpdateData.IsSet() &&
+		PathEvents.Num() == 0 &&
+		AssetEvents.Num() == 0 &&
+		RequiredLoads.Num() == 0;
+}
+
+void FEventContext::Append(FEventContext&& Other)
+{
+	if (&Other == this)
+	{
+		return;
+	}
+	bFileLoadedEventBroadcast |= Other.bFileLoadedEventBroadcast;
+	if (Other.ProgressUpdateData.IsSet())
+	{
+		ProgressUpdateData = MoveTemp(Other.ProgressUpdateData);
+	}
+	PathEvents.Append(MoveTemp(Other.PathEvents));
+	AssetEvents.Append(MoveTemp(Other.AssetEvents));
+	RequiredLoads.Append(MoveTemp(Other.RequiredLoads));
+}
+
+}
+
+}
+
+void UAssetRegistryImpl::ReadLockEnumerateTagToAssetDatas(TFunctionRef<void(FName TagName, const TArray<const FAssetData*>& Assets)> Callback) const
+{
+	FReadScopeLock InterfaceScopeLock(InterfaceLock);
+	for (const TPair<FName, const TArray<const FAssetData*>>& Pair : GuardedData.GetState().GetTagToAssetDatasMap())
+	{
+		Callback(Pair.Key, Pair.Value);
+	}
+}
+
+void UAssetRegistryImpl::Broadcast(UE::AssetRegistry::Impl::FEventContext& EventContext)
+{
+	using namespace UE::AssetRegistry::Impl;
+	if (!IsInGameThread())
+	{
+		// By contract events (and PackageLoads) can only be sent on the game thread; some legacy systems depend on 
+		// this and are not threadsafe. If we're not in the game thread, defer all events in the EventContext
+		// instead of broadcasting them on this thread
+		FWriteScopeLock InterfaceScopeLock(InterfaceLock);
+		check(&EventContext != &DeferredEvents); // Only the GameThread should be calling Broadcast on DeferredEvents
+		DeferredEvents.Append(MoveTemp(EventContext));
+		EventContext.Clear();
+		return;
+	}
+
+	if (EventContext.bFileLoadedEventBroadcast)
+	{
+		FileLoadedEvent.Broadcast();
+		EventContext.bFileLoadedEventBroadcast = false;
+	}
+
+	if (EventContext.ProgressUpdateData.IsSet())
+	{
+		FileLoadProgressUpdatedEvent.Broadcast(*EventContext.ProgressUpdateData);
+		EventContext.ProgressUpdateData.Reset();
+	}
+
+	if (EventContext.PathEvents.Num())
+	{
+		for (const TPair<FString, FEventContext::EEvent>& PathEvent : EventContext.PathEvents)
+		{
+			const FString& Path = PathEvent.Get<0>();
+			switch (PathEvent.Get<1>())
+			{
+			case FEventContext::EEvent::Added:
+				PathAddedEvent.Broadcast(Path);
+				break;
+			case FEventContext::EEvent::Removed:
+				PathRemovedEvent.Broadcast(Path);
+				break;
+			}
+		}
+		EventContext.PathEvents.Empty();
+	}
+
+	if (EventContext.AssetEvents.Num())
+	{
+		for (const TPair<FAssetData, FEventContext::EEvent>& AssetEvent : EventContext.AssetEvents)
+		{
+			const FAssetData& AssetData = AssetEvent.Get<0>();
+			switch (AssetEvent.Get<1>())
+			{
+			case FEventContext::EEvent::Added:
+				AssetAddedEvent.Broadcast(AssetData);
+				break;
+			case FEventContext::EEvent::Removed:
+				AssetRemovedEvent.Broadcast(AssetData);
+				break;
+			case FEventContext::EEvent::Updated:
+				AssetUpdatedEvent.Broadcast(AssetData);
+				break;
+			default:
+				checkNoEntry();
+				break;
+			}
+		}
+		EventContext.AssetEvents.Empty();
+	}
+	if (EventContext.RequiredLoads.Num())
+	{
+		for (const FString& RequiredLoad : EventContext.RequiredLoads)
+		{
+			LoadPackage(nullptr, *RequiredLoad, 0);
+		}
+		EventContext.RequiredLoads.Empty();
+	}
+}
+
+
+UAssetRegistryImpl::FPathAddedEvent& UAssetRegistryImpl::OnPathAdded()
+{
+	checkf(IsInGameThread(), TEXT("Registering to AssetRegistry events is not supported from multiple threads."));
+	return PathAddedEvent;
+}
+
+UAssetRegistryImpl::FPathRemovedEvent& UAssetRegistryImpl::OnPathRemoved()
+{
+	checkf(IsInGameThread(), TEXT("Registering to AssetRegistry events is not supported from multiple threads."));
+	return PathRemovedEvent;
+}
+
+UAssetRegistryImpl::FAssetAddedEvent& UAssetRegistryImpl::OnAssetAdded()
+{
+	checkf(IsInGameThread(), TEXT("Registering to AssetRegistry events is not supported from multiple threads."));
+	return AssetAddedEvent;
+}
+
+UAssetRegistryImpl::FAssetRemovedEvent& UAssetRegistryImpl::OnAssetRemoved()
+{
+	checkf(IsInGameThread(), TEXT("Registering to AssetRegistry events is not supported from multiple threads."));
+	return AssetRemovedEvent;
+}
+
+UAssetRegistryImpl::FAssetRenamedEvent& UAssetRegistryImpl::OnAssetRenamed()
+{
+	checkf(IsInGameThread(), TEXT("Registering to AssetRegistry events is not supported from multiple threads."));
+	return AssetRenamedEvent;
+}
+
+UAssetRegistryImpl::FAssetUpdatedEvent& UAssetRegistryImpl::OnAssetUpdated()
+{
+	checkf(IsInGameThread(), TEXT("Registering to AssetRegistry events is not supported from multiple threads."));
+	return AssetUpdatedEvent;
+}
+
+UAssetRegistryImpl::FInMemoryAssetCreatedEvent& UAssetRegistryImpl::OnInMemoryAssetCreated()
+{
+	checkf(IsInGameThread(), TEXT("Registering to AssetRegistry events is not supported from multiple threads."));
+	return InMemoryAssetCreatedEvent;
+}
+
+UAssetRegistryImpl::FInMemoryAssetDeletedEvent& UAssetRegistryImpl::OnInMemoryAssetDeleted()
+{
+	checkf(IsInGameThread(), TEXT("Registering to AssetRegistry events is not supported from multiple threads."));
+	return InMemoryAssetDeletedEvent;
+}
+
+UAssetRegistryImpl::FFilesLoadedEvent& UAssetRegistryImpl::OnFilesLoaded()
+{
+	checkf(IsInGameThread(), TEXT("Registering to AssetRegistry events is not supported from multiple threads."));
+	return FileLoadedEvent;
+}
+
+UAssetRegistryImpl::FFileLoadProgressUpdatedEvent& UAssetRegistryImpl::OnFileLoadProgressUpdated()
+{
+	checkf(IsInGameThread(), TEXT("Registering to AssetRegistry events is not supported from multiple threads."));
+	return FileLoadProgressUpdatedEvent;
+}
diff --git a/Engine/Source/Runtime/AssetRegistry/Private/AssetRegistry.h b/Engine/Source/Runtime/AssetRegistry/Private/AssetRegistry.h
--- a/Engine/Source/Runtime/AssetRegistry/Private/AssetRegistry.h
+++ b/Engine/Source/Runtime/AssetRegistry/Private/AssetRegistry.h
@@ -2,23 +2,12 @@
 
 #pragma once
 
-#include "CoreMinimal.h"
-
-#include "AssetRegistry/AssetData.h"
-#include "AssetRegistry/IAssetRegistry.h"
-#include "AssetRegistry/AssetRegistryState.h"
-#include "Containers/RingBuffer.h"
-#include "ModuleDescriptor.h"
-#include "PackageDependencyData.h"
-#include "PathTree.h"
-#include "Templates/UniquePtr.h"
+#include "AssetRegistryImpl.h"
+#include "UObject/Object.h"
 
 #include "AssetRegistry.generated.h"
 
-class FDependsNode;
-struct FARFilter;
-class FAssetDataGatherer;
-class FPackageReader;
+class FRWScopeLock;
 
 /**
  * The AssetRegistry singleton gathers information about .uasset files in the background so things
@@ -68,6 +57,7 @@ public:
 	virtual bool GetReferencers(FName PackageName, TArray<FName>& OutReferencers, EAssetRegistryDependencyType::Type InReferenceType) const override;
 	virtual bool GetReferencers(FName PackageName, TArray<FName>& OutReferencers, UE::AssetRegistry::EDependencyCategory Category = UE::AssetRegistry::EDependencyCategory::Package, const UE::AssetRegistry::FDependencyQuery& Flags = UE::AssetRegistry::FDependencyQuery()) const override;
 	virtual const FAssetPackageData* GetAssetPackageData(FName PackageName) const override;
+	virtual TOptional<FAssetPackageData> GetAssetPackageDataCopy(FName PackageName) const override;
 	virtual FName GetRedirectedObjectPath(const FName ObjectPath) const override;
 	virtual bool GetAncestorClassNames(FName ClassName, TArray<FName>& OutAncestorClassNames) const override;
 	virtual void GetDerivedClassNames(const TArray<FName>& ClassNames, const TSet<FName>& ExcludedClassNames, TSet<FName>& OutDerivedClassNames) const override;
@@ -107,15 +97,21 @@ public:
 	virtual void LoadPackageRegistryData(FArchive& Ar, FLoadPackageRegistryData& InOutData) const override;
 	virtual void LoadPackageRegistryData(const FString& PackageFilename, FLoadPackageRegistryData& InOutData) const override;
 	virtual void InitializeTemporaryAssetRegistryState(FAssetRegistryState& OutState, const FAssetRegistrySerializationOptions& Options, bool bRefreshExisting = false, const TMap<FName, FAssetData*>& OverrideData = TMap<FName, FAssetData*>()) const override;
+#if ASSET_REGISTRY_STATE_DUMPING_ENABLED
+	virtual void DumpState(const TArray<FString>& Arguments, TArray<FString>& OutPages, int32 LinesPerPage = 1) const override;
+#endif
+
 	virtual const FAssetRegistryState* GetAssetRegistryState() const override;
+	virtual TSet<FName> GetCachedEmptyPackagesCopy() const override;
 	virtual const TSet<FName>& GetCachedEmptyPackages() const override;
+	virtual bool ContainsTag(FName TagName) const override;
 	virtual void InitializeSerializationOptions(FAssetRegistrySerializationOptions& Options, const FString& PlatformIniName = FString()) const override;
 
 	DECLARE_DERIVED_EVENT( UAssetRegistryImpl, IAssetRegistry::FPathAddedEvent, FPathAddedEvent);
-	virtual FPathAddedEvent& OnPathAdded() override { return PathAddedEvent; }
+	virtual FPathAddedEvent& OnPathAdded() override;
 
 	DECLARE_DERIVED_EVENT( UAssetRegistryImpl, IAssetRegistry::FPathRemovedEvent, FPathRemovedEvent);
-	virtual FPathRemovedEvent& OnPathRemoved() override { return PathRemovedEvent; }
+	virtual FPathRemovedEvent& OnPathRemoved() override;
 
 	virtual void AssetCreated(UObject* NewAsset) override;
 	virtual void AssetDeleted(UObject* DeletedAsset) override;
@@ -124,28 +120,28 @@ public:
 	virtual void PackageDeleted(UPackage* DeletedPackage) override;
 
 	DECLARE_DERIVED_EVENT( UAssetRegistryImpl, IAssetRegistry::FAssetAddedEvent, FAssetAddedEvent);
-	virtual FAssetAddedEvent& OnAssetAdded() override { return AssetAddedEvent; }
+	virtual FAssetAddedEvent& OnAssetAdded() override;
 
 	DECLARE_DERIVED_EVENT( UAssetRegistryImpl, IAssetRegistry::FAssetRemovedEvent, FAssetRemovedEvent);
-	virtual FAssetRemovedEvent& OnAssetRemoved() override { return AssetRemovedEvent; }
+	virtual FAssetRemovedEvent& OnAssetRemoved() override;
 
 	DECLARE_DERIVED_EVENT( UAssetRegistryImpl, IAssetRegistry::FAssetRenamedEvent, FAssetRenamedEvent);
-	virtual FAssetRenamedEvent& OnAssetRenamed() override { return AssetRenamedEvent; }
+	virtual FAssetRenamedEvent& OnAssetRenamed() override;
 
 	DECLARE_DERIVED_EVENT( UAssetRegistryImpl, IAssetRegistry::FAssetUpdatedEvent, FAssetUpdatedEvent );
-	virtual FAssetUpdatedEvent& OnAssetUpdated() override { return AssetUpdatedEvent; }
+	virtual FAssetUpdatedEvent& OnAssetUpdated() override;
 
 	DECLARE_DERIVED_EVENT( UAssetRegistryImpl, IAssetRegistry::FInMemoryAssetCreatedEvent, FInMemoryAssetCreatedEvent );
-	virtual FInMemoryAssetCreatedEvent& OnInMemoryAssetCreated() override { return InMemoryAssetCreatedEvent; }
+	virtual FInMemoryAssetCreatedEvent& OnInMemoryAssetCreated() override;
 
 	DECLARE_DERIVED_EVENT( UAssetRegistryImpl, IAssetRegistry::FInMemoryAssetDeletedEvent, FInMemoryAssetDeletedEvent );
-	virtual FInMemoryAssetDeletedEvent& OnInMemoryAssetDeleted() override { return InMemoryAssetDeletedEvent; }
+	virtual FInMemoryAssetDeletedEvent& OnInMemoryAssetDeleted() override;
 
 	DECLARE_DERIVED_EVENT( UAssetRegistryImpl, IAssetRegistry::FFilesLoadedEvent, FFilesLoadedEvent );
-	virtual FFilesLoadedEvent& OnFilesLoaded() override { return FileLoadedEvent; }
+	virtual FFilesLoadedEvent& OnFilesLoaded() override;
 
 	DECLARE_DERIVED_EVENT( UAssetRegistryImpl, IAssetRegistry::FFileLoadProgressUpdatedEvent, FFileLoadProgressUpdatedEvent );
-	virtual FFileLoadProgressUpdatedEvent& OnFileLoadProgressUpdated() override { return FileLoadProgressUpdatedEvent; }
+	virtual FFileLoadProgressUpdatedEvent& OnFileLoadProgressUpdated() override;
 
 	virtual bool IsLoadingAssets() const override;
 
@@ -154,233 +150,80 @@ public:
 	UE_DEPRECATED(4.17, "IsUsingWorldAssets is now always true, remove any code that assumes it could be false")
 	static bool IsUsingWorldAssets() { return true; }
 
+	virtual void ReadLockEnumerateTagToAssetDatas(TFunctionRef<void(FName TagName, const TArray<const FAssetData*>& Assets)> Callback) const override;
+
 protected:
 	virtual void SetManageReferences(const TMultiMap<FAssetIdentifier, FAssetIdentifier>& ManagerMap, bool bClearExisting, UE::AssetRegistry::EDependencyCategory RecurseType, TSet<FDependsNode*>& ExistingManagedNodes, ShouldSetManagerPredicate ShouldSetManager = nullptr) override;
 	virtual bool SetPrimaryAssetIdForObjectPath(const FName ObjectPath, FPrimaryAssetId PrimaryAssetId) override;
-	virtual const FAssetData* GetCachedAssetDataForObjectPath(const FName ObjectPath) const override;
 
 private:
+	void InitializeEvents(UE::AssetRegistry::Impl::FInitializeContext& Context);
+	void Broadcast(UE::AssetRegistry::Impl::FEventContext& EventContext);
 
-	void InitRedirectors();
-	/** Read script packages if all initial plugins have been loaded, otherwise register with the callback to read them once read. */
-	void RegisterReadOfScriptPackages();
-	void OnPluginLoadingPhaseComplete(ELoadingPhase::Type LoadingPhase, bool bPhaseSuccessful);
-	/** If collecting dependencies, create an FAssetPackageData for every script package, to make sure dependencies can find a data for them. */
-	void ReadScriptPackages();
-
-	/**
-	 * Internal implementation for RunAssetsThroughFilter and UseFilterToExcludeAssets
-	 * @param AssetDataList The asset list to filter
-	 * @param Filter The filter to use
-	 * @param bInclusive Run the filter as inclusive (assets left in the list passes the filter) or exclusive (assets left in the list fail the filter)
-	 */
-	void InternalRunAssetsThroughFilter(TArray<FAssetData>& AssetDataList, const FARFilter& Filter, bool bInclusive) const;
-
-	/* Construct the gatherer if it does not already exist */
-	void ConstructGatherer();
-	/**  Called to set up timing variables and launch the on-constructor SearchAllAssets async call */
-	void SearchAllAssetsInitialAsync();
-
-	typedef TFunctionRef<void(const TRingBuffer<FAssetData*>&)> FAssetsFoundCallback;
-	/** Read any results from the gatherer, if it is running asynchronously */
-	void TickGatherer(float DeltaTime, TOptional<FAssetsFoundCallback> AssetsFoundCallback = TOptional<FAssetsFoundCallback>());
-	/** Look for and load a single AssetData result from the gatherer. */
-	void TickGatherPackage(const FString& PackageName);
-	/** Internal handler for ScanPathsSynchronous */
-	void ScanPathsSynchronousInternal(const TArray<FString>& InDirs, const TArray<FString>& InFiles, bool bForceRescan, bool bIgnoreBlackListScanFilters, TArray<FName>* OutFoundAssets);
-
-	/** Called every tick to when data is retrieved by the background asset search. If TickStartTime is < 0, the entire list of gathered assets will be cached. Also used in sychronous searches */
-	void AssetSearchDataGathered(const double TickStartTime, TRingBuffer<FAssetData*>& AssetResults);
-
-	/** Called every tick when data is retrieved by the background path search. If TickStartTime is < 0, the entire list of gathered assets will be cached. Also used in sychronous searches */
-	void PathDataGathered(const double TickStartTime, TRingBuffer<FString>& PathResults);
-
-	/** Called every tick when data is retrieved by the background dependency search */
-	void DependencyDataGathered(const double TickStartTime, TRingBuffer<FPackageDependencyData>& DependsResults);
-
-	/** Called every tick when data is retrieved by the background search for cooked packages that do not contain asset data */
-	void CookedPackageNamesWithoutAssetDataGathered(const double TickStartTime, TRingBuffer<FString>& CookedPackageNamesWithoutAssetDataResults);
-
-	/** Adds an asset to the empty package list which contains packages that have no assets left in them */
-	void AddEmptyPackage(FName PackageName);
-
-	/** Removes an asset from the empty package list because it is no longer empty */
-	bool RemoveEmptyPackage(FName PackageName);
-
-	/** Adds a path to the cached paths tree. Returns true if the path was added to the tree, as opposed to already existing in the tree */
-	bool AddAssetPath(FName PathToAdd);
-
-	/** Removes a path to the cached paths tree. Returns true if successful. */
-	bool RemoveAssetPath(FName PathToRemove, bool bEvenIfAssetsStillExist = false);
-
-	/** Helper function to return the name of an object, given the objects export text path */
-	FString ExportTextPathToObjectName(const FString& InExportTextPath) const;
-
-	/** Adds the asset data to the lookup maps */
-	void AddAssetData(FAssetData* AssetData);
-
-	/** Updates an existing asset data with the new value and updates lookup maps */
-	void UpdateAssetData(FAssetData* AssetData, const FAssetData& NewAssetData);
-
-	/** Removes the asset data from the lookup maps */
-	bool RemoveAssetData(FAssetData* AssetData);
-
-	/** Removes the asset data associated with this package from the look-up maps */
-	void RemovePackageData(const FName PackageName);
-
-	/** Adds a list of files which will be searched for asset data */
-	void AddFilesToSearch (const TArray<FString>& Files);
+	bool OnResolveRedirect(const FString& InPackageName, FString& OutPackageName);
 
 #if WITH_EDITOR
 	/** Called when a file in a content directory changes on disk */
 	void OnDirectoryChanged(const TArray<struct FFileChangeData>& Files);
 
 	/** Called when an asset is loaded, it will possibly update the cache */
-	void OnAssetLoaded(UObject *AssetLoaded);
-
-	/** Process Loaded Assets to update cache */
-	void ProcessLoadedAssetsToUpdateCache(const double TickStartTime, bool bBecameIdle);
-
-	/** Update Redirect collector with redirects loaded from asset registry */
-	void UpdateRedirectCollector();
-#endif // WITH_EDITOR
+	void OnAssetLoaded(UObject* AssetLoaded);
+#endif
 
 	/**
-	 * Called by the engine core when a new content path is added dynamically at runtime.  This is wired to 
+	 * Called by the engine core when a new content path is added dynamically at runtime.  This is wired to
 	 * FPackageName's static delegate.
 	 *
 	 * @param	AssetPath		The new content root asset path that was added (e.g. "/MyPlugin/")
 	 * @param	FileSystemPath	The filesystem path that the AssetPath is mapped to
 	 */
-	void OnContentPathMounted( const FString& AssetPath, const FString& FileSystemPath );
+	void OnContentPathMounted(const FString& AssetPath, const FString& FileSystemPath);
 
 	/**
-	 * Called by the engine core when a content path is removed dynamically at runtime.  This is wired to 
+	 * Called by the engine core when a content path is removed dynamically at runtime.  This is wired to
 	 * FPackageName's static delegate.
 	 *
 	 * @param	AssetPath		The new content root asset path that was added (e.g. "/MyPlugin/")
 	 * @param	FileSystemPath	The filesystem path that the AssetPath is mapped to
 	 */
-	void OnContentPathDismounted( const FString& AssetPath, const FString& FileSystemPath );
-
-	/** Called to refresh the native classes list, called at end of engine initialization */
-	void RefreshNativeClasses();
+	void OnContentPathDismounted(const FString& AssetPath, const FString& FileSystemPath);
 
-	/** Returns the names of all subclasses of the class whose name is ClassName */
-	void GetSubClasses(const TArray<FName>& InClassNames, const TSet<FName>& ExcludedClassNames, TSet<FName>& SubClassNames) const;
-	void GetSubClasses_Recursive(FName InClassName, TSet<FName>& SubClassNames, TSet<FName>& ProcessedClassNames, const TMap<FName, TArray<FName>>& ReverseInheritanceMap, const TSet<FName>& ExcludedClassNames) const;
+	/** Called to refresh the native classes list, called at end of engine initialization. */
+	void OnRefreshNativeClasses();
 
-	/** Finds all class names of classes capable of generating new UClasses */
-	void CollectCodeGeneratorClasses() const;
-
-	/** Updates TempCachedInheritanceMap from native classes */
-	void UpdateTemporaryCaches() const;
-
-	/** Deletes any temporary cached data as needed */
-	void ClearTemporaryCaches() const;
-
-	/** Initialize the scan filters from the ini */
-	void InitializeBlacklistScanFiltersFromIni();
-
-	bool ResolveRedirect(const FString& InPackageName, FString& OutPackageName);
-
-	/** Internal helper which processes a given state and adds its contents to the current registry */
-	void CachePathsFromState(const FAssetRegistryState& InState);
-
-	enum class EARFilterMode : uint8
-	{
-		/** Include things that pass the filter; include everything if the filter is empty */
-		Inclusive,
-		/** Exclude things that pass the filter; exclude nothing if the filter is empty */
-		Exclusive,
-	};
-
-	/**
-	 * Given an asset data, say whether it would pass the filter based on the inclusion/exclusion mode used.
-	 *  - If an asset data passes a filter, then in inclusive mode it will return true, and in exclusive mode it will return false.
-	 *  - If an asset data fails a filter, then in inclusive mode it will return false, and in exclusive mode it will return true.
-	 *  - If the filter is empty, then in inclusive mode it will return true, and in exclusive mode it will return false.
-	 */
-	bool RunAssetThroughFilterImpl(const FAssetData& AssetData, const FARCompiledFilter& Filter, const EARFilterMode FilterMode) const;
-	bool RunAssetThroughFilterImpl_Unchecked(const FAssetData& AssetData, const FARCompiledFilter& Filter, const bool bPassFilterValue) const;
+	/** Called from the PluginManager's loading phase, used to scan classes that were loaded by plugins. */
+	void OnPluginLoadingPhaseComplete(ELoadingPhase::Type LoadingPhase, bool bPhaseSuccessful);
 
-	/**
-	 * Given an array of asset data, trim the items that fail the filter based on the inclusion/exclusion mode used.
-	 *  - In inclusive mode it will remove all assets that fail the filter, and in exclusive mode it will remove all assets that pass the filter.
-	 *  - If the filter is empty, then the array will be untouched.
-	 */
-	void RunAssetsThroughFilterImpl(TArray<FAssetData>& AssetDataList, const FARFilter& Filter, const EARFilterMode FilterMode) const;
+	/** Shared helper for Scan*Synchronous function */
+	void ScanPathsSynchronousInternal(const TArray<FString>& InDirs, const TArray<FString>& InFiles,
+		bool bInForceRescan, bool bInIgnoreBlackListScanFilters);
 
+#if WITH_EDITOR
+	/** Create FAssetData from any loaded UObject assets and store the updated AssetData in the state */
+	void ProcessLoadedAssetsToUpdateCache(UE::AssetRegistry::Impl::FEventContext& EventContext,
+		const double TickStartTime, bool bIsIdle);
+#endif
 	/**
-	 * Returns true if path belongs to one of the mount points provided
-	 *
-	 * @param	Path				Path to check if mounted, example "/MyPlugin/SomeAsset"
-	 * @param	MountPointsNoTrailingSlashes		Mount points without the trailing slash. Example: "/MyPlugin"
-	 * @param	StringBuffer		String buffer to avoid re-allocation performance hit when searching TSet
+	 * Remain under the given lock and return an InheritanceContext based on the appropriate choice of the persistent
+	 * caching buffer or the function-scope-only passed in StackBuffer. Mark whether the buffer needs to be updated
+	 * before being used. If the buffer needs to be updated and its the persistent buffer (which is protected data),
+	 * convert the given lock to a write lock if not one already.
 	 */
-	bool IsPathMounted(const FString& Path, const TSet<FString>& MountPointsNoTrailingSlashes, FString& StringBuffer) const;
+	void GetInheritanceContextWithRequiredLock(FRWScopeLock& InOutScopeLock,
+		UE::AssetRegistry::Impl::FClassInheritanceContext& InheritanceContext,
+		UE::AssetRegistry::Impl::FClassInheritanceBuffer& StackBuffer);
 
-	/**
-	 * Reads Asset file from a previously initialized package reader
-	 * 
-	 * @param PackageReader	Previsouly initialized package reader that maps to the data to be read
-	 * @param AssetDataList	List of the read asset data
-	 */
-	bool ReadAssetFile(FPackageReader& PackageReader, FLoadPackageRegistryData& InOutData) const;
 private:
-	
-	/** Internal state of the cached asset registry */
-	FAssetRegistryState State;
-
-	/** Default options used for serialization */
-	FAssetRegistrySerializationOptions SerializationOptions;
-
-	/** The set of empty package names (packages which contain no assets but have not yet been saved) */
-	TSet<FName> CachedEmptyPackages;
-
-	/** The map of classes to their parents, only full for offline blueprints */
-	TMap<FName, FName> CachedBPInheritanceMap;
-
-	/** If true, search caching is enabled */
-	bool bIsTempCachingEnabled;
-
-	/** If true, search caching is enabled permanently */
-	bool bIsTempCachingAlwaysEnabled;
-
-	/** A temporary fully cached list including native classes */
-	mutable TMap<FName, FName> TempCachedInheritanceMap;
-
-	/** A reverse map of TempCachedInheritanceMap, only kept during temp caching */
-	mutable TMap<FName, TArray<FName>> TempReverseInheritanceMap;
-
-	/** If true, temp caching has been computed and is valid.
-	    Set this to false when changing something that might invalidate the cache so it gets recomputed on-demand.
-	*/
-	mutable bool bIsTempCachingUpToDate;
-
-	/** Contains a snapshot of GetRegisteredClassesVersionNumber() at the time of caching so the cache can
-	    be invalidated whenever registered classes have changed.
-	*/
-	mutable uint64 TempCachingRegisteredClassesVersionNumber;
-	mutable uint64 ClassGeneratorNamesRegisteredClassesVersionNumber;
-
-	/** If true, will cache AssetData loaded from in memory assets back into the disk cache */
-	bool bUpdateDiskCacheAfterLoad;
 
-	/** The tree of known cached paths that assets may reside within */
-	FPathTree CachedPathTree;
+	UE::AssetRegistry::FAssetRegistryImpl GuardedData;
 
-	/** Async task that gathers asset information from disk */
-	TUniquePtr<FAssetDataGatherer> GlobalGatherer;
+	/** Lock guarding the GuardedData */
+	mutable FRWLock InterfaceLock;
 
-	/** A list of results that were gathered from the background thread that are waiting to get processed by the main thread */
-	TRingBuffer<FAssetData*> BackgroundAssetResults;
-	TRingBuffer<FString> BackgroundPathResults;
-	TRingBuffer<FPackageDependencyData> BackgroundDependencyResults;
-	TRingBuffer<FString> BackgroundCookedPackageNamesWithoutAssetDataResults;
-
-	/** The max number of results to process per tick */
-	float MaxSecondsPerFrame;
+#if WITH_EDITOR
+	/** Handles to all registered OnDirectoryChanged delegates */
+	TMap<FString, FDelegateHandle> OnDirectoryChangedDelegateHandles;
+#endif
 
 	/** The delegate to execute when an asset path is added to the registry */
 	FPathAddedEvent PathAddedEvent;
@@ -412,46 +255,6 @@ private:
 	/** The delegate to execute while loading files to update progress */
 	FFileLoadProgressUpdatedEvent FileLoadProgressUpdatedEvent;
 
-	/** The start time of the full asset search */
-	double FullSearchStartTime;
-	double AmortizeStartTime;
-	double TotalAmortizeTime;
-
-	/** Flag to indicate if we used an initial async search */
-	bool bInitialSearchStarted;
-	/** Flag to indicate if the initial background search has completed */
-	bool bInitialSearchCompleted;
-	/** Flag that indicates the background search is idle, so we can take actions when idle changes */
-	bool bGatherIdle;
-
-	/** Enables extra check to make sure path still mounted before adding. Removing mount point can happen between scan (background thread + multiple ticks and the add). */
-	bool bVerifyMountPointAfterGather;
-
-	/** Record whether SearchAllAssets has been called; if so we will also search new mountpoints when added */
-	bool bSearchAllAssets;
-
-	/** List of all class names derived from Blueprint (including Blueprint itself) */
-	mutable TSet<FName> ClassGeneratorNames;
-
-	/** Handles to all registered OnDirectoryChanged delegates */
-	TMap<FString, FDelegateHandle> OnDirectoryChangedDelegateHandles;
-
-	struct FAssetRegistryPackageRedirect
-	{
-	public:
-		FAssetRegistryPackageRedirect(const FString& InSourcePackageName, const FString& InDestPackageName) : SourcePackageName(InSourcePackageName), DestPackageName(InDestPackageName) { }
-		FString SourcePackageName;
-		FString DestPackageName;
-	};
-	TArray<FAssetRegistryPackageRedirect> PackageRedirects;
-
-#if WITH_EDITOR
-	/** List of loaded objects that need to be processed */
-	TRingBuffer<TWeakObjectPtr<UObject>> LoadedAssetsToProcess;
-
-	/** The set of object paths that have had their disk cache updated from the in memory version */
-	TSet<FName> AssetDataObjectPathsUpdatedOnLoad;
-#endif
-
+	UE::AssetRegistry::Impl::FEventContext DeferredEvents;
 };
 PRAGMA_ENABLE_DEPRECATION_WARNINGS
diff --git a/Engine/Source/Runtime/AssetRegistry/Private/AssetRegistryImpl.h b/Engine/Source/Runtime/AssetRegistry/Private/AssetRegistryImpl.h
new file mode 100644
--- /dev/null
+++ b/Engine/Source/Runtime/AssetRegistry/Private/AssetRegistryImpl.h
@@ -0,0 +1,491 @@
+// Copyright Epic Games, Inc. All Rights Reserved.
+
+#pragma once
+
+#include "AssetRegistry/AssetData.h"
+#include "AssetRegistry/IAssetRegistry.h"
+#include "AssetRegistry/AssetRegistryState.h"
+#include "Containers/RingBuffer.h"
+#include "HAL/CriticalSection.h"
+#include "Math/NumericLimits.h"
+#include "ModuleDescriptor.h"
+#include "PackageDependencyData.h"
+#include "PathTree.h"
+#include "Templates/Function.h"
+#include "Templates/UniquePtr.h"
+
+class FDependsNode;
+struct FARFilter;
+class FAssetDataGatherer;
+struct FFileChangeData;
+class FPackageReader;
+
+namespace UE
+{
+namespace AssetRegistry
+{
+
+namespace Impl
+{
+	struct FEventContext;
+	struct FInitializeContext;
+	struct FScanPathContext;
+	struct FClassInheritanceContext;
+}
+
+namespace Impl
+{
+	/** Container for global class inheritance data; the AssetRegistry may use a persistent buffer or a stack local buffer. */
+	struct FClassInheritanceBuffer
+	{
+		/** Map from Class->SuperClass for all classes including native classes and blueprint classes. Updated on demand. */
+		TMap<FName, FName> InheritanceMap;
+		/** Map from Class->(All subclasses) for all classes including native classes and blueprint classes. Updated on demand. */
+		TMap<FName, TArray<FName>> ReverseInheritanceMap;
+		/** Snapshot of GetRegisteredClassesVersionNumber() at the time of the last update, to invalidate on changes to classes. */
+		uint64 RegisteredClassesVersionNumber = MAX_uint64;
+		/** Dirty flag to invalidate on other changes requiring a recompute. */
+		bool bDirty = true;
+
+		/** Report whether the dirty flag and other checks indicate this buffer does not need to be updated. */
+		bool IsUpToDate(uint64 CurrentClassesVersionNumber) const;
+		/** Delete data and free allocations. */
+		void Clear();
+		/** Report size of dynamic allocations. */
+		uint32 GetAllocatedSize() const;
+	};
+}
+
+/**
+ * Threading helper class for UAssetRegistryImpl that holds all of the data.
+ * This class is not threadsafe; all const functions must be called from within a ReadLock and all non-const 
+ * from within a WriteLock.
+ */
+class FAssetRegistryImpl
+{
+public:
+	/** Constructor initializes as little as possible; meaningful operations are first done during Initialize. */
+	FAssetRegistryImpl();
+	/** Construct the AssetRegistryImpl, including initial scans if applicable. */
+	void Initialize(Impl::FInitializeContext& Context);
+
+
+	// Helpers for functions of the same name from UAssetRegistryImpl
+
+
+	bool HasAssets(const FName PackagePath, const bool bRecursive) const;
+	FName GetRedirectedObjectPath(const FName ObjectPath) const;
+	bool GetAncestorClassNames(Impl::FClassInheritanceContext& InheritanceContext, FName ClassName,
+		TArray<FName>& OutAncestorClassNames) const;
+	void CompileFilter(Impl::FClassInheritanceContext& InheritanceContext, const FARFilter& InFilter,
+		FARCompiledFilter& OutCompiledFilter) const;
+	void SetTemporaryCachingMode(bool bEnable);
+	void SetTemporaryCachingModeInvalidated();
+	bool AddPath(Impl::FEventContext& EventContext, const FString& PathToAdd);
+	void SearchAllAssets(Impl::FEventContext& EventContext, bool bSynchronousSearch);
+	void ScanPathsSynchronous(Impl::FScanPathContext& Context);
+	void PrioritizeSearchPath(const FString& PathToPrioritize);
+	void ScanModifiedAssetFiles(Impl::FEventContext& EventContext, const TArray<FString>& InFilePaths);
+	void Serialize(FArchive& Ar, Impl::FEventContext& EventContext);
+	void AppendState(Impl::FEventContext& EventContext, const FAssetRegistryState& InState);
+	void GetAllocatedSize(bool bLogDetailed, uint32& StateSize, uint32& StaticSize, uint32& SearchSize) const;
+	bool IsLoadingAssets() const;
+	void SetManageReferences(const TMultiMap<FAssetIdentifier, FAssetIdentifier>& ManagerMap,
+		bool bClearExisting, UE::AssetRegistry::EDependencyCategory RecurseType,
+		TSet<FDependsNode*>& ExistingManagedNodes, IAssetRegistry::ShouldSetManagerPredicate ShouldSetManager);
+	bool SetPrimaryAssetIdForObjectPath(Impl::FEventContext& EventContext, const FName ObjectPath, FPrimaryAssetId PrimaryAssetId);
+	bool ResolveRedirect(const FString& InPackageName, FString& OutPackageName) const;
+#if WITH_EDITOR
+	void OnDirectoryChanged(Impl::FEventContext& EventContext, TArray<FFileChangeData>& FileChangesProcessed);
+	void OnAssetLoaded(UObject* AssetLoaded);
+#endif
+	void OnContentPathMounted(Impl::FEventContext& EventContext, const FString& InAssetPath,
+		const FString& AssetPathWithTrailingSlash, const FString& FileSystemPath);
+	void OnContentPathDismounted(Impl::FEventContext& EventContext, const FString& InAssetPath,
+		const FString& AssetPathNoTrailingSlash, const FString& FileSystemPath);
+
+
+	// Other helper functions called by UAssetRegistryImpl
+
+
+	/** Update cached values that need to be reread when IPluginManager reaches PostEngineInit phase */
+	void OnPostEngineInit(bool bPhaseSuccessful);
+	/** Update cached values that need to be reread from the Engine's PostEngineInit callback */
+	void RefreshNativeClasses();
+
+	/** Enumerate assets in the State, filtering by filter and package not in PackagesToSkip */
+	void EnumerateDiskAssets(const FARCompiledFilter& InFilter, TSet<FName>& PackagesToSkip,
+		TFunctionRef<bool(const FAssetData&)> Callback) const;
+	/** Enumerate assets in the State, filtering by package not in PackagesToSkip */
+	void EnumerateAllDiskAssets(TSet<FName>& PackageNamesToSkip, TFunctionRef<bool(const FAssetData&)> Callback) const;
+
+	/** Delete any temporary allocated objects that were passed out to external callers and hopefully have been deleted by now */
+	UE_DEPRECATED(5.0, "Supports deprecated functions returning pointers")
+	void TickDeletes();
+	/** Callback type for TickGatherer */
+	typedef TFunctionRef<void(const TRingBuffer<FAssetData*>&)> FAssetsFoundCallback;
+	/** Read any results from the gatherer, if it is running asynchronously */
+	void TickGatherer(Impl::FEventContext& EventContext, const double TickStartTime, bool& bOutIdle, bool& bOutInterrupted,
+		TOptional<FAssetsFoundCallback> AssetsFoundCallback = TOptional<FAssetsFoundCallback>());
+	/** Look for and load a single AssetData result from the gatherer. */
+	void TickGatherPackage(Impl::FEventContext& EventContext, const FString& PackageName, const FString& LocalPath);
+#if WITH_EDITOR
+	void GetProcessLoadedAssetsBatch(TArray<UObject*>& OutLoadedAssets, uint32 BatchSize);
+	void PushProcessLoadedAssetsBatch(Impl::FEventContext& EventContext,
+		TArrayView<FAssetData> LoadedAssetDatas, TArrayView<UObject*> UnprocessedFromBatch);
+#endif
+
+	/** Adds an asset to the empty package list which contains packages that have no assets left in them */
+	void AddEmptyPackage(FName PackageName);
+	/** Removes an asset from the empty package list because it is no longer empty */
+	bool RemoveEmptyPackage(FName PackageName);
+	/** Adds a path to the cached paths tree. Returns true if the path was added and did not already exist */
+	bool AddAssetPath(Impl::FEventContext& EventContext, FName PathToAdd);
+	/** Removes a path to the cached paths tree. Returns true if successful. */
+	bool RemoveAssetPath(Impl::FEventContext& EventContext, FName PathToRemove, bool bEvenIfAssetsStillExist = false);
+	/** Removes the asset data associated with this package from the look-up maps */
+	void RemovePackageData(Impl::FEventContext& EventContext, const FName PackageName);
+
+	/** Returns the names of all subclasses of the class whose name is ClassName */
+	void GetSubClasses(Impl::FClassInheritanceContext& InheritanceContext, const TArray<FName>& InClassNames,
+		const TSet<FName>& ExcludedClassNames, TSet<FName>& SubClassNames) const;
+
+	bool IsUpdateDiskCacheAfterLoad() const { return bUpdateDiskCacheAfterLoad; }
+	bool IsInitialSearchCompleted() const { return bInitialSearchCompleted; }
+	bool IsTempCachingEnabled() const { return bIsTempCachingEnabled; }
+	bool IsTempCachingAlwaysEnabled() const { return bIsTempCachingAlwaysEnabled; }
+	bool IsInitialSearchStarted() const { return bInitialSearchStarted; }
+	bool IsSearchAllAssets() const { return bSearchAllAssets; }
+	Impl::FClassInheritanceBuffer& GetTempCachedInheritanceBuffer() { return TempCachedInheritanceBuffer; }
+	uint64 GetClassGeneratorNamesRegisteredClassesVersionNumber() const { return ClassGeneratorNamesRegisteredClassesVersionNumber; }
+	/** Get a copy of the cached serialization options that were parsed from ini */
+	void CopySerializationOptions(FAssetRegistrySerializationOptions& OutOptions) const;
+
+	const FAssetRegistryState& GetState() const;
+	const FPathTree& GetCachedPathTree() const;
+	const TSet<FName>& GetCachedEmptyPackages() const;
+	/** Find the AssetPackageData for the given PackageName and return a pointer to it or nullptr if not found */
+	const FAssetPackageData* GetAssetPackageData(FName PackageName);
+
+	/** Same as UE::AssetRegistry::Filtering::ShouldSkipAsset, but can be read from any thread under readlock */
+	bool ShouldSkipAsset(FName AssetClass, uint32 PackageFlags) const;
+	bool ShouldSkipAsset(const UObject* InAsset) const;
+
+	/** Finds all class names of classes capable of generating new UClasses */
+	void CollectCodeGeneratorClasses();
+
+
+private:
+
+	/**
+	 * Enumerate all AssetData in memory and in the State that match the given path.
+	 * The returned AssetData for objects present in memory have no AssetRegistryTags calculated
+	 */
+	void EnumerateAssetsByPathNoTags(FName PackagePath,
+		TFunctionRef<bool(const FAssetData&)> Callback, bool bRecursive, bool bIncludeOnlyOnDiskAssets) const;
+	/** Create redirectors read by this->GetRedirectedObjectPath, based on plugin settings */
+	void InitRedirectors(Impl::FEventContext& EventContext, bool& bOutRedirectorsNeedSubscribe);
+	/** If collecting dependencies, create an FAssetPackageData for every script package, to make sure dependencies can find a data for them. */
+	void ReadScriptPackages();
+	/* Construct the gatherer if it does not already exist */
+	void ConstructGatherer();
+	/**  Called to set up timing variables and launch the on-constructor SearchAllAssets async call */
+	void SearchAllAssetsInitialAsync(Impl::FEventContext& EventContext);
+	/**
+	 * Called every tick to when data is retrieved by the background asset search.
+	 * If TickStartTime is < 0, the entire list of gathered assets will be cached. Also used in sychronous searches
+	 */
+	void AssetSearchDataGathered(Impl::FEventContext& EventContext, const double TickStartTime, TRingBuffer<FAssetData*>& AssetResults);
+
+	/**
+	 * Called every tick when data is retrieved by the background path search.
+	 * If TickStartTime is < 0, the entire list of gathered assets will be cached. Also used in sychronous searches
+	 */
+	void PathDataGathered(Impl::FEventContext& EventContext, const double TickStartTime, TRingBuffer<FString>& PathResults);
+
+	/** Called every tick when data is retrieved by the background dependency search */
+	void DependencyDataGathered(const double TickStartTime, TRingBuffer<FPackageDependencyData>& DependsResults);
+
+	/** Called every tick when data is retrieved by the background search for cooked packages that do not contain asset data */
+	void CookedPackageNamesWithoutAssetDataGathered(Impl::FEventContext& EventContext, const double TickStartTime,
+		TRingBuffer<FString>& CookedPackageNamesWithoutAssetDataResults, bool& bOutInterrupted);
+
+	/** Adds the asset data to the lookup maps */
+	void AddAssetData(Impl::FEventContext& EventContext, FAssetData* AssetData);
+
+	/** Updates an existing asset data with the new value and updates lookup maps */
+	void UpdateAssetData(Impl::FEventContext& EventContext, FAssetData* AssetData, const FAssetData& NewAssetData);
+
+	/** Removes the asset data from the lookup maps */
+	bool RemoveAssetData(Impl::FEventContext& EventContext, FAssetData* AssetData);
+
+#if WITH_EDITOR
+	/** Update Redirect collector with redirects loaded from asset registry */
+	void UpdateRedirectCollector();
+#endif // WITH_EDITOR
+	void GetSubClasses_Recursive(Impl::FClassInheritanceContext& InheritanceContext, FName InClassName,
+		TSet<FName>& SubClassNames, TSet<FName>& ProcessedClassNames, const TSet<FName>& ExcludedClassNames) const;
+
+	/** Internal helper which processes a given state and adds its contents to the current registry */
+	void CachePathsFromState(Impl::FEventContext& EventContext, const FAssetRegistryState& InState);
+
+	/** Updates OutBuffer from loaded classes and registered-in-CachedBPInheritanceMap blueprint classes */
+	void UpdateInheritanceBuffer(Impl::FClassInheritanceBuffer& OutBuffer) const;
+
+private:
+
+	/** Internal state of the cached asset registry */
+	FAssetRegistryState State;
+
+	/** Default options used for serialization */
+	FAssetRegistrySerializationOptions SerializationOptions;
+
+	/** The set of empty package names (packages which contain no assets but have not yet been saved) */
+	TSet<FName> CachedEmptyPackages;
+
+	/** The map of classes to their parents, only full for offline blueprints */
+	TMap<FName, FName> CachedBPInheritanceMap;
+
+	/** If true, search caching is enabled */
+	bool bIsTempCachingEnabled;
+
+	/** If true, search caching is enabled permanently */
+	bool bIsTempCachingAlwaysEnabled;
+
+	/** Persistent InheritanceBuffer used when SetTemporaryCachingMode is on */
+	Impl::FClassInheritanceBuffer TempCachedInheritanceBuffer;
+
+	uint64 ClassGeneratorNamesRegisteredClassesVersionNumber;
+
+	/** If true, will cache AssetData loaded from in memory assets back into the disk cache */
+	bool bUpdateDiskCacheAfterLoad;
+
+	/** The tree of known cached paths that assets may reside within */
+	FPathTree CachedPathTree;
+
+	/** Async task that gathers asset information from disk */
+	TUniquePtr<FAssetDataGatherer> GlobalGatherer;
+
+	/** A list of results from the background thread that are waiting to get processed by the main thread */
+	TRingBuffer<FAssetData*> BackgroundAssetResults;
+	TRingBuffer<FString> BackgroundPathResults;
+	TRingBuffer<FPackageDependencyData> BackgroundDependencyResults;
+	TRingBuffer<FString> BackgroundCookedPackageNamesWithoutAssetDataResults;
+
+	UE_DEPRECATED(5.0, "DelayDelete is only intended to support deprecated functions.")
+	TArray<TUniqueFunction<void()>> DeleteActions;
+
+	/** The start time of the full asset search */
+	double FullSearchStartTime;
+	double AmortizeStartTime;
+	double TotalAmortizeTime;
+	/** The highest number of pending results observed during initial gathering */
+	int32 HighestPending = 0;
+
+	/** Flag to indicate if we used an initial async search */
+	bool bInitialSearchStarted;
+	/** Flag to indicate if the initial background search has completed */
+	bool bInitialSearchCompleted;
+	/** Flag that indicates the background search is idle, so we can take actions when idle changes */
+	bool bGatherIdle;
+
+	/**
+	 * Enables extra check to make sure path still mounted before adding.
+	 * Removing mount point can happen between scan (background thread + multiple ticks and the add).
+	 */
+	bool bVerifyMountPointAfterGather;
+
+	/** Record whether SearchAllAssets has been called; if so we will also search new mountpoints when added */
+	bool bSearchAllAssets;
+
+	/** List of all class names derived from Blueprint (including Blueprint itself) */
+	TSet<FName> ClassGeneratorNames;
+
+	struct FAssetRegistryPackageRedirect
+	{
+	public:
+		FAssetRegistryPackageRedirect(const FString& InSourcePackageName, const FString& InDestPackageName)
+			: SourcePackageName(InSourcePackageName), DestPackageName(InDestPackageName)
+		{
+		}
+		FString SourcePackageName;
+		FString DestPackageName;
+	};
+	TArray<FAssetRegistryPackageRedirect> PackageRedirects;
+
+#if WITH_EDITOR
+	/** List of loaded objects that need to be processed */
+	TRingBuffer<TWeakObjectPtr<UObject>> LoadedAssetsToProcess;
+
+	/** The set of object paths that have had their disk cache updated from the in memory version */
+	TSet<FName> AssetDataObjectPathsUpdatedOnLoad;
+#endif
+#if WITH_ENGINE && WITH_EDITOR
+	/** Class names that return true for IsAsset but which should not be treated as assets in uncooked packages */
+	TSet<FName> SkipUncookedClasses;
+	/** Class names that return true for IsAsset but which should not be treated as assets in cooked packages */
+	TSet<FName> SkipCookedClasses;
+#endif
+
+	friend struct Impl::FClassInheritanceContext;
+};
+
+namespace Impl
+{
+
+/**
+ * Stores events that need to be broadcasted from the AssetRegistry; these events are added to 
+ * by functions within a lock, and are broadcast after the lock is dropped. The broadcast is also deferred
+ * until on the GameThread.
+ */
+struct FEventContext
+{
+	enum class EEvent
+	{
+		Added,
+		Removed,
+		Updated,
+	};
+	TOptional<IAssetRegistry::FFileLoadProgressUpdateData> ProgressUpdateData;
+	TArray<TPair<FString, EEvent>> PathEvents;
+	TArray<TPair<FAssetData, EEvent>> AssetEvents;
+	TArray<FString> RequiredLoads;
+	bool bFileLoadedEventBroadcast = false;
+
+	/** Remove all stored events */
+	void Clear();
+	/** Report whether this context has any stored events */
+	bool IsEmpty() const;
+	/** Add all events from other onto this context's collection of events */
+	void Append(FEventContext&& Other);
+};
+
+/*
+ * An accessor for the inheritance map and related data for functions that need them; the source
+ * of this data is an internal detail unneeded by the functions reading it. Readers of the data call
+ * ConditionalUpdate to ensure it is up to date before reading it.
+ */
+struct FClassInheritanceContext
+{
+	/** The Buffer providing the data for this context. */
+	FClassInheritanceBuffer* Buffer = nullptr;
+	/** Back pointer to the AssetRegistryImpl to update the Buffer if necessary */
+	FAssetRegistryImpl* AssetRegistryImpl = nullptr;
+	/** Whether ConditionalUpdate needs to update inheritance data */
+	bool bInheritanceMapUpToDate = false;
+	/** Whether ConditionalUpdate needs to update code generate classes before it can update inheritance data */
+	bool bCodeGeneratorClassesUpToDate = false;
+
+	/** Set the context to use the data from the given buffer */
+	void BindToBuffer(FClassInheritanceBuffer& InBuffer, FAssetRegistryImpl& InAssetRegistryImpl,
+		bool bInInheritanceMapUpToDate, bool bInCodeGeneratorClassesUpToDate);
+	/** Update the backing buffer if it is out of date */
+	void ConditionalUpdate();
+};
+
+/** Context to transform, collect, and pass along ScanPathsSynchronous arguments and results */
+struct FScanPathContext
+{
+	FScanPathContext(FEventContext& InEventContext, const TArray<FString>& InDirs, const TArray<FString>& InFiles,
+		bool bInForceRescan = false, bool bInIgnoreBlackListScanFilters = false, TArray<FName>* FoundAssets = nullptr);
+
+	TArray<FString> PackageDirs;
+	TArray<FString> LocalDirs;
+	TArray<FString> PackageFiles;
+	TArray<FString> LocalFiles;
+	TArray<FString> LocalPaths;
+	FEventContext& EventContext;
+	TArray<FName>* OutFoundAssets = nullptr;
+	int32 NumFoundAssets = 0;
+	bool bForceRescan = false;
+	bool bIgnoreBlackListScanFilters = false;
+	bool bIdle = false;
+};
+
+}
+
+namespace Utils
+{
+
+/** Different modes for RunAssetThroughFilter and related filter functions */
+enum class EFilterMode : uint8
+{
+	/** Include things that pass the filter; include everything if the filter is empty */
+	Inclusive,
+	/** Exclude things that pass the filter; exclude nothing if the filter is empty */
+	Exclusive,
+};
+
+/** Report whether the external-code-created filter is in a state that is safe to use in filtering */
+bool IsFilterValid(const FARCompiledFilter& Filter);
+/** Report whether the given AssetData passes/fails the given filter */
+bool RunAssetThroughFilter(const FAssetData& AssetData, const FARCompiledFilter& Filter, const EFilterMode FilterMode);
+/** Helper for RunAssetThroughFilter and RunAssetsThroughFilter that skips validity check */
+bool RunAssetThroughFilter_Unchecked(const FAssetData& AssetData, const FARCompiledFilter& Filter, const bool bPassFilterValue);
+/**
+ * Given an array of asset data, trim the items that fail the filter based on the inclusion/exclusion mode used.
+ *  - In inclusive mode it will remove all assets that fail the filter, and in exclusive mode it will remove all assets that pass the filter.
+ *  - If the filter is empty, then the array will be untouched.
+ */
+void RunAssetsThroughFilter(TArray<FAssetData>& AssetDataList, const FARCompiledFilter& Filter, const EFilterMode FilterMode);
+
+/** This will always read the ini, public version may return cache */
+void InitializeSerializationOptionsFromIni(FAssetRegistrySerializationOptions& Options, const FString& PlatformIniName);
+
+/** Load via -Serialize the Out state from file at the given path */
+bool LoadAssetRegistry(const TCHAR* Path, const FAssetRegistryLoadOptions& Options, FAssetRegistryState& Out);
+
+/** Gets the current availability of an asset, primarily for streaming install purposes. */
+EAssetAvailability::Type GetAssetAvailability(const FAssetData& AssetData);
+/** Gets an ETA or percentage complete for an asset that is still in the process of being installed. */
+float GetAssetAvailabilityProgress(const FAssetData& AssetData, EAssetAvailabilityProgressReportingType::Type ReportType);
+/** Returns if a given report type is supported on the current platform */
+bool GetAssetAvailabilityProgressTypeSupported(EAssetAvailabilityProgressReportingType::Type ReportType);
+/**  Hint the streaming installers to prioritize a specific asset for install. */
+void PrioritizeAssetInstall(const FAssetData& AssetData);
+
+/**
+ * Reads Asset file from a previously initialized package reader
+ *
+ * @param PackageReader	Previously initialized package reader that maps to the data to be read
+ * @param AssetDataList	List of the read asset data
+ */
+bool ReadAssetFile(FPackageReader& PackageReader, IAssetRegistry::FLoadPackageRegistryData& InOutData);
+
+/**
+ * Returns true if path belongs to one of the mount points provided
+ *
+ * @param	Path				Path to check if mounted, example "/MyPlugin/SomeAsset"
+ * @param	MountPointsNoTrailingSlashes		Mount points without the trailing slash. Example: "/MyPlugin"
+ * @param	StringBuffer		String buffer to avoid re-allocation performance hit when searching TSet
+ */
+bool IsPathMounted(const FString& Path, const TSet<FString>& MountPointsNoTrailingSlashes, FString& StringBuffer);
+
+/** Helper function to return the name of an object, given the object's export text path */
+FString ExportTextPathToObjectName(const FString& InExportTextPath);
+
+/**
+ * Iterate all UObjects where this->IsAsset and !ShouldSkipAsset, create an FAssetData for them and call the callback.
+ * Stores packagenames of all assets found in OutPackageNamesWithAssets for later filtering of on-disk assets.
+ * Stops iteration and sets bOutStopIteration=true iff the callback returns false.
+ */
+void EnumerateAllMemoryAssets(TSet<FName>& OutPackageNamesWithAssets, bool& bOutStopIteration, TFunctionRef<bool(FAssetData&&)> Callback);
+/**
+ * Helper for EnumerateMemoryAssets* functions. Calls the given callback on all memory assets that pass the portion of the
+ * given filter that does not require tags, and passes in an FAssetData constructed from the UObject but missing tags.
+ * Fills in OutPackageNamesWithAssets with names of all packages tested.
+ */
+void EnumerateMemoryAssetsHelper(const FARCompiledFilter& InFilter, TSet<FName>& OutPackageNamesWithAssets,
+	bool& bOutStopIteration, TFunctionRef<bool(const UObject* Object, FAssetData&& PartialAssetData)> Callback);
+/**
+ * Call the given callback on all UObjects in memory that pass the given filter.
+ * Fills in OutPackageNamesWithAssets with names of all packages tested.
+*/
+void EnumerateMemoryAssets(const FARCompiledFilter& InFilter, TSet<FName>& OutPackageNamesWithAssets,
+	bool& bOutStopIteration, TFunctionRef<bool(FAssetData&&)> Callback);
+
+}
+
+}
+}
+
diff --git a/Engine/Source/Runtime/AssetRegistry/Private/AssetRegistryState.cpp b/Engine/Source/Runtime/AssetRegistry/Private/AssetRegistryState.cpp
--- a/Engine/Source/Runtime/AssetRegistry/Private/AssetRegistryState.cpp
+++ b/Engine/Source/Runtime/AssetRegistry/Private/AssetRegistryState.cpp
@@ -1,6 +1,7 @@
 // Copyright Epic Games, Inc. All Rights Reserved.
 
 #include "AssetRegistry/AssetRegistryState.h"
+
 #include "Misc/CommandLine.h"
 #include "Misc/FileHelper.h"
 #include "Misc/Paths.h"
@@ -11,6 +12,7 @@
 #include "UObject/UObjectIterator.h"
 #include "UObject/MetaData.h"
 #include "AssetRegistryArchive.h"
+#include "AssetRegistryImpl.h"
 #include "AssetRegistryPrivate.h"
 #include "AssetRegistry/ARFilter.h"
 #include "DependsNode.h"
@@ -1898,10 +1900,18 @@ bool FAssetRegistryState::RemovePackageData(FName PackageName)
 }
 
 bool FAssetRegistryState::IsFilterValid(const FARCompiledFilter& Filter)
+{
+	return UE::AssetRegistry::Utils::IsFilterValid(Filter);
+}
+
+namespace UE::AssetRegistry::Utils
+{
+
+bool IsFilterValid(const FARCompiledFilter& Filter)
 {
 	if (Filter.PackageNames.Contains(NAME_None) ||
 		Filter.PackagePaths.Contains(NAME_None) ||
-		Filter.ObjectPaths.Contains(NAME_None) || 
+		Filter.ObjectPaths.Contains(NAME_None) ||
 		Filter.ClassNames.Contains(NAME_None) ||
 		Filter.TagsAndValues.Contains(NAME_None)
 		)
@@ -1912,6 +1922,8 @@ bool FAssetRegistryState::IsFilterValid(const FARCompiledFilter& Filter)
 	return true;
 }
 
+}
+
 #if ASSET_REGISTRY_STATE_DUMPING_ENABLED
 namespace UE
 {
diff --git a/Engine/Source/Runtime/AssetRegistry/Private/PathTree.h b/Engine/Source/Runtime/AssetRegistry/Private/PathTree.h
--- a/Engine/Source/Runtime/AssetRegistry/Private/PathTree.h
+++ b/Engine/Source/Runtime/AssetRegistry/Private/PathTree.h
@@ -28,6 +28,11 @@ public:
 	/** Recursively enumerates all child paths from the specified base path relative to this node */
 	bool EnumerateSubPaths(FName BasePath, TFunctionRef<bool(FName)> Callback, bool bRecurse = true) const;
 
+	int32 NumPaths() const
+	{
+		return ParentPathToChildPaths.Num();
+	}
+
 	uint32 GetAllocatedSize(void) const
 	{
 		uint32 AllocatedSize = ParentPathToChildPaths.GetAllocatedSize() + ChildPathToParentPath.GetAllocatedSize();
diff --git a/Engine/Source/Runtime/AssetRegistry/Public/AssetRegistry/AssetRegistryState.h b/Engine/Source/Runtime/AssetRegistry/Public/AssetRegistry/AssetRegistryState.h
--- a/Engine/Source/Runtime/AssetRegistry/Public/AssetRegistry/AssetRegistryState.h
+++ b/Engine/Source/Runtime/AssetRegistry/Public/AssetRegistry/AssetRegistryState.h
@@ -10,9 +10,10 @@
 class FDependsNode;
 struct FARCompiledFilter;
 
-#ifndef ASSET_REGISTRY_STATE_DUMPING_ENABLED
-	#define ASSET_REGISTRY_STATE_DUMPING_ENABLED !(UE_BUILD_SHIPPING || UE_BUILD_TEST)
-#endif
+namespace UE::AssetRegistry
+{
+	class FAssetRegistryImpl;
+}
 
 /** Load/Save options used to modify how the cache is serialized. These are read out of the AssetRegistry section of Engine.ini and can be changed per platform. */
 struct FAssetRegistrySerializationOptions
@@ -435,4 +436,5 @@ private:
 	int32 NumPackageData = 0;
 
 	friend class UAssetRegistryImpl;
+	friend class UE::AssetRegistry::FAssetRegistryImpl;
 };
diff --git a/Engine/Source/Runtime/AssetRegistry/Public/AssetRegistry/IAssetRegistry.h b/Engine/Source/Runtime/AssetRegistry/Public/AssetRegistry/IAssetRegistry.h
--- a/Engine/Source/Runtime/AssetRegistry/Public/AssetRegistry/IAssetRegistry.h
+++ b/Engine/Source/Runtime/AssetRegistry/Public/AssetRegistry/IAssetRegistry.h
@@ -3,14 +3,21 @@
 #pragma once
 
 #include "CoreMinimal.h"
+
+#include "AssetRegistry/ARFilter.h"
+#include "AssetRegistry/AssetData.h"
 #include "Containers/BitArray.h"
 #include "Containers/StringFwd.h"
-#include "AssetRegistry/AssetData.h"
+#include "Misc/Optional.h"
 #include "Misc/AssetRegistryInterface.h"
 #include "UObject/Interface.h"
-#include "AssetRegistry/ARFilter.h"
+
 #include "IAssetRegistry.generated.h"
 
+#ifndef ASSET_REGISTRY_STATE_DUMPING_ENABLED
+#define ASSET_REGISTRY_STATE_DUMPING_ENABLED !(UE_BUILD_SHIPPING || UE_BUILD_TEST)
+#endif
+
 class FArchive;
 struct FARFilter;
 struct FARCompiledFilter;
@@ -299,6 +306,8 @@ public:
 	virtual bool K2_GetReferencers(FName PackageName, const FAssetRegistryDependencyOptions& ReferenceOptions, TArray<FName>& OutReferencers) const;
 
 	/** Finds Package data for a package name. This data is only updated on save and can only be accessed for valid packages */
+	virtual TOptional<FAssetPackageData> GetAssetPackageDataCopy(FName PackageName) const = 0;
+	UE_DEPRECATED(5.0, "Receiving a pointer is not threadsafe. Use GetAssetPackageDataCopy instead.")
 	virtual const FAssetPackageData* GetAssetPackageData(FName PackageName) const = 0;
 
 	/** Uses the asset registry to look for ObjectRedirectors. This will follow the chain of redirectors. It will return the original path if no redirectors are found */
@@ -528,12 +537,28 @@ public:
 	 */
 	virtual void InitializeTemporaryAssetRegistryState(FAssetRegistryState& OutState, const FAssetRegistrySerializationOptions& Options, bool bRefreshExisting = false, const TMap<FName, FAssetData*>& OverrideData = TMap<FName, FAssetData*>()) const = 0;
 
-	/** Returns read only reference to the current asset registry state. The contents of this may change at any time so do not save internal pointers */
+	UE_DEPRECATED(5.0, "Receiving a pointer is not threadsafe. Use other functions on IAssetRegistry to access the same data, or contact Epic Core team to add the threadsafe functions you require..")
 	virtual const FAssetRegistryState* GetAssetRegistryState() const = 0;
 
+#if ASSET_REGISTRY_STATE_DUMPING_ENABLED
+	/**
+	 * Writes out the state in textual form. Use arguments to control which segments to emit.
+	 * @param Arguments List of segments to emit. Possible values: 'ObjectPath', 'PackageName', 'Path', 'Class', 'Tag', 'Dependencies' and 'PackageData'
+	 * @param OutPages Textual representation will be written to this array; each entry will have LinesPerPage lines of the full dump.
+	 * @param LinesPerPage - how many lines should be combined into each string element of OutPages, for e.g. breaking up the dump into separate files.
+	 *        To facilitate diffing between similar-but-different registries, the actual number of lines per page will be slightly less than LinesPerPage; we introduce partially deterministic pagebreaks near the end of each page.
+	 */
+	virtual void DumpState(const TArray<FString>& Arguments, TArray<FString>& OutPages, int32 LinesPerPage = 1) const = 0;
+#endif
+
 	/** Returns the set of empty package names fast iteration */
+	virtual TSet<FName> GetCachedEmptyPackagesCopy() const = 0;
+	UE_DEPRECATED(5.0, "Receiving a reference is not threadsafe. Use GetCachedEmptyPackagesCopy instead.")
 	virtual const TSet<FName>& GetCachedEmptyPackages() const = 0;
 
+	/** Return whether the given TagName occurs in the tags of any asset in the AssetRegistry */
+	virtual bool ContainsTag(FName TagName) const = 0;
+
 	/** Fills in FAssetRegistrySerializationOptions from ini, optionally using a target platform ini name */
 	virtual void InitializeSerializationOptions(FAssetRegistrySerializationOptions& Options, const FString& PlatformIniName = FString()) const = 0;
 
@@ -555,10 +580,13 @@ public:
 	/** Load FAssetData from the specified package filename */
 	virtual void LoadPackageRegistryData(const FString& PackageFilename, FLoadPackageRegistryData& InOutData) const = 0;
 
-protected:
-	// Functions specifically for calling from the asset manager
-	friend class UAssetManager;
-	
+	/**
+	 * Enumerate all pairs in State->TagToAssetDataMapAssetRegistry and call a callback on each pair.
+	 * To avoid copies, the callback is called from within the ReadLock.
+	 * DO NOT CALL AssetRegistry functions from the callback; doing so will create a deadlock.
+	 */
+	virtual void ReadLockEnumerateTagToAssetDatas(TFunctionRef<void(FName TagName, const TArray<const FAssetData*>& Assets)> Callback) const = 0;
+
 	/**
 	 * Predicate called to decide whether to recurse into a reference when setting manager references
 	 *
@@ -571,6 +599,10 @@ protected:
 	typedef TFunction<EAssetSetManagerResult::Type(const FAssetIdentifier& Manager, const FAssetIdentifier& Source, const FAssetIdentifier& Target,
 		UE::AssetRegistry::EDependencyCategory Category, UE::AssetRegistry::EDependencyProperty Properties, EAssetSetManagerFlags::Type Flags)> ShouldSetManagerPredicate;
 
+protected:
+	// Functions specifically for calling from the asset manager
+	friend class UAssetManager;
+
 	/**
 	 * Specifies a list of manager mappings, optionally recursing to dependencies. These mappings can then be queried later to see which assets "manage" other assets
 	 * This function is only meant to be called by the AssetManager, calls from anywhere else will conflict and lose data
@@ -584,9 +616,6 @@ protected:
 
 	/** Sets the PrimaryAssetId for a specific asset. This should only be called by the AssetManager, and is needed when the AssetManager is more up to date than the on disk Registry */
 	virtual bool SetPrimaryAssetIdForObjectPath(const FName ObjectPath, FPrimaryAssetId PrimaryAssetId) = 0;
-
-	/** Returns pointer to cached AssetData for an object path. This is always the on disk version. This will return null if not found, and is exposed for  */
-	virtual const FAssetData* GetCachedAssetDataForObjectPath(const FName ObjectPath) const = 0;
 };
 
 namespace UE
diff --git a/Engine/Source/Runtime/Engine/Private/AssetManager.cpp b/Engine/Source/Runtime/Engine/Private/AssetManager.cpp
--- a/Engine/Source/Runtime/Engine/Private/AssetManager.cpp
+++ b/Engine/Source/Runtime/Engine/Private/AssetManager.cpp
@@ -1213,11 +1213,11 @@ bool UAssetManager::GetPrimaryAssetData(const FPrimaryAssetId& PrimaryAssetId, F
 
 	if (NameData)
 	{
-		const FAssetData* CachedAssetData = GetAssetRegistry().GetCachedAssetDataForObjectPath(NameData->AssetDataPath);
+		FAssetData CachedAssetData = GetAssetRegistry().GetAssetByObjectPath(NameData->AssetDataPath, true /* bIncludeOnlyOnDiskAssets */);
 
-		if (CachedAssetData && CachedAssetData->IsValid())
+		if (CachedAssetData.IsValid())
 		{
-			AssetData = *CachedAssetData;
+			AssetData = MoveTemp(CachedAssetData);
 			return true;
 		}
 	}
@@ -1236,12 +1236,12 @@ bool UAssetManager::GetPrimaryAssetDataList(FPrimaryAssetType PrimaryAssetType,
 
 		for (const TPair<FName, FPrimaryAssetData>& Pair : TypeData.AssetMap)
 		{
-			const FAssetData* CachedAssetData = Registry.GetCachedAssetDataForObjectPath(Pair.Value.AssetDataPath);
+			FAssetData CachedAssetData = Registry.GetAssetByObjectPath(Pair.Value.AssetDataPath, true /* bIncludeOnlyOnDiskAssets */);
 
-			if (CachedAssetData && CachedAssetData->IsValid())
+			if (CachedAssetData.IsValid())
 			{
 				bAdded = true;
-				AssetDataList.Add(*CachedAssetData);
+				AssetDataList.Add(MoveTemp(CachedAssetData));
 			}
 		}
 	}
