AssetRegistry: MakeDirectory was missing param Tree=true and DirectoryWatcher was failing when packaging WorldPartition levels

diff --git a/Engine/Source/Runtime/AssetRegistry/Private/AssetRegistry.cpp b/Engine/Source/Runtime/AssetRegistry/Private/AssetRegistry.cpp
--- a/Engine/Source/Runtime/AssetRegistry/Private/AssetRegistry.cpp
+++ b/Engine/Source/Runtime/AssetRegistry/Private/AssetRegistry.cpp
@@ -4274,7 +4274,7 @@ void UAssetRegistryImpl::OnContentPathMounted(const FString& InAssetPath, const
 		if (DirectoryWatcher)
 		{
 			// If the path doesn't exist on disk, make it so the watcher will work.
-			IFileManager::Get().MakeDirectory(*FileSystemPath);
+			IFileManager::Get().MakeDirectory(*FileSystemPath, /*Tree=*/true);
 		}
 	}
 		
