Attempt to silence static analysis warning:
V561: It's probably better to assign value to <VariableName> variable than to declare it anew.
#rnx

diff --git a/Engine/Source/Runtime/AssetRegistry/Private/AssetRegistryState.cpp b/Engine/Source/Runtime/AssetRegistry/Private/AssetRegistryState.cpp
--- a/Engine/Source/Runtime/AssetRegistry/Private/AssetRegistryState.cpp
+++ b/Engine/Source/Runtime/AssetRegistry/Private/AssetRegistryState.cpp
@@ -622,8 +622,8 @@ bool FAssetRegistryState::EnumerateAssets(const FARCompiledFilter& Filter, const
 		return false;
 	}
 
-	const uint32 FilterWithoutPackageFlags = Filter.WithoutPackageFlags;
-	const uint32 FilterWithPackageFlags = Filter.WithPackageFlags;
+	const uint32 FilterWithoutPackageFlags = Filter.WithoutPackageFlags; // -V561
+	const uint32 FilterWithPackageFlags = Filter.WithPackageFlags; // -V561
 
 	// The assets that match each filter
 	TArray<TArray<FAssetData*>, TInlineAllocator<5>> FilterResults;
