log a warning when scanning the disk with asset registry in -game to resolve partial package name

diff --git a/Engine/Source/Runtime/AssetRegistry/Private/AssetRegistry.cpp b/Engine/Source/Runtime/AssetRegistry/Private/AssetRegistry.cpp
--- a/Engine/Source/Runtime/AssetRegistry/Private/AssetRegistry.cpp
+++ b/Engine/Source/Runtime/AssetRegistry/Private/AssetRegistry.cpp
@@ -1657,6 +1657,10 @@ FName UAssetRegistryImpl::GetFirstPackageByName(FStringView PackageName) const
 		// In the future we plan to avoid this situation by having -game run the scan as well
 		if (LongPackageName.IsNone())
 		{
+			UE_LOG(LogAssetRegistry, Warning,
+				TEXT("GetFirstPackageByName is being called in `-game` to resolve partial package name. ")
+				TEXT("This may cause a slow scan on disk. ")
+				TEXT("Consider using the fully qualified package name for better performance. "));
 			FString LongPackageNameString;
 			if (FPackageName::SearchForPackageOnDisk(FString(PackageName), &LongPackageNameString))
 			{
