EditorDomain: Add list of ImportedClasses for each package to the AssetRegistry's FAssetPackageData.
#rnx

diff --git a/Engine/Source/Runtime/AssetRegistry/Private/AssetRegistryState.cpp b/Engine/Source/Runtime/AssetRegistry/Private/AssetRegistryState.cpp
--- a/Engine/Source/Runtime/AssetRegistry/Private/AssetRegistryState.cpp
+++ b/Engine/Source/Runtime/AssetRegistry/Private/AssetRegistryState.cpp
@@ -1212,20 +1212,13 @@ void FAssetRegistryState::Load(Archive&& Ar, FAssetRegistryVersion::Type Version
 			{
 				NewPackageData = &SerializedElement;
 			}
-			if (Version >= FAssetRegistryVersion::WorkspaceDomain)
+			if (Version >= FAssetRegistryVersion::LatestVersion)
 			{
 				NewPackageData->SerializeForCache(Ar);
 			}
 			else
 			{
-				Ar << NewPackageData->DiskSize;
-				PRAGMA_DISABLE_DEPRECATION_WARNINGS
-				Ar << NewPackageData->PackageGuid;
-				PRAGMA_ENABLE_DEPRECATION_WARNINGS
-				if (Version >= FAssetRegistryVersion::AddedCookedMD5Hash)
-				{
-					Ar << NewPackageData->CookedHash;
-				}
+				NewPackageData->SerializeForCacheOldVersion(Ar, Version);
 			}
 		}
 	}
diff --git a/Engine/Source/Runtime/AssetRegistry/Private/PackageReader.cpp b/Engine/Source/Runtime/AssetRegistry/Private/PackageReader.cpp
--- a/Engine/Source/Runtime/AssetRegistry/Private/PackageReader.cpp
+++ b/Engine/Source/Runtime/AssetRegistry/Private/PackageReader.cpp
@@ -368,6 +368,10 @@ bool FPackageReader::ReadDependencyData(FPackageDependencyData& OutDependencyDat
 	{
 		return false;
 	}
+	if (!SerializeImportedClasses(OutDependencyData.ImportMap, PackageData.ImportedClasses))
+	{
+		return false;
+	}
 	if (!SerializeSoftPackageReferenceList(OutDependencyData.SoftPackageReferenceList))
 	{
 		return false;
@@ -437,8 +441,8 @@ bool FPackageReader::SerializeImportMap(TArray<FObjectImport>& OutImportMap)
 		}
 		for ( int32 ImportMapIdx = 0; ImportMapIdx < PackageFileSummary.ImportCount; ++ImportMapIdx )
 		{
-			FObjectImport* Import = new(OutImportMap)FObjectImport;
-			*this << *Import;
+			FObjectImport& Import = OutImportMap.Emplace_GetRef();
+			*this << Import;
 			if (IsError())
 			{
 				UE_PACKAGEREADER_CORRUPTPACKAGE_WARNING("SerializeImportMapInvalidImport", PackageFilename);
@@ -450,6 +454,106 @@ bool FPackageReader::SerializeImportMap(TArray<FObjectImport>& OutImportMap)
 	return true;
 }
 
+bool FPackageReader::SerializeImportedClasses(const TArray<FObjectImport>& ImportMap, TArray<FName>& OutClassNames)
+{
+	OutClassNames.Reset();
+
+	TSet<int32> ClassImportIndices;
+	if (PackageFileSummary.ExportCount > 0)
+	{
+		if (!StartSerializeSection(PackageFileSummary.ExportOffset))
+		{
+			UE_PACKAGEREADER_CORRUPTPACKAGE_WARNING("SerializeExportMapInvalidExportOffset", PackageFilename);
+			return false;
+		}
+
+		const int MinSizePerExport = 1;
+		if (PackageFileSize < Tell() + PackageFileSummary.ExportCount * MinSizePerExport)
+		{
+			UE_PACKAGEREADER_CORRUPTPACKAGE_WARNING("SerializeExportMapInvalidExportCount", PackageFilename);
+			return false;
+		}
+		FObjectExport ExportBuffer;
+		for (int32 ExportMapIdx = 0; ExportMapIdx < PackageFileSummary.ExportCount; ++ExportMapIdx)
+		{
+			*this << ExportBuffer;
+			if (IsError())
+			{
+				UE_PACKAGEREADER_CORRUPTPACKAGE_WARNING("SerializeExportMapInvalidExport", PackageFilename);
+				return false;
+			}
+			if (ExportBuffer.ClassIndex.IsImport())
+			{
+				ClassImportIndices.Add(ExportBuffer.ClassIndex.ToImport());
+			}
+		}
+	}
+
+	TArray<FName, TInlineAllocator<5>>  ParentChain;
+	FNameBuilder ClassObjectPath;
+	for (int32 ClassImportIndex : ClassImportIndices)
+	{
+		ParentChain.Reset();
+		ClassObjectPath.Reset();
+		if (!ImportMap.IsValidIndex(ClassImportIndex))
+		{
+			UE_PACKAGEREADER_CORRUPTPACKAGE_WARNING("SerializeImportedClassesInvalidClassIndex", PackageFilename);
+			return false;
+		}
+		bool bParentChainComplete = false;
+		int32 CurrentParentIndex = ClassImportIndex;
+		for (;;)
+		{
+			const FObjectImport& ObjectImport = ImportMap[CurrentParentIndex];
+			ParentChain.Add(ObjectImport.ObjectName);
+			if (ObjectImport.OuterIndex.IsImport())
+			{
+				CurrentParentIndex = ObjectImport.OuterIndex.ToImport();
+				if (!ImportMap.IsValidIndex(CurrentParentIndex))
+				{
+					UE_PACKAGEREADER_CORRUPTPACKAGE_WARNING("SerializeImportedClassesInvalidImportInParentChain",
+						PackageFilename);
+					return false;
+				}
+			}
+			else if (ObjectImport.OuterIndex.IsNull())
+			{
+				bParentChainComplete = true;
+				break;
+			}
+			else
+			{
+				check(ObjectImport.OuterIndex.IsExport());
+				// Ignore classes in an external package but with an object in this package as one of their outers;
+				// We do not need to handle that case yet for Import Classes, and we would have to make this
+				// loop more complex (searching in both ExportMap and ImportMap) to do so
+				UE_PACKAGEREADER_CORRUPTPACKAGE_WARNING("SerializeImportedClassesInvalidExportInParentChain",
+					PackageFilename);
+				break;
+			}
+		} while (!bParentChainComplete);
+
+		if (bParentChainComplete)
+		{
+			int32 NumTokens = ParentChain.Num();
+			check(NumTokens >= 1);
+			const TCHAR Delimiters[] = {'.', SUBOBJECT_DELIMITER_CHAR, '.' };
+			int32 DelimiterIndex = 0;
+			ParentChain[NumTokens - 1].AppendString(ClassObjectPath);
+			for (int32 TokenIndex = NumTokens - 2; TokenIndex >= 0; --TokenIndex)
+			{
+				ClassObjectPath << Delimiters[DelimiterIndex];
+				DelimiterIndex = FMath::Min(DelimiterIndex+1, static_cast<int32>(UE_ARRAY_COUNT(Delimiters))-1);
+				ParentChain[TokenIndex].AppendString(ClassObjectPath);
+			}
+			OutClassNames.Emplace(ClassObjectPath);
+		}
+	}
+
+	OutClassNames.Sort(FNameLexicalLess());
+	return true;
+}
+
 bool FPackageReader::SerializeExportMap(TArray<FObjectExport>& OutExportMap)
 {
 	if (PackageFileSummary.ExportCount > 0)
diff --git a/Engine/Source/Runtime/AssetRegistry/Private/PackageReader.h b/Engine/Source/Runtime/AssetRegistry/Private/PackageReader.h
--- a/Engine/Source/Runtime/AssetRegistry/Private/PackageReader.h
+++ b/Engine/Source/Runtime/AssetRegistry/Private/PackageReader.h
@@ -48,6 +48,7 @@ public:
 	bool SerializeNameMap();
 	bool SerializeImportMap(TArray<FObjectImport>& OutImportMap);
 	bool SerializeExportMap(TArray<FObjectExport>& OutExportMap);
+	bool SerializeImportedClasses(const TArray<FObjectImport>& ImportMap, TArray<FName>& OutClassNames);
 	bool SerializeSoftPackageReferenceList(TArray<FName>& OutSoftPackageReferenceList);
 	bool SerializeSearchableNamesMap(FPackageDependencyData& OutDependencyData);
 	bool SerializeAssetRegistryDependencyData(FPackageDependencyData& DependencyData);
diff --git a/Engine/Source/Runtime/CoreUObject/Public/AssetRegistry/AssetData.h b/Engine/Source/Runtime/CoreUObject/Public/AssetRegistry/AssetData.h
--- a/Engine/Source/Runtime/CoreUObject/Public/AssetRegistry/AssetData.h
+++ b/Engine/Source/Runtime/CoreUObject/Public/AssetRegistry/AssetData.h
@@ -50,6 +50,7 @@ struct COREUOBJECT_API FAssetRegistryVersion
 								// * Removed global tag storage, a tag map reference-counts one store per asset registry
 								// * All configs can mix fixed and loose tag maps 
 		WorkspaceDomain,		// Added Version information to AssetPackageData
+		PackageImportedClasses, // Added ImportedClasses to AssetPackageData
 
 		// -----<new versions can be added above this line>-------------------------------------------------
 		VersionPlusOne,
@@ -620,6 +621,7 @@ private:
 }
 
 /** A class to hold data about a package on disk, this data is updated on save/load and is not updated when an asset changes in memory */
+PRAGMA_DISABLE_DEPRECATION_WARNINGS // Silence deprecation warnings for deprecated PackageGuid member in implicit constructors
 class FAssetPackageData
 {
 public:
@@ -630,6 +632,9 @@ public:
 	/** MD5 of the cooked package on disk, for tracking nondeterministic changes */
 	FMD5Hash CookedHash;
 
+	/** List of classes used by exports in the package. Does not include classes in the same package. */
+	TArray<FName> ImportedClasses;
+
 	/** Total size of this asset on disk */
 	int64 DiskSize;
 
@@ -653,17 +658,6 @@ public:
 		, Flags(0)
 	{
 	}
-	// Silence clang deprecation warnings for deprecated PackageGuid member in implicit constructors
-	PRAGMA_DISABLE_DEPRECATION_WARNINGS
-	FAssetPackageData(FAssetPackageData&& Other) = default;
-	FAssetPackageData(const FAssetPackageData& Other)
-		: FAssetPackageData()
-	{
-		*this = Other;
-	}
-	FAssetPackageData& operator=(FAssetPackageData&& Other) = default;
-	FAssetPackageData& operator=(const FAssetPackageData& Other) = default;
-	PRAGMA_ENABLE_DEPRECATION_WARNINGS
 
 	/**
 	 * Custom versions used by the package, used to check whether we need to update the package for the current binary.
@@ -696,15 +690,12 @@ public:
 	 */
 	void SerializeForCache(FArchive& Ar)
 	{
-		Ar << DiskSize;
-		PRAGMA_DISABLE_DEPRECATION_WARNINGS
-		Ar << PackageGuid;
-		PRAGMA_ENABLE_DEPRECATION_WARNINGS
-		Ar << CookedHash;
-		Ar << FileVersionUE;
-		Ar << FileVersionLicenseeUE;
-		Ar << Flags;
-		Ar << CustomVersions;
+		// Calling with hard-coded version and using force-inline on SerializeForCacheInternal eliminates the cost of its if-statements
+		SerializeForCacheInternal(Ar, FAssetRegistryVersion::LatestVersion);
+	}
+	void SerializeForCacheOldVersion(FArchive& Ar, FAssetRegistryVersion::Type Version)
+	{
+		SerializeForCacheInternal(Ar, Version);
 	}
 
 private:
@@ -712,7 +703,33 @@ private:
 	{
 		FLAG_LICENSEE_VERSION = 0x1,
 	};
+
+	void SerializeForCacheInternal(FArchive& Ar, FAssetRegistryVersion::Type Version);
 };
+PRAGMA_ENABLE_DEPRECATION_WARNINGS
+FORCEINLINE void FAssetPackageData::SerializeForCacheInternal(FArchive& Ar, FAssetRegistryVersion::Type Version)
+{
+	Ar << DiskSize;
+	PRAGMA_DISABLE_DEPRECATION_WARNINGS
+	Ar << PackageGuid;
+	PRAGMA_ENABLE_DEPRECATION_WARNINGS
+	if (Version >= FAssetRegistryVersion::AddedCookedMD5Hash)
+	{
+		Ar << CookedHash;
+	}
+	if (Version >= FAssetRegistryVersion::WorkspaceDomain)
+	{
+		Ar << FileVersionUE;
+		Ar << FileVersionLicenseeUE;
+		Ar << Flags;
+		Ar << CustomVersions;
+	}
+	if (Version >= FAssetRegistryVersion::PackageImportedClasses)
+	{
+		Ar << ImportedClasses;
+	}
+}
+
 
 /**
  * Helper struct for FAssetIdentifier (e.g., for the FOnViewAssetIdentifiersInReferenceViewer delegate and Reference Viewer functions).
