PackageReader: Fix typo: accidental remnant of previous version of the for loop.
#rnx

diff --git a/Engine/Source/Runtime/AssetRegistry/Private/PackageReader.cpp b/Engine/Source/Runtime/AssetRegistry/Private/PackageReader.cpp
--- a/Engine/Source/Runtime/AssetRegistry/Private/PackageReader.cpp
+++ b/Engine/Source/Runtime/AssetRegistry/Private/PackageReader.cpp
@@ -527,11 +527,9 @@ bool FPackageReader::SerializeImportedClasses(const TArray<FObjectImport>& Impor
 				// Ignore classes in an external package but with an object in this package as one of their outers;
 				// We do not need to handle that case yet for Import Classes, and we would have to make this
 				// loop more complex (searching in both ExportMap and ImportMap) to do so
-				UE_PACKAGEREADER_CORRUPTPACKAGE_WARNING("SerializeImportedClassesInvalidExportInParentChain",
-					PackageFilename);
 				break;
 			}
-		} while (!bParentChainComplete);
+		}
 
 		if (bParentChainComplete)
 		{
