Fix use-after-free in UpdateInheritanceBuffer caused by use of dangling reference from TMap::FindOrAdd

#rnx

diff --git a/Engine/Source/Runtime/AssetRegistry/Private/AssetRegistry.cpp b/Engine/Source/Runtime/AssetRegistry/Private/AssetRegistry.cpp
--- a/Engine/Source/Runtime/AssetRegistry/Private/AssetRegistry.cpp
+++ b/Engine/Source/Runtime/AssetRegistry/Private/AssetRegistry.cpp
@@ -4523,8 +4523,6 @@ bool UAssetRegistryImpl::GetTemporaryCachingMode() const
 	return GuardedData.IsTempCachingEnabled();
 }
 
-COREUOBJECT_API TMap<UClass*, TSet<UClass*>> GetAllDerivedClasses();
-
 namespace UE::AssetRegistry
 {
 
@@ -4553,24 +4551,31 @@ void FAssetRegistryImpl::UpdateInheritanceBuffer(Impl::FClassInheritanceBuffer&
 	{
 		FName SuperclassName = Pair.Key->GetFName();
 
-		TArray<FName>& TempCachedSubclasses = OutBuffer.ReverseInheritanceMap.FindOrAdd(SuperclassName);
-		TempCachedSubclasses.Reserve(Pair.Value.Num());
+		TArray<FName>* OutputSubclasses = &OutBuffer.ReverseInheritanceMap.FindOrAdd(SuperclassName);
+		OutputSubclasses->Reserve(Pair.Value.Num());
 		for (UClass* Subclass : Pair.Value)
 		{
 			if (!Subclass->HasAnyClassFlags(CLASS_Deprecated | CLASS_NewerVersionExists))
 			{
 				FName SubclassName = Subclass->GetFName();
-				TempCachedSubclasses.Add(SubclassName);
+				OutputSubclasses->Add(SubclassName);
 				OutBuffer.InheritanceMap.Add(SubclassName, SuperclassName);
 
-				// Add any implemented interfaces to the reverse inheritance map, but not to the forward map
-				for (const FImplementedInterface& Interface : Subclass->Interfaces)
+				if (Subclass->Interfaces.Num())
 				{
-					if (UClass* InterfaceClass = Interface.Class) // could be nulled out by ForceDelete of a blueprint interface
+					// Add any implemented interfaces to the reverse inheritance map, but not to the forward map
+					for (const FImplementedInterface& Interface : Subclass->Interfaces)
 					{
-						TArray<FName>& Implementations = OutBuffer.ReverseInheritanceMap.FindOrAdd(InterfaceClass->GetFName());
-						Implementations.Add(SubclassName);
+						if (UClass* InterfaceClass = Interface.Class) // could be nulled out by ForceDelete of a blueprint interface
+						{
+							TArray<FName>& Implementations = OutBuffer.ReverseInheritanceMap.FindOrAdd(InterfaceClass->GetFName());
+							Implementations.Add(SubclassName);
+						}
 					}
+
+					// Refetch OutputSubClasses from ReverseInheritanceMap because we just modified the ReverseInheritanceMap and may have resized
+					OutputSubclasses = OutBuffer.ReverseInheritanceMap.Find(SuperclassName);
+					check(OutputSubclasses); // It was added above
 				}
 			}
 		}
