AssetRegistry::WaitForPackage: Fix AssetDataGatherer refusing to wait because we do not have the full path with extension of the package name.
Repro case was on the first package requested when starting the editor: /Engine/EditorResources/S_Actor
#rnx

diff --git a/Engine/Source/Runtime/AssetRegistry/Private/AssetDataGatherer.cpp b/Engine/Source/Runtime/AssetRegistry/Private/AssetDataGatherer.cpp
--- a/Engine/Source/Runtime/AssetRegistry/Private/AssetDataGatherer.cpp
+++ b/Engine/Source/Runtime/AssetRegistry/Private/AssetDataGatherer.cpp
@@ -1642,26 +1642,45 @@ void FAssetDataDiscovery::WaitForIdle()
 
 void FAssetDataDiscovery::SetPropertiesAndWait(const FString& LocalAbsPath, bool bAddToWhitelist, bool bForceRescan, bool bIgnoreBlackListScanFilters)
 {
-	FFileStatData StatData = IFileManager::Get().GetStatData(*LocalAbsPath);
-	if (!StatData.bIsValid)
+	enum class EPathType
 	{
-		// SetPropertiesAndWait is called for every ScanPathsSynchronous, and this is the first spot that checks for existence. Some systems call ScanPathsSynchronous
-		// speculatively to scan whatever is present, so this is not a significant enough occurrence for a log.
-		UE_LOG(LogAssetRegistry, Verbose, TEXT("SetPropertiesAndWait called on non-existent path %.*s. Call will be ignored."), LocalAbsPath.Len(), *LocalAbsPath);
-		return;
+		Directory,
+		File,
+		UnknownFile,
+	};
+	FFileStatData StatData = IFileManager::Get().GetStatData(*LocalAbsPath);
+	// We might have been asked to wait on a filename missing the extension, in which case GetStatData will return DNE
+	// We need to handle Directory, File, and missing (presumed a file of unknown name) in unique ways
+	EPathType PathType = !StatData.bIsValid ? EPathType::UnknownFile : (StatData.bIsDirectory ? EPathType::Directory : EPathType::File);
+	FString UnknownFileParentPath;
+	const FString* SearchPath = &LocalAbsPath;
+	if (PathType == EPathType::UnknownFile)
+	{
+		// GetControllingDir assumes the parent path of our filename exists, so if it doesn't we need to eject
+		UnknownFileParentPath = FPaths::GetPath(LocalAbsPath);
+		StatData = IFileManager::Get().GetStatData(*UnknownFileParentPath);
+		if (!StatData.bIsValid || !StatData.bIsDirectory)
+		{
+			// SetPropertiesAndWait is called for every ScanPathsSynchronous, and this is the first spot that checks for existence. Some systems call ScanPathsSynchronous
+			// speculatively to scan whatever is present, so this is not a significant enough occurrence for a log.
+			UE_LOG(LogAssetRegistry, Verbose, TEXT("SetPropertiesAndWait called on non-existent path %.*s. Call will be ignored."), LocalAbsPath.Len(), *LocalAbsPath);
+			return;
+		}
+		SearchPath = &UnknownFileParentPath;
 	}
 
+
 	{
 		CHECK_IS_NOT_LOCKED_CURRENT_THREAD(ResultsLock);
 		FGathererScopeLock TreeScopeLock(&TreeLock);
-		FMountDir* MountDir = FindContainingMountPoint(LocalAbsPath);
+		FMountDir* MountDir = FindContainingMountPoint(*SearchPath);
 		if (!MountDir)
 		{
 			UE_LOG(LogAssetRegistry, Log, TEXT("SetPropertiesAndWait called on %.*s which is not in a mounted directory. Call will be ignored."), LocalAbsPath.Len(), *LocalAbsPath);
 			return;
 		}
 
-		if (StatData.bIsDirectory)
+		if (PathType == EPathType::Directory)
 		{
 			FSetPathProperties Properties;
 			if (bAddToWhitelist)
@@ -1679,13 +1698,13 @@ void FAssetDataDiscovery::SetPropertiesAndWait(const FString& LocalAbsPath, bool
 			if (Properties.IsSet())
 			{
 				SetIsIdle(false);
-				MountDir->TrySetDirectoryProperties(LocalAbsPath, Properties, true /* bConfirmedExists */);
+				MountDir->TrySetDirectoryProperties(*SearchPath, Properties, true /* bConfirmedExists */);
 			}
 		}
 
 		FString RelPath;
 		FScanDir::FInherited MonitorData;
-		TRefCountPtr<FScanDir> ScanDir = MountDir->GetControllingDir(LocalAbsPath, StatData.bIsDirectory, MonitorData, RelPath);
+		TRefCountPtr<FScanDir> ScanDir = MountDir->GetControllingDir(*SearchPath, PathType == EPathType::Directory, MonitorData, RelPath);
 		bool bIsWhitelistedInThisCall = MonitorData.IsWhitelisted() || bAddToWhitelist;
 		bool bIsBlacklistedInThisCall = MonitorData.IsBlacklisted() && !bIgnoreBlackListScanFilters;
 		bool bIsMonitoredInThisCall = bIsWhitelistedInThisCall && !bIsBlacklistedInThisCall;
@@ -1695,7 +1714,7 @@ void FAssetDataDiscovery::SetPropertiesAndWait(const FString& LocalAbsPath, bool
 			return;
 		}
 
-		if (StatData.bIsDirectory)
+		if (PathType == EPathType::Directory || PathType == EPathType::UnknownFile)
 		{
 			// If Relpath from the controlling dir to the requested dir is not empty then we have found a parent directory rather than the requested directory.
 			// This can only occur for a monitored directory when the requested directory is already complete and we do not need to wait on it.
@@ -1708,6 +1727,7 @@ void FAssetDataDiscovery::SetPropertiesAndWait(const FString& LocalAbsPath, bool
 
 				FScopedPause ScopedPause(*this);
 				TreeScopeLock.Unlock(); // Entering the ticklock, as well as any long duration task such as a tick, has to be done outside of any locks
+				bool bScanEntireTree = PathType == EPathType::Directory;
 
 				CHECK_IS_NOT_LOCKED_CURRENT_THREAD(TreeLock);
 				CHECK_IS_NOT_LOCKED_CURRENT_THREAD(ResultsLock);
@@ -1716,7 +1736,11 @@ void FAssetDataDiscovery::SetPropertiesAndWait(const FString& LocalAbsPath, bool
 				{
 					TickInternal();
 					FGathererScopeLock LoopTreeScopeLock(&TreeLock);
-					if (!ScanDir->IsValid() || ScanDir->IsComplete())
+					if (!ScanDir->IsValid())
+					{
+						break;
+					}
+					if ((bScanEntireTree && ScanDir->IsComplete()) || (!bScanEntireTree && ScanDir->HasScanned()))
 					{
 						break;
 					}
@@ -1740,7 +1764,8 @@ void FAssetDataDiscovery::SetPropertiesAndWait(const FString& LocalAbsPath, bool
 					LongPackageName << MountDir->GetLongPackageName();
 					FPathViews::AppendPath(LongPackageName, ScanDir->GetMountRelPath());
 					FPathViews::AppendPath(LongPackageName, FileRelPathNoExt);
-					AddDiscovered(TConstArrayView<FDiscoveredPathData>(), TConstArrayView<FDiscoveredPathData> { FDiscoveredPathData(LocalAbsPath, LongPackageName, RelPathFromParentDir, StatData.ModificationTime) });
+					AddDiscovered(TConstArrayView<FDiscoveredPathData>(), TConstArrayView<FDiscoveredPathData>
+						{ FDiscoveredPathData(*SearchPath, LongPackageName, RelPathFromParentDir, StatData.ModificationTime) });
 					if (FPathViews::IsPathLeaf(RelPath) && !ScanDir->HasScanned())
 					{
 						SetIsIdle(false);
@@ -3348,7 +3373,9 @@ void FAssetDataGatherer::SortPathsByPriority(TConstArrayView<FString> LocalAbsPa
 	{
 		for (const FString& LocalAbsPathToPrioritize : LocalAbsPathsToPrioritize)
 		{
-			if (FPathViews::IsParentPathOf(LocalAbsPathToPrioritize, FilesToSearch[FilenameIdx].LocalAbsPath))
+			// Handle LocalAbsPathsToPrioritize that might be a directory, or a filename without an extension, or a full filename
+			if (FPathViews::IsParentPathOf(LocalAbsPathToPrioritize, FilesToSearch[FilenameIdx].LocalAbsPath) ||
+				FPathViews::Equals(LocalAbsPathToPrioritize, FPathViews::GetBaseFilenameWithPath(FilesToSearch[FilenameIdx].LocalAbsPath)))
 			{
 				Swap(FilesToSearch[FilenameIdx], FilesToSearch[LowestNonPriorityFileIdx++]);
 				break;
