IAssetRegistry->WaitOnPackage: improve performance when looking for package in AssetDataGatherer::FilesToSearch by changing from array to a DirectoryTree.
 247s -> 92s to start large project Editor (from process launch until first frame of the LevelEditor window) with a warm EditorDomain cache but an empty CachedAssetRegistry.bin.
#rnx

diff --git a/Engine/Source/Runtime/AssetRegistry/Private/AssetDataGatherer.cpp b/Engine/Source/Runtime/AssetRegistry/Private/AssetDataGatherer.cpp
--- a/Engine/Source/Runtime/AssetRegistry/Private/AssetDataGatherer.cpp
+++ b/Engine/Source/Runtime/AssetRegistry/Private/AssetDataGatherer.cpp
@@ -101,6 +101,13 @@ FGatheredPathData::FGatheredPathData(const FDiscoveredPathData& DiscoveredData)
 {
 }
 
+FGatheredPathData::FGatheredPathData(FDiscoveredPathData&& DiscoveredData)
+	: LocalAbsPath(MoveTemp(DiscoveredData.LocalAbsPath))
+	, LongPackageName(MoveTemp(DiscoveredData.LongPackageName))
+	, PackageTimestamp(MoveTemp(DiscoveredData.PackageTimestamp))
+{
+}
+
 void FGatheredPathData::Assign(FStringView InLocalAbsPath, FStringView InLongPackageName, const FDateTime& InPackageTimestamp)
 {
 	AssignStringWithoutShrinking(LocalAbsPath, InLocalAbsPath);
@@ -1525,7 +1532,7 @@ void FAssetDataDiscovery::TickInternal()
 
 	if (bValid && (!LocalSubDirs.IsEmpty() || !LocalDiscoveredFiles.IsEmpty()))
 	{
-		AddDiscovered(LocalSubDirs, LocalDiscoveredFiles);
+		AddDiscovered(DirLocalAbsPath, LocalSubDirs, LocalDiscoveredFiles);
 	}
 }
 
@@ -1607,18 +1614,25 @@ void FAssetDataDiscovery::SetIsIdle(bool bInIsIdle)
 }
 
 
-void FAssetDataDiscovery::GetAndTrimSearchResults(bool& bOutIsComplete, TArray<FString>& OutDiscoveredPaths, TRingBuffer<FGatheredPathData>& OutDiscoveredFiles, int32& OutNumPathsToSearch)
+void FAssetDataDiscovery::GetAndTrimSearchResults(bool& bOutIsComplete, TArray<FString>& OutDiscoveredPaths, FFilesToSearch& OutFilesToSearch, int32& OutNumPathsToSearch)
 {
 	FGathererScopeLock ResultsScopeLock(&ResultsLock);
 
 	OutDiscoveredPaths.Append(MoveTemp(DiscoveredDirectories));
 	DiscoveredDirectories.Reset();
-	OutDiscoveredFiles.Reserve(OutDiscoveredFiles.Num() + DiscoveredFiles.Num());
-	for (FGatheredPathData& DiscoveredFile : DiscoveredFiles)
+
+	for (FDirectoryResult& DirectoryResult : DiscoveredFiles)
 	{
-		OutDiscoveredFiles.Add(MoveTemp(DiscoveredFile));
+		OutFilesToSearch.AddDirectory(MoveTemp(DirectoryResult.DirAbsPath), MoveTemp(DirectoryResult.Files));
 	}
 	DiscoveredFiles.Reset();
+	for (FGatheredPathData& FileResult : DiscoveredSingleFiles)
+	{
+		// Single files are currently only added from the blocking function FAssetDataDiscovery::SetPropertiesAndWait,
+		// so we add them at blocking priority.
+		OutFilesToSearch.AddPriorityFile(MoveTemp(FileResult));
+	}
+	DiscoveredSingleFiles.Reset();
 
 	OutNumPathsToSearch = NumDirectoriesToScan.GetValue();
 	bOutIsComplete = bIsIdle;
@@ -1640,47 +1654,92 @@ void FAssetDataDiscovery::WaitForIdle()
 	}
 }
 
-void FAssetDataDiscovery::SetPropertiesAndWait(const FString& LocalAbsPath, bool bAddToWhitelist, bool bForceRescan, bool bIgnoreBlackListScanFilters)
+FPathExistence::FPathExistence(FStringView InLocalAbsPath)
+	:LocalAbsPath(InLocalAbsPath)
+{
+}
+
+const FString& FPathExistence::GetLocalAbsPath() const
 {
-	enum class EPathType
+	return LocalAbsPath;
+}
+
+FStringView FPathExistence::GetLowestExistingPath()
+{
+	LoadExistenceData();
+	switch (PathType)
 	{
-		Directory,
-		File,
-		UnknownFile,
-	};
+	case EType::MissingButDirExists:
+		return FPathViews::GetPath(LocalAbsPath);
+	case EType::MissingParentDir:
+		return FStringView();
+	default:
+		return LocalAbsPath;
+	}
+}
+
+FPathExistence::EType FPathExistence::GetType()
+{
+	LoadExistenceData();
+	return PathType;
+}
+
+FDateTime FPathExistence::GetModificationTime()
+{
+	LoadExistenceData();
+	return ModificationTime;
+}
+
+void FPathExistence::LoadExistenceData()
+{
+	if (bHasExistenceData)
+	{
+		return;
+	}
 	FFileStatData StatData = IFileManager::Get().GetStatData(*LocalAbsPath);
-	// We might have been asked to wait on a filename missing the extension, in which case GetStatData will return DNE
-	// We need to handle Directory, File, and missing (presumed a file of unknown name) in unique ways
-	EPathType PathType = !StatData.bIsValid ? EPathType::UnknownFile : (StatData.bIsDirectory ? EPathType::Directory : EPathType::File);
-	FString UnknownFileParentPath;
-	const FString* SearchPath = &LocalAbsPath;
-	if (PathType == EPathType::UnknownFile)
-	{
-		// GetControllingDir assumes the parent path of our filename exists, so if it doesn't we need to eject
-		UnknownFileParentPath = FPaths::GetPath(LocalAbsPath);
-		StatData = IFileManager::Get().GetStatData(*UnknownFileParentPath);
-		if (!StatData.bIsValid || !StatData.bIsDirectory)
-		{
-			// SetPropertiesAndWait is called for every ScanPathsSynchronous, and this is the first spot that checks for existence. Some systems call ScanPathsSynchronous
-			// speculatively to scan whatever is present, so this is not a significant enough occurrence for a log.
-			UE_LOG(LogAssetRegistry, Verbose, TEXT("SetPropertiesAndWait called on non-existent path %.*s. Call will be ignored."), LocalAbsPath.Len(), *LocalAbsPath);
-			return;
-		}
-		SearchPath = &UnknownFileParentPath;
+	if (StatData.bIsValid)
+	{
+		ModificationTime = StatData.ModificationTime;
+		PathType = StatData.bIsDirectory ? EType::Directory : EType::File;
 	}
+	else
+	{
+		FString ParentPath = FPaths::GetPath(LocalAbsPath);
+		StatData = IFileManager::Get().GetStatData(*ParentPath);
+		PathType = (StatData.bIsValid && StatData.bIsDirectory)
+			? EType::MissingButDirExists : EType::MissingParentDir;
+	}
+
+	bHasExistenceData = true;
+}
 
+void FAssetDataDiscovery::SetPropertiesAndWait(FPathExistence& QueryPath, bool bAddToWhitelist, bool bForceRescan, bool bIgnoreBlackListScanFilters)
+{
+	// We might have been asked to wait on a filename missing the extension, in which case QueryPath.GetType() == MissingButDirExists
+	// We need to handle Directory, File, and MissingButDirExists in unique ways
+	FPathExistence::EType PathType = QueryPath.GetType();
+	if (PathType == FPathExistence::EType::MissingParentDir)
+	{
+		// SetPropertiesAndWait is called for every ScanPathsSynchronous, and this is the first spot that checks for existence. 
+		// Some systems call ScanPathsSynchronous speculatively to scan whatever is present, so this log is verbose-only.
+		UE_LOG(LogAssetRegistry, Verbose, TEXT("SetPropertiesAndWait called on non-existent path %s. Call will be ignored."),
+			*QueryPath.GetLocalAbsPath());
+		return;
+	}
+	FStringView SearchPath = QueryPath.GetLowestExistingPath();
 
 	{
 		CHECK_IS_NOT_LOCKED_CURRENT_THREAD(ResultsLock);
 		FGathererScopeLock TreeScopeLock(&TreeLock);
-		FMountDir* MountDir = FindContainingMountPoint(*SearchPath);
+		FMountDir* MountDir = FindContainingMountPoint(SearchPath);
 		if (!MountDir)
 		{
-			UE_LOG(LogAssetRegistry, Log, TEXT("SetPropertiesAndWait called on %.*s which is not in a mounted directory. Call will be ignored."), LocalAbsPath.Len(), *LocalAbsPath);
+			UE_LOG(LogAssetRegistry, Log, TEXT("SetPropertiesAndWait called on %s which is not in a mounted directory. Call will be ignored."),
+				*QueryPath.GetLocalAbsPath());
 			return;
 		}
 
-		if (PathType == EPathType::Directory)
+		if (PathType == FPathExistence::EType::Directory)
 		{
 			FSetPathProperties Properties;
 			if (bAddToWhitelist)
@@ -1698,23 +1757,25 @@ void FAssetDataDiscovery::SetPropertiesAndWait(const FString& LocalAbsPath, bool
 			if (Properties.IsSet())
 			{
 				SetIsIdle(false);
-				MountDir->TrySetDirectoryProperties(*SearchPath, Properties, true /* bConfirmedExists */);
+				MountDir->TrySetDirectoryProperties(SearchPath, Properties, true /* bConfirmedExists */);
 			}
 		}
 
 		FString RelPath;
 		FScanDir::FInherited MonitorData;
-		TRefCountPtr<FScanDir> ScanDir = MountDir->GetControllingDir(*SearchPath, PathType == EPathType::Directory, MonitorData, RelPath);
+		bool bSearchPathIsDirectory = PathType == FPathExistence::EType::Directory || PathType == FPathExistence::EType::MissingButDirExists;
+		TRefCountPtr<FScanDir> ScanDir = MountDir->GetControllingDir(SearchPath, bSearchPathIsDirectory, MonitorData, RelPath);
 		bool bIsWhitelistedInThisCall = MonitorData.IsWhitelisted() || bAddToWhitelist;
 		bool bIsBlacklistedInThisCall = MonitorData.IsBlacklisted() && !bIgnoreBlackListScanFilters;
 		bool bIsMonitoredInThisCall = bIsWhitelistedInThisCall && !bIsBlacklistedInThisCall;
 		if (!ScanDir || !bIsMonitoredInThisCall)
 		{
-			UE_LOG(LogAssetRegistry, Log, TEXT("SetPropertiesAndWait called on %.*s which is not monitored. Call will be ignored."), LocalAbsPath.Len(), *LocalAbsPath);
+			UE_LOG(LogAssetRegistry, Log, TEXT("SetPropertiesAndWait called on %s which is not monitored. Call will be ignored."),
+				*QueryPath.GetLocalAbsPath());
 			return;
 		}
 
-		if (PathType == EPathType::Directory || PathType == EPathType::UnknownFile)
+		if (bSearchPathIsDirectory)
 		{
 			// If Relpath from the controlling dir to the requested dir is not empty then we have found a parent directory rather than the requested directory.
 			// This can only occur for a monitored directory when the requested directory is already complete and we do not need to wait on it.
@@ -1727,7 +1788,8 @@ void FAssetDataDiscovery::SetPropertiesAndWait(const FString& LocalAbsPath, bool
 
 				FScopedPause ScopedPause(*this);
 				TreeScopeLock.Unlock(); // Entering the ticklock, as well as any long duration task such as a tick, has to be done outside of any locks
-				bool bScanEntireTree = PathType == EPathType::Directory;
+				// If the query path is MissingButDirExists, we assume it is a file missing the extension, and scan (or confirm already scanned) its directory
+				bool bScanEntireTree = PathType == FPathExistence::EType::Directory;
 
 				CHECK_IS_NOT_LOCKED_CURRENT_THREAD(TreeLock);
 				CHECK_IS_NOT_LOCKED_CURRENT_THREAD(ResultsLock);
@@ -1753,6 +1815,7 @@ void FAssetDataDiscovery::SetPropertiesAndWait(const FString& LocalAbsPath, bool
 		}
 		else
 		{
+			check(PathType == FPathExistence::EType::File);
 			bool bAlreadyScanned = ScanDir->HasScanned() && MonitorData.IsMonitored();
 			if (!bAlreadyScanned || bForceRescan)
 			{
@@ -1764,8 +1827,7 @@ void FAssetDataDiscovery::SetPropertiesAndWait(const FString& LocalAbsPath, bool
 					LongPackageName << MountDir->GetLongPackageName();
 					FPathViews::AppendPath(LongPackageName, ScanDir->GetMountRelPath());
 					FPathViews::AppendPath(LongPackageName, FileRelPathNoExt);
-					AddDiscovered(TConstArrayView<FDiscoveredPathData>(), TConstArrayView<FDiscoveredPathData>
-						{ FDiscoveredPathData(*SearchPath, LongPackageName, RelPathFromParentDir, StatData.ModificationTime) });
+					AddDiscoveredFile(FDiscoveredPathData(SearchPath, LongPackageName, RelPathFromParentDir, QueryPath.GetModificationTime()));
 					if (FPathViews::IsPathLeaf(RelPath) && !ScanDir->HasScanned())
 					{
 						SetIsIdle(false);
@@ -1874,6 +1936,7 @@ uint32 FAssetDataDiscovery::GetAllocatedSize() const
 
 	Result += GetArrayRecursiveAllocatedSize(DiscoveredDirectories);
 	Result += GetArrayRecursiveAllocatedSize(DiscoveredFiles);
+	Result += GetArrayRecursiveAllocatedSize(DiscoveredSingleFiles);
 
 	Result += MountDirs.GetAllocatedSize();
 	for (const TUniquePtr<FMountDir>& Value : MountDirs)
@@ -1886,6 +1949,11 @@ uint32 FAssetDataDiscovery::GetAllocatedSize() const
 	return Result;
 }
 
+uint32 FAssetDataDiscovery::FDirectoryResult::GetAllocatedSize() const
+{
+	return DirAbsPath.GetAllocatedSize() + Files.GetAllocatedSize();
+}
+
 void FAssetDataDiscovery::Shrink()
 {
 	CHECK_IS_LOCKED_CURRENT_THREAD(TickLock);
@@ -1894,6 +1962,7 @@ void FAssetDataDiscovery::Shrink()
 	DirLongPackageNamesToNotReport.Shrink();
 	DiscoveredDirectories.Shrink();
 	DiscoveredFiles.Shrink();
+	DiscoveredSingleFiles.Shrink();
 	MountDirs.Shrink();
 	for (TUniquePtr<FMountDir>& MountDir : MountDirs)
 	{
@@ -2040,7 +2109,7 @@ void FAssetDataDiscovery::OnDirectoryCreated(FStringView LocalAbsPath)
 	// Any files and paths under it will be added by their own event from the directory watcher, so a scan is unnecessary.
 	// The directory may also be scanned in the future because a parent directory is still yet pending to scan,
 	// we do not try to prevent that wasteful rescan because this is a rare event and it does not cause a behavior problem
-	AddDiscovered(TConstArrayView<FDiscoveredPathData>(&DirData, 1), TConstArrayView<FDiscoveredPathData>());
+	AddDiscovered(DirData.LocalAbsPath, TConstArrayView<FDiscoveredPathData>(&DirData, 1), TConstArrayView<FDiscoveredPathData>());
 	SetIsIdle(false);
 }
 
@@ -2089,7 +2158,7 @@ void FAssetDataDiscovery::OnFileCreated(const FString& LocalAbsPath)
 		LongPackageName << MountDir->GetLongPackageName();
 		FPathViews::AppendPath(LongPackageName, ScanDir->GetMountRelPath());
 		FPathViews::AppendPath(LongPackageName, FileRelPathNoExt);
-		AddDiscovered(TConstArrayView<FDiscoveredPathData>(), TConstArrayView<FDiscoveredPathData> { FDiscoveredPathData(LocalAbsPath, LongPackageName, RelPathFromParentDir, StatData.ModificationTime) });
+		AddDiscoveredFile(FDiscoveredPathData(LocalAbsPath, LongPackageName, RelPathFromParentDir, StatData.ModificationTime));
 		if (FPathViews::IsPathLeaf(FileRelPath))
 		{
 			ScanDir->MarkFileAlreadyScanned(FileRelPath);
@@ -2166,7 +2235,7 @@ int32 FAssetDataDiscovery::FindLowerBoundMountPoint(FStringView LocalAbsPath) co
 	);
 }
 
-void FAssetDataDiscovery::AddDiscovered(TConstArrayView<FDiscoveredPathData> SubDirs, TConstArrayView<FDiscoveredPathData> Files)
+void FAssetDataDiscovery::AddDiscovered(FStringView DirAbsPath, TConstArrayView<FDiscoveredPathData> SubDirs, TConstArrayView<FDiscoveredPathData> Files)
 {
 	// This function is inside the critical section so we have moved filtering results outside of it
 	// Caller is responsible for filtering SubDirs and Files by ShouldScan and packagename validity
@@ -2175,13 +2244,26 @@ void FAssetDataDiscovery::AddDiscovered(TConstArrayView<FDiscoveredPathData> Sub
 	{
 		DiscoveredDirectories.Add(FString(SubDir.LongPackageName));
 	}
-	for (const FDiscoveredPathData& DiscoveredFile : Files)
-	{
-		DiscoveredFiles.Emplace(DiscoveredFile); // Emplace passes the FDiscoveredPathData to the FGatheredPathData explicit constructor for it
-	}
+	DiscoveredFiles.Emplace(DirAbsPath, Files);
 	NumDiscoveredFiles += Files.Num();
 }
 
+void FAssetDataDiscovery::AddDiscoveredFile(FDiscoveredPathData&& File)
+{
+	DiscoveredSingleFiles.Emplace(MoveTemp(File));
+	NumDiscoveredFiles++;
+}
+
+FAssetDataDiscovery::FDirectoryResult::FDirectoryResult(FStringView InDirAbsPath, TConstArrayView<FDiscoveredPathData> InFiles)
+	: DirAbsPath(InDirAbsPath)
+{
+	Files.Reserve(InFiles.Num());
+	for (const FDiscoveredPathData& DiscoveredFile : InFiles)
+	{
+		Files.Emplace(DiscoveredFile); // Emplace passes the FDiscoveredPathData to the FGatheredPathData explicit constructor for it
+	}
+
+}
 bool FAssetDataDiscovery::ShouldDirBeReported(FStringView LongPackageName) const
 {
 	return !DirLongPackageNamesToNotReport.ContainsByHash(GetTypeHash(LongPackageName), LongPackageName);
@@ -2234,6 +2316,7 @@ FAssetDataGatherer::FAssetDataGatherer(const TArray<FString>& InBlacklistLongPac
 	bIsSynchronousTick = bIsSynchronous;
 
 	Discovery = MakeUnique<UE::AssetDataGather::Private::FAssetDataDiscovery>(InBlacklistLongPackageNames, InBlacklistMountRelativePaths, bInIsSynchronous);
+	FilesToSearch = MakeUnique<UE::AssetDataGather::Private::FFilesToSearch>();
 }
 
 FAssetDataGatherer::~FAssetDataGatherer()
@@ -2420,7 +2503,7 @@ void FAssetDataGatherer::TickInternal(bool& bOutIsTickInterrupt)
 		IngestDiscoveryResults();
 
 		// Take a batch off of the work list. If we're waiting only on the first WaitBatchCount results don't take more than that
-		int32 NumToProcess = FMath::Min<int32>(BatchSize-LocalFilesToSearch.Num(), FilesToSearch.Num());
+		int32 NumToProcess = FMath::Min<int32>(BatchSize-LocalFilesToSearch.Num(), FilesToSearch->Num());
 		if (WaitBatchCount > 0)
 		{
 			bWaitBatchCountDecremented = true;
@@ -2432,16 +2515,12 @@ void FAssetDataGatherer::TickInternal(bool& bOutIsTickInterrupt)
 			}
 		}
 
-		while (NumToProcess > 0)
-		{
-			LocalFilesToSearch.Add(FilesToSearch.PopFrontValue());
-			--NumToProcess;
-		}
+		FilesToSearch->PopFront(LocalFilesToSearch, NumToProcess);
 
 		// If all work is finished mark idle and exit
 		if (LocalFilesToSearch.Num() == 0 && bDiscoveryIsComplete)
 		{
-			WaitBatchCount = 0; // Clear WaitBatchCount in case it was set higher than FilesToSearch.Num()
+			WaitBatchCount = 0; // Clear WaitBatchCount in case it was set higher than FilesToSearch->Num()
 			bOutIsTickInterrupt = true;
 
 			if (!bFinishedInitialDiscovery)
@@ -2470,7 +2549,7 @@ void FAssetDataGatherer::TickInternal(bool& bOutIsTickInterrupt)
 	{
 		FName PackageName;
 		FName Extension;
-		const FGatheredPathData& AssetFileData;
+		FGatheredPathData& AssetFileData;
 		TArray<FAssetData*> AssetDataFromFile;
 		FPackageDependencyData DependencyData;
 		TArray<FString> CookedPackageNamesWithoutAssetData;
@@ -2478,7 +2557,7 @@ void FAssetDataGatherer::TickInternal(bool& bOutIsTickInterrupt)
 		bool bResult = false;
 		bool bCanceled = false;
 
-		FReadContext(FName InPackageName, FName InExtension, const FGatheredPathData& InAssetFileData)
+		FReadContext(FName InPackageName, FName InExtension, FGatheredPathData& InAssetFileData)
 			: PackageName(InPackageName)
 			, Extension(InExtension)
 			, AssetFileData(InAssetFileData)
@@ -2488,7 +2567,7 @@ void FAssetDataGatherer::TickInternal(bool& bOutIsTickInterrupt)
 
 	// Try to read each file in the batch out of the cache, and accumulate a list for more expensive reading of all of the files that are not in the cache 
 	TArray<FReadContext> ReadContexts;
-	for (const FGatheredPathData& AssetFileData : LocalFilesToSearch)
+	for (FGatheredPathData& AssetFileData : LocalFilesToSearch)
 	{
 		const FName PackageName = *AssetFileData.LongPackageName;
 		const FName Extension = FName(*FPaths::GetExtension(AssetFileData.LocalAbsPath));
@@ -2615,7 +2694,7 @@ void FAssetDataGatherer::TickInternal(bool& bOutIsTickInterrupt)
 		{
 			// If the read temporarily failed, return it to the worklist, pushed to the end
 			FGathererScopeLock ResultsScopeLock(&ResultsLock);
-			FilesToSearch.Add(ReadContext.AssetFileData);
+			FilesToSearch->AddFileForLaterRetry(MoveTemp(ReadContext.AssetFileData));
 		}
 	}
 
@@ -2635,7 +2714,7 @@ void FAssetDataGatherer::TickInternal(bool& bOutIsTickInterrupt)
 				FReadContext& ReadContext = ReadContexts[Index];
 				if (ReadContext.bCanceled)
 				{
-					FilesToSearch.AddFront(ReadContext.AssetFileData);
+					FilesToSearch->AddFileAgainAfterTimeout(MoveTemp(ReadContext.AssetFileData));
 					if (bWaitBatchCountDecremented)
 					{
 						++WaitBatchCount;
@@ -2657,7 +2736,7 @@ void FAssetDataGatherer::IngestDiscoveryResults()
 {
 	CHECK_IS_LOCKED_CURRENT_THREAD(TickLock);
 	CHECK_IS_LOCKED_CURRENT_THREAD(ResultsLock);
-	Discovery->GetAndTrimSearchResults(bDiscoveryIsComplete, DiscoveredPaths, FilesToSearch, NumPathsToSearchAtLastSyncPoint);
+	Discovery->GetAndTrimSearchResults(bDiscoveryIsComplete, DiscoveredPaths, *FilesToSearch, NumPathsToSearchAtLastSyncPoint);
 }
 
 bool FAssetDataGatherer::ReadAssetFile(const FString& AssetFilename, TArray<FAssetData*>& AssetDataList, FPackageDependencyData& DependencyData, TArray<FString>& CookedPackageNamesWithoutAssetData, bool& OutCanRetry) const
@@ -2773,7 +2852,7 @@ void FAssetDataGatherer::GetAndTrimSearchResults(bool& bOutIsSearching, TRingBuf
 	OutSearchTimes.Append(MoveTemp(SearchTimes));
 	SearchTimes.Reset();
 
-	OutNumFilesToSearch = FilesToSearch.Num();
+	OutNumFilesToSearch = FilesToSearch->Num();
 	OutNumPathsToSearch = NumPathsToSearchAtLastSyncPoint;
 	OutIsDiscoveringFiles = !bDiscoveryIsComplete;
 
@@ -2813,22 +2892,24 @@ void FAssetDataGatherer::WaitOnPath(FStringView InPath)
 		}
 	}
 	FString LocalAbsPath = NormalizeLocalPath(InPath);
-	Discovery->SetPropertiesAndWait(LocalAbsPath, false /* bAddToWhitelist */, false /* bForceRescan */, false /* bIgnoreBlacklistScanFilters */);
-	WaitOnPathsInternal(TConstArrayView<FString>(&LocalAbsPath, 1), FString(), TArray<FString>());
+	UE::AssetDataGather::Private::FPathExistence QueryPath(LocalAbsPath);
+	Discovery->SetPropertiesAndWait(QueryPath, false /* bAddToWhitelist */, false /* bForceRescan */, false /* bIgnoreBlacklistScanFilters */);
+	WaitOnPathsInternal(TArrayView<UE::AssetDataGather::Private::FPathExistence>(&QueryPath, 1), FString(), TArray<FString>());
 }
 
-void FAssetDataGatherer::ScanPathsSynchronous(const TArray<FString>& InLocalPaths, bool bForceRescan, bool bIgnoreBlackListScanFilters, const FString& SaveCacheFilename, const TArray<FString>& SaveCacheLongPackageNameDirs)
+void FAssetDataGatherer::ScanPathsSynchronous(const TArray<FString>& InLocalPaths, bool bForceRescan,
+	bool bIgnoreBlackListScanFilters, const FString& SaveCacheFilename, const TArray<FString>& SaveCacheLongPackageNameDirs)
 {
-	TArray<FString> LocalAbsPaths;
-	LocalAbsPaths.Reserve(InLocalPaths.Num());
+	TArray<UE::AssetDataGather::Private::FPathExistence> QueryPaths;
+	QueryPaths.Reserve(InLocalPaths.Num());
 	for (const FString& LocalPath : InLocalPaths)
 	{
-		LocalAbsPaths.Add(NormalizeLocalPath(LocalPath));
+		QueryPaths.Add(UE::AssetDataGather::Private::FPathExistence(NormalizeLocalPath(LocalPath)));
 	}
 
-	for (const FString& LocalAbsPath : LocalAbsPaths)
+	for (UE::AssetDataGather::Private::FPathExistence& QueryPath: QueryPaths)
 	{
-		Discovery->SetPropertiesAndWait(LocalAbsPath, true /* bAddToWhitelist */, bForceRescan, bIgnoreBlackListScanFilters);
+		Discovery->SetPropertiesAndWait(QueryPath, true /* bAddToWhitelist */, bForceRescan, bIgnoreBlackListScanFilters);
 	}
 
 	{
@@ -2836,10 +2917,11 @@ void FAssetDataGatherer::ScanPathsSynchronous(const TArray<FString>& InLocalPath
 		SetIsIdle(false);
 	}
 
-	WaitOnPathsInternal(LocalAbsPaths, SaveCacheFilename, SaveCacheLongPackageNameDirs);
+	WaitOnPathsInternal(QueryPaths, SaveCacheFilename, SaveCacheLongPackageNameDirs);
 }
 
-void FAssetDataGatherer::WaitOnPathsInternal(TConstArrayView<FString> LocalAbsPaths, const FString& SaveCacheFilename, const TArray<FString>& SaveCacheLongPackageNameDirs)
+void FAssetDataGatherer::WaitOnPathsInternal(TArrayView<UE::AssetDataGather::Private::FPathExistence> QueryPaths,
+	const FString& SaveCacheFilename, const TArray<FString>& SaveCacheLongPackageNameDirs)
 {
 	// Request a halt to the async tick
 	FScopedPause ScopedPause(*this);
@@ -2853,7 +2935,7 @@ void FAssetDataGatherer::WaitOnPathsInternal(TConstArrayView<FString> LocalAbsPa
 			IngestDiscoveryResults();
 
 			int32 NumDiscoveredPaths;
-			SortPathsByPriority(LocalAbsPaths, NumDiscoveredPaths);
+			SortPathsByPriority(QueryPaths, UE::AssetDataGather::Private::EPriority::Blocking, NumDiscoveredPaths);
 			if (NumDiscoveredPaths == 0)
 			{
 				return;
@@ -3213,7 +3295,7 @@ uint32 FAssetDataGatherer::GetAllocatedSize() const
 	FGathererScopeLock TickScopeLock(&TickLock);
 	FGathererScopeLock ResultsScopeLock(&ResultsLock);
 
-	Result += GetArrayRecursiveAllocatedSize(FilesToSearch);
+	Result += sizeof(*FilesToSearch) + FilesToSearch->GetAllocatedSize();
 
 	Result += AssetResults.GetAllocatedSize();
 	FAssetDataTagMapSharedView::FMemoryCounter TagMemoryUsage;
@@ -3251,14 +3333,7 @@ uint32 FAssetDataGatherer::GetAllocatedSize() const
 void FAssetDataGatherer::Shrink()
 {
 	CHECK_IS_LOCKED_CURRENT_THREAD(ResultsLock);
-	// TODO: Make TRingBuffer::Shrink
-	TRingBuffer<UE::AssetDataGather::Private::FGatheredPathData> Buffer;
-	Buffer.Reserve(FilesToSearch.Num());
-	for (UE::AssetDataGather::Private::FGatheredPathData& File : FilesToSearch)
-	{
-		Buffer.Add(MoveTemp(File));
-	}
-	Swap(Buffer, FilesToSearch);
+	FilesToSearch->Shrink();
 	AssetResults.Shrink();
 	DependencyResults.Shrink();
 	CookedPackageNamesWithoutAssetDataResults.Shrink();
@@ -3354,35 +3429,39 @@ void FAssetDataGatherer::SetDirectoryProperties(FStringView LocalPath, const UE:
 		if (InProperties.Priority.IsSet())
 		{
 			int32 NumPrioritizedPaths;
-			SortPathsByPriority(TConstArrayView<FString>(&LocalAbsPath, 1), NumPrioritizedPaths);
+			UE::AssetDataGather::Private::FPathExistence QueryPath(LocalAbsPath);
+			SortPathsByPriority(TArrayView<UE::AssetDataGather::Private::FPathExistence>(&QueryPath, 1),
+				*InProperties.Priority, NumPrioritizedPaths);
 		}
 	}
 }
 
-void FAssetDataGatherer::SortPathsByPriority(TConstArrayView<FString> LocalAbsPathsToPrioritize, int32& OutNumPaths)
+void FAssetDataGatherer::SortPathsByPriority(
+	TArrayView<UE::AssetDataGather::Private::FPathExistence> LocalAbsPathsToPrioritize,
+	UE::AssetDataGather::Private::EPriority Priority, int32& OutNumPaths)
 {
 	CHECK_IS_LOCKED_CURRENT_THREAD(ResultsLock);
 	using namespace UE::AssetDataGather::Private;
 
-	// This code needs to be as fast as possible since it is in a critical section!
-
-	// Swap all priority files to the top of the list
-	int32 LowestNonPriorityFileIdx = 0;
-	int32 NumFilesToSearch = FilesToSearch.Num();
-	for (int32 FilenameIdx = 0; FilenameIdx < NumFilesToSearch; ++FilenameIdx)
+	for (UE::AssetDataGather::Private::FPathExistence& QueryPath : LocalAbsPathsToPrioritize)
 	{
-		for (const FString& LocalAbsPathToPrioritize : LocalAbsPathsToPrioritize)
+		switch (QueryPath.GetType())
 		{
-			// Handle LocalAbsPathsToPrioritize that might be a directory, or a filename without an extension, or a full filename
-			if (FPathViews::IsParentPathOf(LocalAbsPathToPrioritize, FilesToSearch[FilenameIdx].LocalAbsPath) ||
-				FPathViews::Equals(LocalAbsPathToPrioritize, FPathViews::GetBaseFilenameWithPath(FilesToSearch[FilenameIdx].LocalAbsPath)))
-			{
-				Swap(FilesToSearch[FilenameIdx], FilesToSearch[LowestNonPriorityFileIdx++]);
-				break;
-			}
+		case FPathExistence::EType::Directory:
+			FilesToSearch->PrioritizeDirectory(QueryPath.GetLocalAbsPath(), Priority);
+			break;
+		case FPathExistence::EType::File:
+			// Intentional fall-through
+		case FPathExistence::EType::MissingButDirExists:
+			// We assume MissingButDirExists is a file missing the extension;
+			// PrioritizeFile handles searching for BaseNameWithPath.* and does nothing if no matches are found
+			FilesToSearch->PrioritizeFile(QueryPath.GetLocalAbsPath(), Priority);
+			break;
+		default:
+			break;
 		}
 	}
-	OutNumPaths = LowestNonPriorityFileIdx;
+	OutNumPaths = FilesToSearch->NumBlockingFiles();
 }
 
 void FAssetDataGatherer::SetIsWhitelisted(FStringView LocalPath, bool bIsWhitelisted)
@@ -3455,3 +3534,352 @@ FStringView FAssetDataGatherer::NormalizeLongPackageName(FStringView LongPackage
 	return LongPackageName;
 }
 
+namespace UE::AssetDataGather::Private
+{
+
+void FFilesToSearch::AddPriorityFile(FGatheredPathData&& FilePath)
+{
+	BlockingFiles.Add(MoveTemp(FilePath));
+}
+
+void FFilesToSearch::AddDirectory(FString&& DirAbsPath, TArray<FGatheredPathData>&& FilePaths)
+{
+	if (FilePaths.Num() == 0)
+	{
+		return;
+	}
+	check(!DirAbsPath.IsEmpty());
+
+	FTreeNode& Node = Root.FindOrAddNode(DirAbsPath);
+	Node.AddFiles(MoveTemp(FilePaths));
+}
+
+void FFilesToSearch::AddFileAgainAfterTimeout(FGatheredPathData&& FilePath)
+{
+	BlockingFiles.AddFront(MoveTemp(FilePath));
+}
+
+void FFilesToSearch::AddFileForLaterRetry(FGatheredPathData&& FilePath)
+{
+	LaterRetryFiles.Add(FilePath);
+}
+
+template <typename AllocatorType>
+void FFilesToSearch::PopFront(TArray<FGatheredPathData, AllocatorType>& Out, int32 NumToPop)
+{
+	while (NumToPop > 0 && !BlockingFiles.IsEmpty())
+	{
+		Out.Add(BlockingFiles.PopFrontValue());
+		--NumToPop;
+	}
+	Root.PopFiles(Out, NumToPop);
+	while (NumToPop > 0 && !LaterRetryFiles.IsEmpty())
+	{
+		Out.Add(LaterRetryFiles.PopFrontValue());
+		--NumToPop;
+	}
+}
+
+void FFilesToSearch::PrioritizeDirectory(FStringView DirAbsPath, EPriority Priority)
+{
+	// We may need to prioritize a LaterRetryFile that is now loadable, so add them all into the Root
+	while (!LaterRetryFiles.IsEmpty())
+	{
+		FGatheredPathData FilePath = LaterRetryFiles.PopFrontValue();
+		FTreeNode& Node = Root.FindOrAddNode(FPathViews::GetPath(FilePath.LocalAbsPath));
+		Node.AddFile(MoveTemp(FilePath));
+	}
+
+	if (Priority > EPriority::Blocking)
+	{
+		// TODO: Implement another tree that is searched first for the High Priority 
+		// We cannot add the High Priority files to the BlockingFiles array, because
+		// then blocking on BlockingFiles to be empty could be slow. We cannot add them
+		// as a separate simple array, because we would have to search that (sometimes large)
+		// array linearly when looking for files to accomodate a blocking priority request
+		return;
+	}
+	FTreeNode* TreeNode = Root.FindNode(DirAbsPath);
+	if (TreeNode)
+	{
+		TreeNode->PopAllFiles(BlockingFiles);
+	}
+}
+
+void FFilesToSearch::PrioritizeFile(FStringView FileAbsPathExtOptional, EPriority Priority)
+{
+	// We may need to prioritize a LaterRetryFile that is now loadable, so add them all into the Root
+	while (!LaterRetryFiles.IsEmpty())
+	{
+		FGatheredPathData FilePath = LaterRetryFiles.PopFrontValue();
+		FTreeNode& Node = Root.FindOrAddNode(FPathViews::GetPath(FilePath.LocalAbsPath));
+		Node.AddFile(MoveTemp(FilePath));
+	}
+
+	if (Priority > EPriority::Blocking)
+	{
+		// TODO: Implement High Priority; see note in PrioritizeDirectory
+		return;
+	}
+	FTreeNode* TreeNode = Root.FindNode(FPathViews::GetPath(FileAbsPathExtOptional));
+	if (TreeNode)
+	{
+		int32 BeforeSize = BlockingFiles.Num();
+		TreeNode->PopMatchingDirectFiles(BlockingFiles, FileAbsPathExtOptional);
+	}
+}
+
+int32 FFilesToSearch::NumBlockingFiles() const
+{
+	return BlockingFiles.Num();
+}
+
+void FFilesToSearch::Shrink()
+{
+	// TODO: Make TRingBuffer::Shrink
+	TRingBuffer<UE::AssetDataGather::Private::FGatheredPathData> Buffer;
+	Buffer.Reserve(BlockingFiles.Num());
+	for (UE::AssetDataGather::Private::FGatheredPathData& File : BlockingFiles)
+	{
+		Buffer.Add(MoveTemp(File));
+	}
+	Swap(Buffer, BlockingFiles);
+
+	Buffer.Empty(LaterRetryFiles.Num());
+	for (UE::AssetDataGather::Private::FGatheredPathData& File : LaterRetryFiles)
+	{
+		Buffer.Add(MoveTemp(File));
+	}
+	Swap(Buffer, LaterRetryFiles);
+
+	Root.Shrink();
+}
+
+int32 FFilesToSearch::Num() const
+{
+	return BlockingFiles.Num() + Root.NumFiles() + LaterRetryFiles.Num();
+}
+
+uint32 FFilesToSearch::GetAllocatedSize() const
+{
+	uint32 Size = 0;
+	Size += BlockingFiles.GetAllocatedSize();
+	for (const FGatheredPathData& PathData : BlockingFiles)
+	{
+		Size += PathData.GetAllocatedSize();
+	}
+	Size += Root.GetAllocatedSize();
+	Size += LaterRetryFiles.GetAllocatedSize();
+	for (const FGatheredPathData& PathData : LaterRetryFiles)
+	{
+		Size += PathData.GetAllocatedSize();
+	}
+	return Size;
+}
+
+FFilesToSearch::FTreeNode::FTreeNode(FStringView InRelPath)
+	:RelPath(InRelPath)
+{
+}
+
+FStringView FFilesToSearch::FTreeNode::GetRelPath() const
+{
+	return RelPath;
+}
+
+FFilesToSearch::FTreeNode& FFilesToSearch::FTreeNode::FindOrAddNode(FStringView InRelPath)
+{
+	if (InRelPath.IsEmpty())
+	{
+		return *this;
+	}
+	FStringView FirstComponent;
+	FStringView RemainingPath;
+	FPathViews::SplitFirstComponent(InRelPath, FirstComponent, RemainingPath);
+	FTreeNode& SubDir = FindOrAddSubDir(FirstComponent);
+	return SubDir.FindOrAddNode(RemainingPath);
+}
+
+FFilesToSearch::FTreeNode* FFilesToSearch::FTreeNode::FindNode(FStringView InRelPath)
+{
+	if (InRelPath.IsEmpty())
+	{
+		return this;
+	}
+	FStringView FirstComponent;
+	FStringView RemainingPath;
+	FPathViews::SplitFirstComponent(InRelPath, FirstComponent, RemainingPath);
+	FTreeNode* SubDir = FindSubDir(FirstComponent);
+	if (!SubDir)
+	{
+		return nullptr;
+	}
+
+	return SubDir->FindNode(RemainingPath);
+}
+
+FFilesToSearch::FTreeNode& FFilesToSearch::FTreeNode::FindOrAddSubDir(FStringView SubDirBaseName)
+{
+	int32 Index = FindLowerBoundSubDir(SubDirBaseName);
+	if (Index == SubDirs.Num() || !FPathViews::Equals(SubDirs[Index]->GetRelPath(), SubDirBaseName))
+	{
+		return *SubDirs.EmplaceAt_GetRef(Index, MakeUnique<FTreeNode>(SubDirBaseName));
+	}
+	else
+	{
+		return *SubDirs[Index];
+	}
+}
+
+FFilesToSearch::FTreeNode* FFilesToSearch::FTreeNode::FindSubDir(FStringView SubDirBaseName)
+{
+	int32 Index = FindLowerBoundSubDir(SubDirBaseName);
+	if (Index == SubDirs.Num() || !FPathViews::Equals(SubDirs[Index]->GetRelPath(), SubDirBaseName))
+	{
+		return nullptr;
+	}
+	else
+	{
+		return SubDirs[Index].Get();
+	}
+}
+
+int32 FFilesToSearch::FTreeNode::FindLowerBoundSubDir(FStringView SubDirBaseName)
+{
+	return Algo::LowerBound(SubDirs, SubDirBaseName,
+		[](const TUniquePtr<FTreeNode>& SubDir, FStringView BaseName)
+		{
+			return FPathViews::Less(SubDir->RelPath, BaseName);
+		}
+	);
+}
+
+void FFilesToSearch::FTreeNode::AddFiles(TArray<FGatheredPathData>&& FilePaths)
+{
+	if (Files.Num() == 0)
+	{
+		Files = MoveTemp(FilePaths);
+	}
+	else
+	{
+		Files.Append(MoveTemp(FilePaths));
+	}
+}
+
+void FFilesToSearch::FTreeNode::AddFile(FGatheredPathData&& FilePath)
+{
+	Files.Add(MoveTemp(FilePath));
+}
+
+template <typename RangeType>
+void FFilesToSearch::FTreeNode::PopFiles(RangeType& Out, int32& NumToPop)
+{
+	while (NumToPop > 0 && !Files.IsEmpty())
+	{
+		Out.Add(Files.Pop(false /* bAllowShrinking */));
+		--NumToPop;
+	}
+	while (NumToPop > 0 && SubDirs.Num() != 0)
+	{
+		TUniquePtr<FTreeNode>& SubDir = SubDirs[SubDirs.Num() - 1];
+		SubDir->PopFiles(Out, NumToPop);
+		if (SubDir->IsEmpty())
+		{
+			SubDirs.RemoveAt(SubDirs.Num() - 1);
+		}
+	}
+}
+
+template <typename RangeType>
+void FFilesToSearch::FTreeNode::PopAllFiles(RangeType& Out)
+{
+	while (!Files.IsEmpty())
+	{
+		Out.Add(Files.Pop(false /* bAllowShrinking */));
+	}
+	for (int32 Index = SubDirs.Num() - 1; Index >= 0; --Index) // Match the order of PopFiles
+	{
+		TUniquePtr<FTreeNode>& SubDir = SubDirs[Index];
+		SubDir->PopAllFiles(Out);
+		SubDir.Reset();
+	}
+	SubDirs.Empty();
+}
+
+template <typename RangeType>
+void FFilesToSearch::FTreeNode::PopMatchingDirectFiles(RangeType& Out, FStringView FileAbsPathExtOptional)
+{
+	// TODO: Make this more performant by sorting the list of Files.
+	// To prevent shifting costs, when popping the file we will leave a placeholder behind with an ignore flag set.
+	FStringView FileAbsPathNoExt = FPathViews::GetBaseFilenameWithPath(FileAbsPathExtOptional);
+	int32 NumFiles = Files.Num();
+	for (int32 Index = 0; Index < NumFiles; )
+	{
+		FGatheredPathData& PathData = Files[Index];
+		if (FPathViews::Equals(FPathViews::GetBaseFilenameWithPath(PathData.LocalAbsPath), FileAbsPathNoExt))
+		{
+			Out.Add(MoveTemp(PathData));
+			Files.RemoveAt(Index);
+			--NumFiles;
+		}
+		else
+		{
+			++Index;
+		}
+	}
+}
+
+void FFilesToSearch::FTreeNode::PruneEmptyChild(FStringView SubDirBaseName)
+{
+	int32 Index = FindLowerBoundSubDir(SubDirBaseName);
+	if (!(Index == SubDirs.Num() || !FPathViews::Equals(SubDirs[Index]->GetRelPath(), SubDirBaseName)))
+	{
+		if (SubDirs[Index]->IsEmpty())
+		{
+			SubDirs.RemoveAt(Index);
+		}
+	}
+}
+
+bool FFilesToSearch::FTreeNode::IsEmpty() const
+{
+	return Files.Num() == 0 && SubDirs.Num() == 0;
+}
+
+void FFilesToSearch::FTreeNode::Shrink()
+{
+	Files.Shrink();
+	SubDirs.Shrink();
+	for (TUniquePtr<FTreeNode>& SubDir : SubDirs)
+	{
+		SubDir->Shrink();
+	}
+}
+
+uint32 FFilesToSearch::FTreeNode::GetAllocatedSize() const
+{
+	uint32 Size = 0;
+	Size = Files.GetAllocatedSize();
+	for (const FGatheredPathData& File : Files)
+	{
+		Size += File.GetAllocatedSize();
+	}
+	Size += SubDirs.GetAllocatedSize() + SubDirs.Num()*sizeof(FTreeNode);
+	for (const TUniquePtr<FTreeNode>& SubDir : SubDirs)
+	{
+		Size += SubDir->GetAllocatedSize();
+	}
+	return Size;
+}
+
+int32 FFilesToSearch::FTreeNode::NumFiles() const
+{
+	int32 Num = Files.Num();
+	for (const TUniquePtr<FTreeNode>& SubDir : SubDirs)
+	{
+		Num += SubDir->NumFiles();
+	}
+	return Num;
+}
+
+}
diff --git a/Engine/Source/Runtime/AssetRegistry/Private/AssetDataGatherer.h b/Engine/Source/Runtime/AssetRegistry/Private/AssetDataGatherer.h
--- a/Engine/Source/Runtime/AssetRegistry/Private/AssetDataGatherer.h
+++ b/Engine/Source/Runtime/AssetRegistry/Private/AssetDataGatherer.h
@@ -28,17 +28,14 @@ class FDiskCachedAssetData;
 class FAssetRegistryReader;
 class FAssetRegistryWriter; // Not defined if !ALLOW_NAME_BATCH_SAVING
 class FPackageReader;
-namespace UE
-{
-namespace AssetDataGather
-{
-namespace Private
+namespace UE::AssetDataGather::Private
 {
 class FAssetDataDiscovery;
+class FFilesToSearch;
 struct FGatheredPathData;
+struct FPathExistence;
 struct FSetPathProperties;
-}
-}
+enum class EPriority : uint8;
 }
 
 #if DO_CHECK
@@ -199,11 +196,16 @@ private:
 	 */
 	void SetDirectoryProperties(FStringView LocalPath, const UE::AssetDataGather::Private::FSetPathProperties& Properties);
 
-	/** Wait for all monitored assets under the given path to be added to search results. Returns immediately if the given path are not monitored. */
-	void WaitOnPathsInternal(TConstArrayView<FString> LocalAbsPaths, const FString& SaveCacheFilename, const TArray<FString>& SaveCacheFilterDirs);
+	/**
+	 * Wait for all monitored assets under the given path to be added to search results.
+	 * Returns immediately if the given path are not monitored.
+	 */
+	void WaitOnPathsInternal(TArrayView<UE::AssetDataGather::Private::FPathExistence> QueryPaths,
+		const FString& SaveCacheFilename, const TArray<FString>& SaveCacheFilterDirs);
 
-	/** Sort the pending list of discovered asset filepaths so that assets under the given directory/filename are processed first. */
-	void SortPathsByPriority(TConstArrayView<FString> LocalAbsPaths, int32& OutNumPaths);
+	/** Sort the pending list of filepaths so that assets under the given directory/filename are processed first. */
+	void SortPathsByPriority(TArrayView<UE::AssetDataGather::Private::FPathExistence> QueryPaths,
+		UE::AssetDataGather::Private::EPriority Priority, int32& OutNumPaths);
 	/**
 	 * Reads FAssetData information out of a file
 	 *
@@ -317,7 +319,8 @@ private:
 	// Variable section for variables that are read/writable only within ResultsLock.
 
 	/** List of files that need to be processed by the search. Read/writable only within ResultsLock. */
-	TRingBuffer<UE::AssetDataGather::Private::FGatheredPathData> FilesToSearch;
+	TUniquePtr<UE::AssetDataGather::Private::FFilesToSearch> FilesToSearch;
+
 	/** The asset data gathered from the searched files. Read/writable only within ResultsLock. */
 	TArray<FAssetData*> AssetResults;
 	/** Dependency data gathered from the searched files packages. Read/writable only within ResultsLock. */
diff --git a/Engine/Source/Runtime/AssetRegistry/Private/AssetDataGathererPrivate.h b/Engine/Source/Runtime/AssetRegistry/Private/AssetDataGathererPrivate.h
--- a/Engine/Source/Runtime/AssetRegistry/Private/AssetDataGathererPrivate.h
+++ b/Engine/Source/Runtime/AssetRegistry/Private/AssetDataGathererPrivate.h
@@ -101,6 +101,7 @@ struct FGatheredPathData
 	FGatheredPathData() = default;
 	FGatheredPathData(FStringView InLocalAbsPath, FStringView InLongPackageName, const FDateTime& InPackageTimestamp);
 	explicit FGatheredPathData(const FDiscoveredPathData& DiscoveredData);
+	explicit FGatheredPathData(FDiscoveredPathData&& DiscoveredData);
 	void Assign(FStringView InLocalAbsPath, FStringView InLongPackageName, const FDateTime& InPackageTimestamp);
 	void Assign(const FDiscoveredPathData& DiscoveredData);
 
@@ -111,6 +112,133 @@ struct FGatheredPathData
 	uint32 GetAllocatedSize() const;
 };
 
+/** A container to efficiently receive files from discovery and search the container to raise priority of a path */
+class FFilesToSearch
+{
+public:
+	/** Add the given FilePath as a single file with blocking priority. */
+	void AddPriorityFile(FGatheredPathData&& FilePath);
+	/** Add the given directory and its discovered direct Files. */
+	void AddDirectory(FString&& DirAbsPath, TArray<FGatheredPathData>&& FilePaths);
+	/** Re-add a file from PopFront that we didn't have time to process. */
+	void AddFileAgainAfterTimeout(FGatheredPathData&& FilePath);
+	/** Re-add a file from PopFront that uses a not-yet-loaded class for later retry. */
+	void AddFileForLaterRetry(FGatheredPathData&& FilePath);
+
+	/** In coarse priority order, pop the given number of files out of the container. */
+	template <typename AllocatorType>
+	void PopFront(TArray<FGatheredPathData, AllocatorType>& Out, int32 NumToPop);
+
+	/** Search the container for all files with the given directory and set their priority. */
+	void PrioritizeDirectory(FStringView DirAbsPath, EPriority Priority);
+	/**
+	 * Search the container for all files with the given Path and set their priority.
+	 * Path may be BaseNameWithPath or the full BaseNameWithPath.Extension.
+	 */
+	void PrioritizeFile(FStringView FileAbsPathExtOptional, EPriority Priority);
+	/** How many files with blocking priority are in the container. */
+	int32 NumBlockingFiles() const;
+
+	/** Reduce memory used in buffers. */
+	void Shrink();
+	/** How many files are in the container. */
+	int32 Num() const;
+	/** How much memory is used by the container, not counting sizeof(*this). */
+	uint32 GetAllocatedSize() const;
+private:
+	/** A directory-search-tree structure; each node has a list of direct files and subdirectories. */
+	struct FTreeNode
+	{
+		FTreeNode() = default;
+		explicit FTreeNode(FStringView InRelPath);
+
+		/** The name of the directory relative to its parent. */
+		FStringView GetRelPath() const;
+
+		/**
+		 * Recursively search the tree to find the given relative directory name,
+		 * adding nodes for the directory and its parents if required.
+		 */
+		FTreeNode& FindOrAddNode(FStringView InRelPath);
+		/**
+		 * Recursively search the tree to find the given relative directory name;
+		 * return nullptr if the node or any parent is not found.
+		 */
+		FTreeNode* FindNode(FStringView InRelPath);
+
+		/** Add the given FilePaths as direct file children of the Node's directory. */
+		void AddFiles(TArray<FGatheredPathData>&& FilePaths);
+		/** Add the given FilePath as a direct file child of the Node's directory. */
+		void AddFile(FGatheredPathData&& FilePath);
+
+		/** Pop the given number of files out of this node and its children and decrement NumToPop by how many were popped. */
+		template <typename RangeType>
+		void PopFiles(RangeType& Out, int32& NumToPop);
+		/** Pop all files out of this node and its children. */
+		template <typename RangeType>
+		void PopAllFiles(RangeType& Out);
+		/** Pop files out of *this but not children, if they match the given full path or BaseNameWithPath. */
+		template <typename RangeType>
+		void PopMatchingDirectFiles(RangeType& Out, FStringView FileAbsPathExtOptional);
+
+		/** Delete the given child node if it is empty. */
+		void PruneEmptyChild(FStringView SubDirBaseName);
+		/** Return true if this node has no files or subdirs. */
+		bool IsEmpty() const;
+		/** Reduce memory used in buffers. */
+		void Shrink();
+		/** How much memory is used by *this, not counting sizeof(*this). */
+		uint32 GetAllocatedSize() const;
+		/** Number of files in *this and its child nodes. */
+		int32 NumFiles() const;
+	private:
+		/** Search direct SubDirs for the given BaseName and add a node for it if not existing. */
+		FTreeNode& FindOrAddSubDir(FStringView SubDirBaseName);
+		/** Search direct SubDirs for the given BaseName and return nullptr if not existing. */
+		FTreeNode* FindSubDir(FStringView SubDirBaseName);
+		/** Find the index of the subdir with the given Relative path. */
+		int32 FindLowerBoundSubDir(FStringView SubDirBaseName);
+
+		FString RelPath;
+		TArray<FGatheredPathData> Files; // Not sorted
+		TArray<TUniquePtr<FTreeNode>> SubDirs; // Sorted
+	};
+	FTreeNode Root;
+	TRingBuffer<UE::AssetDataGather::Private::FGatheredPathData> BlockingFiles;
+	TRingBuffer<UE::AssetDataGather::Private::FGatheredPathData> LaterRetryFiles;
+};
+
+
+/** Stores a LocalAbsPath and its existence information. A file system query is issued the first time the data is needed. */
+struct FPathExistence
+{
+	/** What kind of thing we found at the given path */
+	enum EType
+	{
+		Directory,
+		File,
+		MissingButDirExists,
+		MissingParentDir,
+	};
+
+	explicit FPathExistence(FStringView InLocalAbsPath);
+	const FString& GetLocalAbsPath() const;
+	EType GetType();
+	FDateTime GetModificationTime();
+
+	/** If the path is a directory or file, return the path. Otherwise return the parent directory, if it exists. */
+	FStringView GetLowestExistingPath();
+
+	/** Issue the file system query if not already done. */
+	void LoadExistenceData();
+
+private:
+	FString LocalAbsPath;
+	FDateTime ModificationTime;
+	EType PathType = EType::MissingParentDir;
+	bool bHasExistenceData = false;
+};
+
 /**
  * Tree data node representing a directory in the scan; direct subdirectories are stored as other FScanDir instances
  * referenced from the SubDirs array. Directories are removed from the tree once their scans are finished to save memory.
@@ -449,11 +577,11 @@ public:
 	// Receiving Results and reading properties (possibly while tick is running)
 	/** Gets search results from the file discovery. */
 	void GetAndTrimSearchResults(bool& bOutIsComplete, TArray<FString>& OutDiscoveredPaths,
-		TRingBuffer<FGatheredPathData>& OutDiscoveredFiles, int32& OutNumPathsToSearch);
+		FFilesToSearch& OutFilesToSearch, int32& OutNumPathsToSearch);
 	/** Wait (joining in on the tick) until all currently monitored paths have been scanned. */
 	void WaitForIdle();
 	/** Optionally set some scan properties for the given path and then wait for the scan of it to finish. */
-	void SetPropertiesAndWait(const FString& LocalAbsPath, bool bAddToWhitelist, bool bForceRescan,
+	void SetPropertiesAndWait(FPathExistence& QueryPath, bool bAddToWhitelist, bool bForceRescan,
 		bool bIgnoreBlackListScanFilters);
 	/** Return whether the given path is whitelisted due to e.g. TrySetDirectoryProperties with IsWhitelisted. */
 	bool IsWhitelisted(FStringView LocalAbsPath) const;
@@ -504,7 +632,10 @@ private:
 	void SetIsIdle(bool bInIdle);
 
 	/** Store the given discovered files and directories in the results. */
-	void AddDiscovered(TConstArrayView<FDiscoveredPathData> SubDirs, TConstArrayView<FDiscoveredPathData> Files);
+	void AddDiscovered(FStringView DirAbsPath, TConstArrayView<FDiscoveredPathData> SubDirs, TConstArrayView<FDiscoveredPathData> Files);
+	/** Store the given specially reported single file in the results. */
+	void AddDiscoveredFile(FDiscoveredPathData&& File);
+
 	/**
 	 * Return whether a directory with the given LongPackageName should be reported to the AssetRegistry
 	 * We do not report some directories because they are paths that should not enter the AssetRegistry list of paths if empty,
@@ -600,8 +731,17 @@ private:
 
 	/** Directories found in the scan; may be empty. Read/writable only within ResultsLock. */
 	TArray<FString> DiscoveredDirectories;
+
 	/** Files found found in the scan. Read/writable only within ResultsLock. */
-	TArray<FGatheredPathData> DiscoveredFiles;
+	struct FDirectoryResult
+	{
+		FDirectoryResult(FStringView InDirAbsPath, TConstArrayView<FDiscoveredPathData> InFiles);
+		FString DirAbsPath;
+		TArray<FGatheredPathData> Files;
+		uint32 GetAllocatedSize() const;
+	};
+	TArray<FDirectoryResult> DiscoveredFiles;
+	TArray<FGatheredPathData> DiscoveredSingleFiles;
 	/**
 	 * Time at which the scan was started or last resumed from idle. Used for logging.
 	 * Read/writable only within ResultsLock.
