AssetDataGatherer: When scanning a specific file, exit as well if directory is complete; this is necessary to handle non-whitelisted directories that never scan.
#rnx

diff --git a/Engine/Source/Runtime/AssetRegistry/Private/AssetDataGatherer.cpp b/Engine/Source/Runtime/AssetRegistry/Private/AssetDataGatherer.cpp
--- a/Engine/Source/Runtime/AssetRegistry/Private/AssetDataGatherer.cpp
+++ b/Engine/Source/Runtime/AssetRegistry/Private/AssetDataGatherer.cpp
@@ -1802,7 +1802,7 @@ void FAssetDataDiscovery::SetPropertiesAndWait(FPathExistence& QueryPath, bool b
 					{
 						break;
 					}
-					if ((bScanEntireTree && ScanDir->IsComplete()) || (!bScanEntireTree && ScanDir->HasScanned()))
+					if (ScanDir->IsComplete() || (!bScanEntireTree && ScanDir->HasScanned()))
 					{
 						break;
 					}
