FAssetDataGatherer: Add missing critical section in newly added FAssetDataDiscovery::AddDiscoveredFile.
#rnx

diff --git a/Engine/Source/Runtime/AssetRegistry/Private/AssetDataGatherer.cpp b/Engine/Source/Runtime/AssetRegistry/Private/AssetDataGatherer.cpp
--- a/Engine/Source/Runtime/AssetRegistry/Private/AssetDataGatherer.cpp
+++ b/Engine/Source/Runtime/AssetRegistry/Private/AssetDataGatherer.cpp
@@ -2250,6 +2250,7 @@ void FAssetDataDiscovery::AddDiscovered(FStringView DirAbsPath, TConstArrayView<
 
 void FAssetDataDiscovery::AddDiscoveredFile(FDiscoveredPathData&& File)
 {
+	FGathererScopeLock ResultsScopeLock(&ResultsLock);
 	DiscoveredSingleFiles.Emplace(MoveTemp(File));
 	NumDiscoveredFiles++;
 }
