Added integrity checking for asset registry caches

diff --git a/Engine/Source/Runtime/AssetRegistry/Private/AssetDataGatherer.cpp b/Engine/Source/Runtime/AssetRegistry/Private/AssetDataGatherer.cpp
--- a/Engine/Source/Runtime/AssetRegistry/Private/AssetDataGatherer.cpp
+++ b/Engine/Source/Runtime/AssetRegistry/Private/AssetDataGatherer.cpp
@@ -12,7 +12,9 @@
 #include "Containers/BinaryHeap.h"
 #include "HAL/PlatformProcess.h"
 #include "HAL/RunnableThread.h"
+#include "Hash/Blake3.h"
 #include "Math/NumericLimits.h"
+#include "Memory/MemoryView.h"
 #include "Misc/AsciiSet.h"
 #include "Misc/Char.h"
 #include "Misc/CommandLine.h"
@@ -30,7 +32,7 @@ namespace AssetDataGathererConstants
 	constexpr int32 SingleThreadFilesPerBatch = 3;
 	constexpr int32 ExpectedMaxBatchSize = 100;
 	constexpr int32 MinSecondsToElapseBeforeCacheWrite = 60;
-	static constexpr uint32 CacheSerializationMagic = 0x89ABCDF0;
+	static constexpr uint32 CacheSerializationMagic = 0xCBA78340; // Added integrity checking
 }
 
 namespace UE
@@ -2270,6 +2272,206 @@ bool FAssetDataDiscovery::ShouldDirBeReported(FStringView LongPackageName) const
 	return !DirLongPackageNamesToNotReport.ContainsByHash(GetTypeHash(LongPackageName), LongPackageName);
 }
 
+
+class FChecksumArchiveBase : public FArchiveProxy
+{
+	// Enables both versioning and distinguishing out-of-sync reads from data corruption
+	static constexpr uint32 HeaderMagic = 0xb1a3;
+	static constexpr uint32 SaveBlockSize = 4 << 20;
+
+	static uint64 CalculateChecksum(FMemoryView Data)
+	{
+		return INTEL_ORDER64(*reinterpret_cast<uint64*>(FBlake3::HashBuffer(Data).GetBytes()));
+	}
+
+	struct FBlockHeader
+	{
+		uint32 Magic = 0;
+		uint32 Size = 0;
+		uint64 Checksum = 0;
+	};
+
+	struct FBlock
+	{
+		uint8* Begin = nullptr;
+		uint8* Cursor = nullptr;
+		uint8* End = nullptr;
+
+		explicit FBlock(uint32 Size)				{ Reset(Size); }
+		~FBlock()									{ FMemory::Free(Begin); }
+		uint64 GetCapacity() const					{ return End - Begin; }
+		uint64 GetUsedSize() const					{ return Cursor - Begin; }
+		uint64 GetRemainingSize() const				{ return End - Cursor; }
+		FMutableMemoryView GetUsed() const			{ return {Begin, GetUsedSize()}; }
+		FMutableMemoryView GetRemaining() const		{ return {Cursor, GetRemainingSize()}; };
+
+		void Reset(uint32 Size)
+		{
+			if (GetCapacity() < Size)
+			{
+				FMemory::Free(Begin);
+				Begin = (uint8*)FMemory::Malloc(Size);
+			}
+			
+			// All blocks have the same size except the last one, which may be smaller.
+			// It doesn't matter that we lose some capacity when loading the last block.
+			End = Begin + Size;
+			Cursor = Begin;
+		}
+
+		void Write(FMemoryView In)
+		{
+			check(GetRemainingSize() >= In.GetSize());
+			FMemory::Memcpy(Cursor, In.GetData(), In.GetSize());
+			Cursor += In.GetSize();	
+		}
+
+		void Read(FMutableMemoryView Out)
+		{
+			check(GetRemainingSize() >= Out.GetSize());
+			FMemory::Memcpy(Out.GetData(), Cursor, Out.GetSize());
+			Cursor += Out.GetSize();	
+		}
+	};
+
+	FBlock Block;
+
+	void SaveBlock()
+	{
+		FBlockHeader Header;
+		Header.Magic = HeaderMagic;
+		Header.Size = static_cast<uint32>(Block.GetUsedSize());
+		Header.Checksum = CalculateChecksum(Block.GetUsed());
+		InnerArchive << Header.Magic << Header.Size << Header.Checksum;
+
+		InnerArchive.Serialize(Block.Begin, Header.Size);
+
+		Block.Cursor = Block.Begin;
+	}
+
+	bool LoadBlock()
+	{
+		check(Block.GetRemainingSize() == 0);
+
+		FBlockHeader Header;
+		InnerArchive << Header.Magic << Header.Size << Header.Checksum;
+
+		if (InnerArchive.IsError())
+		{
+			UE_LOG(LogAssetRegistry, Error, TEXT("Couldn't read block header"));
+			return false;
+		}
+		else if (Header.Magic != HeaderMagic)
+		{
+			UE_LOG(LogAssetRegistry, Error, TEXT("Wrong block magic (0x%x)"), Header.Magic);
+			return false;
+		}
+
+		Block.Reset(Header.Size);
+			
+		InnerArchive.Serialize(Block.Begin, Header.Size);
+
+		if (InnerArchive.IsError())
+		{
+			UE_LOG(LogAssetRegistry, Error, TEXT("Couldn't read block data"));
+			return false;
+		}
+		else if (CalculateChecksum(Block.GetRemaining()) != Header.Checksum)
+		{
+			UE_LOG(LogAssetRegistry, Error, TEXT("Wrong block checksum"));
+			return false;
+		}
+
+		return true;
+	}
+
+public:
+	explicit FChecksumArchiveBase(FArchive& Inner)
+	: FArchiveProxy(Inner)
+	, Block(IsLoading() ? 0 : SaveBlockSize)
+	{}
+
+protected:
+	~FChecksumArchiveBase()
+	{
+		if (!IsLoading() && Block.GetUsedSize() > 0)
+		{
+			SaveBlock();
+		}
+	}
+
+	void Save(FMemoryView Data)
+	{
+		for (uint64 Size = Block.GetRemainingSize(); Size < Data.GetSize(); Size = Block.GetRemainingSize())
+		{
+			Block.Write(Data.Left(Size));
+			Data += Size;
+			SaveBlock();
+		}
+
+		Block.Write(Data);
+	}
+
+	void Load(FMutableMemoryView Data)
+	{
+		if (IsError())
+		{
+			return;
+		}
+
+		for (uint64 Size = Block.GetRemainingSize(); Size < Data.GetSize(); Size = Block.GetRemainingSize())
+		{
+			Block.Read(Data.Left(Size));
+			Data += Size;
+
+			if (!LoadBlock())
+			{
+				UE_LOG(LogAssetRegistry, Error, TEXT("Integrity check failed, '%s' cache will be discarded"), *InnerArchive.GetArchiveName());
+				SetError();
+				return;
+			}
+		}
+
+		Block.Read(Data);
+	}
+
+public:
+	// Use FArchive implementations that map back to Serialize() instead of FArchiveProxy overloads
+	// that forward to the inner archive and bypass integrity checking
+	virtual FArchive& operator<<(FText& Value) override					{ return FArchive::operator<<(Value); }
+	virtual void SerializeBits(void* Bits, int64 LengthBits) override	{ FArchive::SerializeBits(Bits, LengthBits); }
+	virtual void SerializeInt(uint32& Value, uint32 Max) override		{ FArchive::SerializeInt(Value, Max); }
+	virtual void SerializeIntPacked(uint32& Value) override				{ FArchive::SerializeIntPacked(Value); }
+
+private:
+	// Wrapping an inner FArchiveUObject is not supported. The inner archive should be low-level
+	// and an outer archive should have intercepted these calls.
+	virtual FArchive& operator<<(FName& Value) override					{ unimplemented(); return *this; }
+	virtual FArchive& operator<<(UObject*& Value) override				{ unimplemented(); return *this; }
+	virtual FArchive& operator<<(FObjectPtr& Value) override			{ unimplemented(); return *this; }
+	virtual FArchive& operator<<(FLazyObjectPtr& Value) override		{ unimplemented(); return *this; }
+	virtual FArchive& operator<<(FSoftObjectPath& Value) override		{ unimplemented(); return *this; }
+	virtual FArchive& operator<<(FSoftObjectPtr& Value) override		{ unimplemented(); return *this; }
+	virtual FArchive& operator<<(FWeakObjectPtr& Value) override		{ unimplemented(); return *this; }
+	virtual FArchive& operator<<(FField*& Value) override				{ unimplemented(); return *this; }
+
+	virtual void Seek(int64 InPos) override								{ unimplemented(); }
+};
+
+class FChecksumArchiveWriter : public FChecksumArchiveBase
+{
+public:
+	using FChecksumArchiveBase::FChecksumArchiveBase;
+	virtual void Serialize(void* V, int64 Len) override { Save(FMemoryView(V, Len)); }
+};
+
+class FChecksumArchiveReader : public FChecksumArchiveBase
+{
+public:
+	using FChecksumArchiveBase::FChecksumArchiveBase;
+	virtual void Serialize(void* V, int64 Len) override { Load(FMutableMemoryView(V, Len)); }
+};
+
 } // namespace Private
 } // namespace AssetDataGather
 } // namespace UE
@@ -3072,7 +3274,8 @@ void FAssetDataGatherer::LoadCacheFileInternal(FStringView CacheFilename)
 			FAssetRegistryVersion::Type RegistryVersion;
 			if (FAssetRegistryVersion::SerializeVersion(*FileAr, RegistryVersion) && RegistryVersion == FAssetRegistryVersion::LatestVersion)
 			{
-				FAssetRegistryReader RegistryReader(*FileAr);
+				UE::AssetDataGather::Private::FChecksumArchiveReader ChecksummingReader(*FileAr);
+				FAssetRegistryReader RegistryReader(ChecksummingReader);
 				if (!RegistryReader.IsError())
 				{
 					SerializeCacheLoad(RegistryReader);
@@ -3167,7 +3370,8 @@ void FAssetDataGatherer::SaveCacheFileInternal(const FString& CacheFilename, con
 		{
 			// We might be able to reduce load time by using AssetRegistry::SerializationOptions
 			// to save certain common tags as FName.
-			FAssetRegistryWriter Ar(FAssetRegistryWriterOptions(), *FileAr);
+			UE::AssetDataGather::Private::FChecksumArchiveWriter ChecksummingWriter(*FileAr);
+			FAssetRegistryWriter Ar(FAssetRegistryWriterOptions(), ChecksummingWriter);
 			SerializeCacheSave(Ar, AssetsToSave);
 		}
 #else		
