AssetDataGatherer: Do not retry files that are missing classes until bInitialPluginsLoaded is set to true, so that we do not waste thread time retrying them repeatedly.
Separate the definitions of idle and complete so that we can allow the async thread to sleep when only these unavailable files remain, without triggering the AssetRegistry's FileLoadedEvent.
#rnx

diff --git a/Engine/Source/Runtime/AssetRegistry/Private/AssetDataGatherer.cpp b/Engine/Source/Runtime/AssetRegistry/Private/AssetDataGatherer.cpp
--- a/Engine/Source/Runtime/AssetRegistry/Private/AssetDataGatherer.cpp
+++ b/Engine/Source/Runtime/AssetRegistry/Private/AssetDataGatherer.cpp
@@ -2706,7 +2706,7 @@ void FAssetDataGatherer::TickInternal(bool& bOutIsTickInterrupt)
 		IngestDiscoveryResults();
 
 		// Take a batch off of the work list. If we're waiting only on the first WaitBatchCount results don't take more than that
-		int32 NumToProcess = FMath::Min<int32>(BatchSize-LocalFilesToSearch.Num(), FilesToSearch->Num());
+		int32 NumToProcess = FMath::Min<int32>(BatchSize-LocalFilesToSearch.Num(), FilesToSearch->GetNumAvailable());
 		if (WaitBatchCount > 0)
 		{
 			bWaitBatchCountDecremented = true;
@@ -2720,10 +2720,10 @@ void FAssetDataGatherer::TickInternal(bool& bOutIsTickInterrupt)
 
 		FilesToSearch->PopFront(LocalFilesToSearch, NumToProcess);
 
-		// If all work is finished mark idle and exit
+		// If no work is available mark idle and exit
 		if (LocalFilesToSearch.Num() == 0 && bDiscoveryIsComplete)
 		{
-			WaitBatchCount = 0; // Clear WaitBatchCount in case it was set higher than FilesToSearch->Num()
+			WaitBatchCount = 0; // Clear WaitBatchCount in case it was set higher than FilesToSearch->GetNumAvailable().
 			bOutIsTickInterrupt = true;
 
 			if (!bFinishedInitialDiscovery)
@@ -2834,7 +2834,6 @@ void FAssetDataGatherer::TickInternal(bool& bOutIsTickInterrupt)
 
 	// Accumulate the results
 	bool bHasCancelation = false;
-	TArray<FGatheredPathData> LocalFilesToRetry;
 	for (FReadContext& ReadContext : ReadContexts)
 	{
 		if (ReadContext.bCanceled)
@@ -3059,12 +3058,14 @@ void FAssetDataGatherer::GetAndTrimSearchResults(bool& bOutIsSearching, TRingBuf
 	OutNumPathsToSearch = NumPathsToSearchAtLastSyncPoint;
 	OutIsDiscoveringFiles = !bDiscoveryIsComplete;
 
-	if (bIsIdle && !bIsComplete)
+	// Idle means no more work OR we are blocked on external events, but complete means no more work period.
+	bool bLocalIsComplete = bIsIdle && FilesToSearch->Num() == 0;
+	if (bLocalIsComplete && !bIsComplete)
 	{
 		bIsComplete = true;
 		Shrink();
 	}
-	bOutIsSearching = !bIsIdle;
+	bOutIsSearching = !bLocalIsComplete;
 }
 
 void FAssetDataGatherer::GetPackageResults(TRingBuffer<FAssetData*>& OutAssetResults, TRingBuffer<FPackageDependencyData>& OutDependencyResults)
@@ -3194,6 +3195,8 @@ void FAssetDataGatherer::WaitForIdle()
 		FGathererScopeLock ResultsScopeLock(&ResultsLock); // bIsIdle requires the lock
 		if (bIsIdle)
 		{
+			// We need to break out of WaitForIdle whenever it requires main thread action to proceed,
+			// so we check bIsIdle rather than whether we're complete
 			break;
 		}
 	}
@@ -3208,6 +3211,9 @@ bool FAssetDataGatherer::IsComplete() const
 void FAssetDataGatherer::SetInitialPluginsLoaded()
 {
 	bInitialPluginsLoaded = true;
+	FGathererScopeLock ResultsScopeLock(&ResultsLock);
+	SetIsIdle(false);
+	FilesToSearch->RetryLaterRetryFiles();
 }
 
 bool FAssetDataGatherer::IsGatheringDependencies() const
@@ -3769,6 +3775,16 @@ void FFilesToSearch::AddFileForLaterRetry(FGatheredPathData&& FilePath)
 	LaterRetryFiles.Add(FilePath);
 }
 
+void FFilesToSearch::RetryLaterRetryFiles()
+{
+	while (!LaterRetryFiles.IsEmpty())
+	{
+		FGatheredPathData FilePath = LaterRetryFiles.PopFrontValue();
+		FTreeNode& Node = Root.FindOrAddNode(FPathViews::GetPath(FilePath.LocalAbsPath));
+		Node.AddFile(MoveTemp(FilePath));
+	}
+}
+
 template <typename AllocatorType>
 void FFilesToSearch::PopFront(TArray<FGatheredPathData, AllocatorType>& Out, int32 NumToPop)
 {
@@ -3778,22 +3794,12 @@ void FFilesToSearch::PopFront(TArray<FGatheredPathData, AllocatorType>& Out, int
 		--NumToPop;
 	}
 	Root.PopFiles(Out, NumToPop);
-	while (NumToPop > 0 && !LaterRetryFiles.IsEmpty())
-	{
-		Out.Add(LaterRetryFiles.PopFrontValue());
-		--NumToPop;
-	}
 }
 
 void FFilesToSearch::PrioritizeDirectory(FStringView DirAbsPath, EPriority Priority)
 {
 	// We may need to prioritize a LaterRetryFile that is now loadable, so add them all into the Root
-	while (!LaterRetryFiles.IsEmpty())
-	{
-		FGatheredPathData FilePath = LaterRetryFiles.PopFrontValue();
-		FTreeNode& Node = Root.FindOrAddNode(FPathViews::GetPath(FilePath.LocalAbsPath));
-		Node.AddFile(MoveTemp(FilePath));
-	}
+	RetryLaterRetryFiles();
 
 	if (Priority > EPriority::Blocking)
 	{
@@ -3814,12 +3820,7 @@ void FFilesToSearch::PrioritizeDirectory(FStringView DirAbsPath, EPriority Prior
 void FFilesToSearch::PrioritizeFile(FStringView FileAbsPathExtOptional, EPriority Priority)
 {
 	// We may need to prioritize a LaterRetryFile that is now loadable, so add them all into the Root
-	while (!LaterRetryFiles.IsEmpty())
-	{
-		FGatheredPathData FilePath = LaterRetryFiles.PopFrontValue();
-		FTreeNode& Node = Root.FindOrAddNode(FPathViews::GetPath(FilePath.LocalAbsPath));
-		Node.AddFile(MoveTemp(FilePath));
-	}
+	RetryLaterRetryFiles();
 
 	if (Priority > EPriority::Blocking)
 	{
@@ -3865,6 +3866,11 @@ int32 FFilesToSearch::Num() const
 	return BlockingFiles.Num() + Root.NumFiles() + LaterRetryFiles.Num();
 }
 
+int32 FFilesToSearch::GetNumAvailable() const
+{
+	return BlockingFiles.Num() + Root.NumFiles();
+}
+
 uint32 FFilesToSearch::GetAllocatedSize() const
 {
 	uint32 Size = 0;
diff --git a/Engine/Source/Runtime/AssetRegistry/Private/AssetDataGatherer.h b/Engine/Source/Runtime/AssetRegistry/Private/AssetDataGatherer.h
--- a/Engine/Source/Runtime/AssetRegistry/Private/AssetDataGatherer.h
+++ b/Engine/Source/Runtime/AssetRegistry/Private/AssetDataGatherer.h
@@ -371,7 +371,7 @@ private:
 	bool bDiscoveryIsComplete;
 	/** Track whether this Gather has gone idle and a caller has read all search data from it. Read/writable only within ResultsLock. */
 	bool bIsComplete;
-	/** Track whether this Gatherer has gone idle. Read/writable only within ResultsLock. */
+	/** Track whether this Gatherer has gone idle, either because of no more work or because its blocked on external events. Read/writable only within ResultsLock. */
 	bool bIsIdle;
 	/** Track the first tick after idle to set up e.g. timing data. Read/writable only within ResultsLock. */
 	bool bFirstTickAfterIdle;
diff --git a/Engine/Source/Runtime/AssetRegistry/Private/AssetDataGathererPrivate.h b/Engine/Source/Runtime/AssetRegistry/Private/AssetDataGathererPrivate.h
--- a/Engine/Source/Runtime/AssetRegistry/Private/AssetDataGathererPrivate.h
+++ b/Engine/Source/Runtime/AssetRegistry/Private/AssetDataGathererPrivate.h
@@ -122,8 +122,13 @@ public:
 	void AddDirectory(FString&& DirAbsPath, TArray<FGatheredPathData>&& FilePaths);
 	/** Re-add a file from PopFront that we didn't have time to process. */
 	void AddFileAgainAfterTimeout(FGatheredPathData&& FilePath);
-	/** Re-add a file from PopFront that uses a not-yet-loaded class for later retry. */
+	/**
+	 * Re-add a file from PopFront that uses a not-yet-loaded class for later retry.
+	 * Later retry files are not returned from PopFront until RetryLaterRetryFiles is called.
+	 */
 	void AddFileForLaterRetry(FGatheredPathData&& FilePath);
+	/** Move all of the LaterRetry files into the main list so that they are once again visible to PopFront. */
+	void RetryLaterRetryFiles();
 
 	/** In coarse priority order, pop the given number of files out of the container. */
 	template <typename AllocatorType>
@@ -141,8 +146,10 @@ public:
 
 	/** Reduce memory used in buffers. */
 	void Shrink();
-	/** How many files are in the container. */
+	/** How many files are in the container, including unavailable files. */
 	int32 Num() const;
+	/** How many files are returnable from PopFront. */
+	int32 GetNumAvailable() const;
 	/** How much memory is used by the container, not counting sizeof(*this). */
 	uint32 GetAllocatedSize() const;
 private:
