ConvertAssetIdentifiersToPackageNames NlogN instead of N^2
use sort-unique pattern

diff --git a/Engine/Source/Runtime/AssetRegistry/Private/AssetRegistry.cpp b/Engine/Source/Runtime/AssetRegistry/Private/AssetRegistry.cpp
--- a/Engine/Source/Runtime/AssetRegistry/Private/AssetRegistry.cpp
+++ b/Engine/Source/Runtime/AssetRegistry/Private/AssetRegistry.cpp
@@ -2,6 +2,7 @@
 
 #include "AssetRegistry.h"
 
+#include "Algo/Unique.h"
 #include "AssetDataGatherer.h"
 #include "AssetRegistry/ARFilter.h"
 #include "AssetRegistryConsoleCommands.h"
@@ -1698,13 +1699,21 @@ bool UAssetRegistryImpl::GetDependencies(const FAssetIdentifier& AssetIdentifier
 
 static void ConvertAssetIdentifiersToPackageNames(const TArray<FAssetIdentifier>& AssetIdentifiers, TArray<FName>& OutPackageNames)
 {
+	// add all PackageNames :
+	OutPackageNames.Reserve(OutPackageNames.Num() + AssetIdentifiers.Num());
 	for (const FAssetIdentifier& AssetId : AssetIdentifiers)
 	{
 		if (AssetId.PackageName != NAME_None)
 		{
-			OutPackageNames.AddUnique(AssetId.PackageName);
+			OutPackageNames.Add(AssetId.PackageName);
 		}
 	}
+
+	// make unique ; sort in previous contents of OutPackageNames to unique against them too
+	OutPackageNames.Sort( FNameFastLess() );
+
+	int UniqueNum = Algo::Unique( OutPackageNames );
+	OutPackageNames.SetNum(UniqueNum,false);
 }
 
 bool UAssetRegistryImpl::GetDependencies(FName PackageName, TArray<FName>& OutDependencies, EAssetRegistryDependencyType::Type InDependencyType) const
