AssetRegistry - Early exit ScanPathsSynchronousInternal when it doesn't find any paths to sync.
#rnx

diff --git a/Engine/Source/Runtime/AssetRegistry/Private/AssetRegistry.cpp b/Engine/Source/Runtime/AssetRegistry/Private/AssetRegistry.cpp
--- a/Engine/Source/Runtime/AssetRegistry/Private/AssetRegistry.cpp
+++ b/Engine/Source/Runtime/AssetRegistry/Private/AssetRegistry.cpp
@@ -2551,6 +2551,10 @@ void UAssetRegistryImpl::ScanPathsSynchronousInternal(const TArray<FString>& InD
 		FWriteScopeLock InterfaceScopeLock(InterfaceLock);
 		GuardedData.ScanPathsSynchronous(Context);
 	}
+	if (Context.LocalPaths.IsEmpty())
+	{
+		return;
+	}
 
 #if WITH_EDITOR
 	ProcessLoadedAssetsToUpdateCache(EventContext, -1., Context.bIdle);
@@ -2565,7 +2569,7 @@ void UAssetRegistryImpl::ScanPathsSynchronousInternal(const TArray<FString>& InD
 	}
 	else
 	{
-		PathsString = FString::Printf(TEXT("'%s'"), Context.LocalPaths.IsEmpty() ? TEXT("<empty paths>") : *Context.LocalPaths[0]);
+		PathsString = FString::Printf(TEXT("'%s'"), *Context.LocalPaths[0]);
 	}
 
 	UE_LOG(LogAssetRegistry, Verbose, TEXT("ScanPathsSynchronous completed scanning %s to find %d assets in %0.4f seconds"), *PathsString,
