FAssetRegistryImpl::TickGatherer: Fix premature declaration of idle due to reading NumPending before it is written.
#rnx

diff --git a/Engine/Source/Runtime/AssetRegistry/Private/AssetRegistry.cpp b/Engine/Source/Runtime/AssetRegistry/Private/AssetRegistry.cpp
--- a/Engine/Source/Runtime/AssetRegistry/Private/AssetRegistry.cpp
+++ b/Engine/Source/Runtime/AssetRegistry/Private/AssetRegistry.cpp
@@ -2867,11 +2867,10 @@ void FAssetRegistryImpl::TickGatherer(Impl::FEventContext& EventContext, const d
 		UE_LOG(LogAssetRegistry, Verbose, TEXT("### Background search completed in %0.4f seconds"), SearchTimes[SearchTimeIdx]);
 	}
 	const bool bHadAssetsToProcess = BackgroundAssetResults.Num() > 0 || BackgroundDependencyResults.Num() > 0;
-	int32 NumPending = 0;
-	auto UpdateStatus = [&NumPending, &bIsSearching, &bHadAssetsToProcess, &bIsDiscoveringFiles, &EventContext, &bOutIdle, &NumFilesToSearch, &NumPathsToSearch, this](bool bInIdle)
+	auto UpdateStatus = [&bIsSearching, &bHadAssetsToProcess, &bIsDiscoveringFiles, &EventContext, &bOutIdle, &NumFilesToSearch, &NumPathsToSearch, this](bool bInterrupted)
 	{
 		// Compute total pending, plus highest pending for this run so we can show a good progress bar
-		NumPending = NumFilesToSearch + NumPathsToSearch + BackgroundPathResults.Num() + BackgroundAssetResults.Num() + BackgroundDependencyResults.Num() + BackgroundCookedPackageNamesWithoutAssetDataResults.Num();
+		int32 NumPending = NumFilesToSearch + NumPathsToSearch + BackgroundPathResults.Num() + BackgroundAssetResults.Num() + BackgroundDependencyResults.Num() + BackgroundCookedPackageNamesWithoutAssetDataResults.Num();
 		HighestPending = FMath::Max(this->HighestPending, NumPending);
 
 		// Notify the status change
@@ -2885,8 +2884,8 @@ void FAssetRegistryImpl::TickGatherer(Impl::FEventContext& EventContext, const d
 			);
 		}
 
-		bOutIdle = bInIdle;
-		this->bGatherIdle = bInIdle;
+		bOutIdle = !bInterrupted && !bIsSearching && NumPending == 0;
+		this->bGatherIdle = bOutIdle;
 	};
 
 
@@ -2930,16 +2929,15 @@ void FAssetRegistryImpl::TickGatherer(Impl::FEventContext& EventContext, const d
 		CookedPackageNamesWithoutAssetDataGathered(EventContext, TickStartTime, BackgroundCookedPackageNamesWithoutAssetDataResults, bOutInterrupted);
 		if (bOutInterrupted)
 		{
-			UpdateStatus(false /* bInIdle */);
+			UpdateStatus(true /* bInterrupted */);
 			return;
 		}
 	}
 
 	// If completing an initial search, refresh the content browser
-	bool bIsIdle = !bIsSearching && NumPending == 0;
-	UpdateStatus(bIsIdle);
+	UpdateStatus(false /* bInterrupted */);
 
-	if (bIsIdle)
+	if (bOutIdle)
 	{
 		HighestPending = 0;
 
