Fix whitespace from merge.
#rnx

diff --git a/Engine/Source/Runtime/AssetRegistry/Private/AssetRegistryState.cpp b/Engine/Source/Runtime/AssetRegistry/Private/AssetRegistryState.cpp
--- a/Engine/Source/Runtime/AssetRegistry/Private/AssetRegistryState.cpp
+++ b/Engine/Source/Runtime/AssetRegistry/Private/AssetRegistryState.cpp
@@ -334,50 +334,46 @@ void FAssetRegistryState::InitializeFromExisting(const TMap<FName, FAssetData*>&
 			continue;
 		}
 
-
 		FAssetData* ExistingData = nullptr;
-
 		if (InInitializationMode != EInitializationMode::Rebuild) // minor optimization to avoid lookup in rebuild mode
 		{
 			ExistingData = CachedAssetsByObjectPath.FindRef(Pair.Key);
 		}
-		if ((InInitializationMode == EInitializationMode::OnlyUpdateExisting) && 
-			(ExistingData == nullptr) )
-			{
-				continue;
-			}
-
-			// Filter asset registry tags now
-			const FAssetData& AssetData = *Pair.Value;
+		if (InInitializationMode == EInitializationMode::OnlyUpdateExisting && ExistingData == nullptr)
+		{
+			continue;
+		}
 
-			FAssetDataTagMap LocalTagsAndValues;
-			FAssetRegistryState::FilterTags(AssetData.TagsAndValues, LocalTagsAndValues, Options.CookFilterlistTagsByClass.Find(AssetData.AssetClass), Options);
+		// Filter asset registry tags now
+		const FAssetData& AssetData = *Pair.Value;
 
+		FAssetDataTagMap LocalTagsAndValues;
+		FAssetRegistryState::FilterTags(AssetData.TagsAndValues, LocalTagsAndValues, Options.CookFilterlistTagsByClass.Find(AssetData.AssetClass), Options);
 		
 		// in append or onlyupdateexisting we may have some existing data
 		// in rebuild mode existing data should never be found
-				if (ExistingData)
-				{
-					// Bundle tags might have changed even if other tags haven't
-					ExistingData->TaggedAssetBundles = AssetData.TaggedAssetBundles;
+		if (ExistingData)
+		{
+			// Bundle tags might have changed even if other tags haven't
+			ExistingData->TaggedAssetBundles = AssetData.TaggedAssetBundles;
 
-					// If tags have changed we need to update CachedAssetsByTag
-					if (LocalTagsAndValues != ExistingData->TagsAndValues)
-					{
-						FAssetData TempData = *ExistingData;
-						TempData.TagsAndValues = FAssetDataTagMapSharedView(MoveTemp(LocalTagsAndValues));
-						UpdateAssetData(ExistingData, TempData);
-					}
-				}
-			else
+			// If tags have changed we need to update CachedAssetsByTag
+			if (LocalTagsAndValues != ExistingData->TagsAndValues)
 			{
-				FAssetData* NewData = new FAssetData(AssetData.PackageName, AssetData.PackagePath, AssetData.AssetName,
-					AssetData.AssetClass, LocalTagsAndValues, AssetData.ChunkIDs, AssetData.PackageFlags);
-				NewData->TaggedAssetBundles = AssetData.TaggedAssetBundles;
-
-				AddAssetData(NewData);
+				FAssetData TempData = *ExistingData;
+				TempData.TagsAndValues = FAssetDataTagMapSharedView(MoveTemp(LocalTagsAndValues));
+				UpdateAssetData(ExistingData, TempData);
 			}
 		}
+		else
+		{
+			FAssetData* NewData = new FAssetData(AssetData.PackageName, AssetData.PackagePath, AssetData.AssetName,
+				AssetData.AssetClass, LocalTagsAndValues, AssetData.ChunkIDs, AssetData.PackageFlags);
+			NewData->TaggedAssetBundles = AssetData.TaggedAssetBundles;
+
+			AddAssetData(NewData);
+		}
+	}
 
 	TSet<FAssetIdentifier> ScriptPackages;
 
