Asset registry type truncation warning fixes
#swarm https://p4-swarm.epicgames.net/reviews/17578628

diff --git a/Engine/Source/Runtime/AssetRegistry/AssetRegistry.Build.cs b/Engine/Source/Runtime/AssetRegistry/AssetRegistry.Build.cs
--- a/Engine/Source/Runtime/AssetRegistry/AssetRegistry.Build.cs
+++ b/Engine/Source/Runtime/AssetRegistry/AssetRegistry.Build.cs
@@ -20,5 +20,7 @@ public class AssetRegistry : ModuleRules
 			PrivateIncludePathModuleNames.AddRange(new string[] { "DirectoryWatcher" });
 			DynamicallyLoadedModuleNames.AddRange(new string[] { "DirectoryWatcher" });
 		}
+
+		UnsafeTypeCastWarningLevel = WarningLevel.Warning;
 	}
 }
diff --git a/Engine/Source/Runtime/AssetRegistry/Private/AssetDataGatherer.cpp b/Engine/Source/Runtime/AssetRegistry/Private/AssetDataGatherer.cpp
--- a/Engine/Source/Runtime/AssetRegistry/Private/AssetDataGatherer.cpp
+++ b/Engine/Source/Runtime/AssetRegistry/Private/AssetDataGatherer.cpp
@@ -86,7 +86,7 @@ void FDiscoveredPathData::Assign(FStringView InLocalAbsPath, FStringView InLongP
 	PackageTimestamp = InPackageTimestamp;
 }
 
-uint32 FDiscoveredPathData::GetAllocatedSize() const
+SIZE_T FDiscoveredPathData::GetAllocatedSize() const
 {
 	return LocalAbsPath.GetAllocatedSize() + LongPackageName.GetAllocatedSize() + RelPath.GetAllocatedSize();
 }
@@ -122,7 +122,7 @@ void FGatheredPathData::Assign(const FDiscoveredPathData& DiscoveredData)
 	Assign(DiscoveredData.LocalAbsPath, DiscoveredData.LongPackageName, DiscoveredData.PackageTimestamp);
 }
 
-uint32 FGatheredPathData::GetAllocatedSize() const
+SIZE_T FGatheredPathData::GetAllocatedSize() const
 {
 	return LocalAbsPath.GetAllocatedSize() + LongPackageName.GetAllocatedSize();
 }
@@ -331,9 +331,9 @@ bool FScanDir::IsComplete() const
 	return bIsComplete;
 }
 
-uint32 FScanDir::GetAllocatedSize() const
+SIZE_T FScanDir::GetAllocatedSize() const
 {
-	uint32 Result = 0;
+	SIZE_T Result = 0;
 	Result += SubDirs.GetAllocatedSize();
 	for (const TRefCountPtr<FScanDir>& Value : SubDirs)
 	{
@@ -629,7 +629,7 @@ void FScanDir::OnChildPriorityChanged(EPriority InPriority, int32 Delta)
 	}
 	else
 	{
-		PriorityRefCount += Delta;
+		PriorityRefCount = (uint8)(PriorityRefCount + Delta);
 	}
 	UpdateAccumulatedPriority();
 }
@@ -1003,9 +1003,9 @@ FScanDir* FMountDir::GetControllingDir(FStringView InLocalAbsPath, bool bIsDirec
 		OutData, OutRelPath);
 }
 
-uint32 FMountDir::GetAllocatedSize() const
+SIZE_T FMountDir::GetAllocatedSize() const
 {
-	uint32 Result = sizeof(*Root);
+	SIZE_T Result = sizeof(*Root);
 	Result += Root->GetAllocatedSize();
 	Result += ChildMountPaths.GetAllocatedSize();
 	for (const FString& Value : ChildMountPaths)
@@ -1907,11 +1907,11 @@ bool FAssetDataDiscovery::IsMonitored(FStringView LocalAbsPath) const
 	return MountDir && MountDir->IsMonitored(LocalAbsPath);
 }
 
-uint32 FAssetDataDiscovery::GetAllocatedSize() const
+SIZE_T FAssetDataDiscovery::GetAllocatedSize() const
 {
 	auto GetArrayRecursiveAllocatedSize = [](auto Container)
 	{
-		int32 Result = Container.GetAllocatedSize();
+		SIZE_T Result = Container.GetAllocatedSize();
 		for (const auto& Value : Container)
 		{
 			Result += Value.GetAllocatedSize();
@@ -1926,7 +1926,7 @@ uint32 FAssetDataDiscovery::GetAllocatedSize() const
 	FGathererScopeLock TreeScopeLock(&TreeLock);
 	FGathererScopeLock ResultsScopeLock(&ResultsLock);
 
-	uint32 Result = 0;
+	SIZE_T Result = 0;
 	Result += GetArrayRecursiveAllocatedSize(BlacklistLongPackageNames);
 	Result += GetArrayRecursiveAllocatedSize(BlacklistMountRelativePaths);
 	Result += GetArrayRecursiveAllocatedSize(DirLongPackageNamesToNotReport);
@@ -1951,7 +1951,7 @@ uint32 FAssetDataDiscovery::GetAllocatedSize() const
 	return Result;
 }
 
-uint32 FAssetDataDiscovery::FDirectoryResult::GetAllocatedSize() const
+SIZE_T FAssetDataDiscovery::FDirectoryResult::GetAllocatedSize() const
 {
 	return DirAbsPath.GetAllocatedSize() + Files.GetAllocatedSize();
 }
@@ -3434,7 +3434,7 @@ void FAssetDataGatherer::SerializeCacheLoad(FAssetRegistryReader& Ar)
 	Ar << LocalNumAssets;
 
 	const int32 MinAssetEntrySize = sizeof(int32);
-	int32 MaxPossibleNumAssets = (Ar.TotalSize() - Ar.Tell()) / MinAssetEntrySize;
+	const int64 MaxPossibleNumAssets = (Ar.TotalSize() - Ar.Tell()) / MinAssetEntrySize;
 	if (Ar.IsError() || LocalNumAssets < 0 || MaxPossibleNumAssets < LocalNumAssets)
 	{
 		Ar.SetError();
@@ -3479,12 +3479,12 @@ void FAssetDataGatherer::SerializeCacheLoad(FAssetRegistryReader& Ar)
 	UE_LOG(LogAssetRegistry, Verbose, TEXT("Asset data gatherer serialized in %0.6f seconds"), FPlatformTime::Seconds() - SerializeStartTime);
 }
 
-uint32 FAssetDataGatherer::GetAllocatedSize() const
+SIZE_T FAssetDataGatherer::GetAllocatedSize() const
 {
 	using namespace UE::AssetDataGather::Private;
 	auto GetArrayRecursiveAllocatedSize = [](auto Container)
 	{
-		int32 Result = Container.GetAllocatedSize();
+		SIZE_T Result = Container.GetAllocatedSize();
 		for (const auto& Value : Container)
 		{
 			Result += Value.GetAllocatedSize();
@@ -3492,7 +3492,7 @@ uint32 FAssetDataGatherer::GetAllocatedSize() const
 		return Result;
 	};
 
-	uint32 Result = 0;
+	SIZE_T Result = 0;
 	if (Thread)
 	{
 		// TODO: Add size of Thread->GetAllocatedSize()
@@ -3871,9 +3871,9 @@ int32 FFilesToSearch::GetNumAvailable() const
 	return BlockingFiles.Num() + Root.NumFiles();
 }
 
-uint32 FFilesToSearch::GetAllocatedSize() const
+SIZE_T FFilesToSearch::GetAllocatedSize() const
 {
-	uint32 Size = 0;
+	SIZE_T Size = 0;
 	Size += BlockingFiles.GetAllocatedSize();
 	for (const FGatheredPathData& PathData : BlockingFiles)
 	{
@@ -4067,9 +4067,9 @@ void FFilesToSearch::FTreeNode::Shrink()
 	}
 }
 
-uint32 FFilesToSearch::FTreeNode::GetAllocatedSize() const
+SIZE_T FFilesToSearch::FTreeNode::GetAllocatedSize() const
 {
-	uint32 Size = 0;
+	SIZE_T Size = 0;
 	Size = Files.GetAllocatedSize();
 	for (const FGatheredPathData& File : Files)
 	{
diff --git a/Engine/Source/Runtime/AssetRegistry/Private/AssetDataGatherer.h b/Engine/Source/Runtime/AssetRegistry/Private/AssetDataGatherer.h
--- a/Engine/Source/Runtime/AssetRegistry/Private/AssetDataGatherer.h
+++ b/Engine/Source/Runtime/AssetRegistry/Private/AssetDataGatherer.h
@@ -123,7 +123,7 @@ public:
 	/** Attempt to read the cache file at the given LocalPath, and store all of its results in the in-memory cache. */
 	void LoadCacheFile(FStringView CacheFilename);
 	/** Return the total amount of heap memory used by the gatherer (including not-yet-claimed search results). Used for performance metrics. */
-	uint32 GetAllocatedSize() const;
+	SIZE_T GetAllocatedSize() const;
 
 
 	// Configuring mount points (possibly while tick is running)
diff --git a/Engine/Source/Runtime/AssetRegistry/Private/AssetDataGathererPrivate.h b/Engine/Source/Runtime/AssetRegistry/Private/AssetDataGathererPrivate.h
--- a/Engine/Source/Runtime/AssetRegistry/Private/AssetDataGathererPrivate.h
+++ b/Engine/Source/Runtime/AssetRegistry/Private/AssetDataGathererPrivate.h
@@ -85,7 +85,7 @@ struct FDiscoveredPathData
 	void Assign(FStringView InLocalAbsPath, FStringView InLongPackageName, FStringView InRelPath);
 
 	/** Return the total amount of heap memory used by the gatherer (including not-yet-claimed search results). */
-	uint32 GetAllocatedSize() const;
+	SIZE_T GetAllocatedSize() const;
 };
 
 /** Information needed about a discovered asset file or path that is needed by the Gatherer */
@@ -109,7 +109,7 @@ struct FGatheredPathData
 	 * Return the total amount of heap memory used by the gatherer (including not-yet-claimed search results).
 	 * Used for performance metrics.
 	 */
-	uint32 GetAllocatedSize() const;
+	SIZE_T GetAllocatedSize() const;
 };
 
 /** A container to efficiently receive files from discovery and search the container to raise priority of a path */
@@ -151,7 +151,7 @@ public:
 	/** How many files are returnable from PopFront. */
 	int32 GetNumAvailable() const;
 	/** How much memory is used by the container, not counting sizeof(*this). */
-	uint32 GetAllocatedSize() const;
+	SIZE_T GetAllocatedSize() const;
 private:
 	/** A directory-search-tree structure; each node has a list of direct files and subdirectories. */
 	struct FTreeNode
@@ -195,7 +195,7 @@ private:
 		/** Reduce memory used in buffers. */
 		void Shrink();
 		/** How much memory is used by *this, not counting sizeof(*this). */
-		uint32 GetAllocatedSize() const;
+		SIZE_T GetAllocatedSize() const;
 		/** Number of files in *this and its child nodes. */
 		int32 NumFiles() const;
 	private:
@@ -320,7 +320,7 @@ public:
 	bool IsComplete() const;
 
 	/** Return the memory used by the tree under this ScanDir; excludes sizeof(*this). */
-	uint32 GetAllocatedSize() const;
+	SIZE_T GetAllocatedSize() const;
 	/**
 	 * Find the Direct parent of InRelPath, or a fallback. Will return null only if !bIsDirectory and InRelPath is empty.
 	 * The fallback is returned if InRelPath has already completed and been deleted, or if InRelPath is not monitored.
@@ -464,7 +464,7 @@ public:
 	/** Find the direct parent of InRelPath, or the lowest fallback. See FScanDir::GetControllingDir. */
 	FScanDir* GetControllingDir(FStringView LocalAbsPath, bool bIsDirectory, FScanDir::FInherited& OutParentData, FString& OutRelPath);
 	/** Return the memory used by the tree under this MountDir, except that sizeof(*this) is excluded. */
-	uint32 GetAllocatedSize() const;
+	SIZE_T GetAllocatedSize() const;
 
 	/** Report whether this MountDir is complete: all ScanDirs under it either have scanned or should not scan. */
 	bool IsComplete() const;
@@ -597,7 +597,7 @@ public:
 	/** Return whether the given path should or has been scanned because it is whitelisted and not blacklisted. */
 	bool IsMonitored(FStringView LocalAbsPath) const;
 	/** Return the memory used by *this. sizeof(*this) is not included. */
-	uint32 GetAllocatedSize() const;
+	SIZE_T GetAllocatedSize() const;
 
 	// Events and setting of properties (possibly while tick is running)
 	/**
@@ -745,7 +745,7 @@ private:
 		FDirectoryResult(FStringView InDirAbsPath, TConstArrayView<FDiscoveredPathData> InFiles);
 		FString DirAbsPath;
 		TArray<FGatheredPathData> Files;
-		uint32 GetAllocatedSize() const;
+		SIZE_T GetAllocatedSize() const;
 	};
 	TArray<FDirectoryResult> DiscoveredFiles;
 	TArray<FGatheredPathData> DiscoveredSingleFiles;
diff --git a/Engine/Source/Runtime/AssetRegistry/Private/AssetRegistry.cpp b/Engine/Source/Runtime/AssetRegistry/Private/AssetRegistry.cpp
--- a/Engine/Source/Runtime/AssetRegistry/Private/AssetRegistry.cpp
+++ b/Engine/Source/Runtime/AssetRegistry/Private/AssetRegistry.cpp
@@ -3101,11 +3101,11 @@ void FAssetRegistryImpl::CachePathsFromState(Impl::FEventContext& EventContext,
 
 }
 
-uint32 UAssetRegistryImpl::GetAllocatedSize(bool bLogDetailed) const
+SIZE_T UAssetRegistryImpl::GetAllocatedSize(bool bLogDetailed) const
 {
-	uint32 StateSize = 0;
-	uint32 StaticSize = 0;
-	uint32 SearchSize = 0;
+	SIZE_T StateSize = 0;
+	SIZE_T StaticSize = 0;
+	SIZE_T SearchSize = 0;
 	{
 		FReadScopeLock InterfaceScopeLock(InterfaceLock);
 		GuardedData.GetAllocatedSize(bLogDetailed, StateSize, StaticSize, SearchSize);
@@ -3117,8 +3117,8 @@ uint32 UAssetRegistryImpl::GetAllocatedSize(bool bLogDetailed) const
 
 	if (bLogDetailed)
 	{
-		UE_LOG(LogAssetRegistry, Log, TEXT("AssetRegistry Static Size: %dk"), StaticSize / 1024);
-		UE_LOG(LogAssetRegistry, Log, TEXT("AssetRegistry Search Size: %dk"), SearchSize / 1024);
+		UE_LOG(LogAssetRegistry, Log, TEXT("AssetRegistry Static Size: %" SIZE_T_FMT "k"), StaticSize / 1024);
+		UE_LOG(LogAssetRegistry, Log, TEXT("AssetRegistry Search Size: %" SIZE_T_FMT "k"), SearchSize / 1024);
 	}
 
 	return StateSize + StaticSize + SearchSize;
@@ -3127,7 +3127,7 @@ uint32 UAssetRegistryImpl::GetAllocatedSize(bool bLogDetailed) const
 namespace UE::AssetRegistry
 {
 
-void FAssetRegistryImpl::GetAllocatedSize(bool bLogDetailed, uint32& StateSize, uint32& StaticSize, uint32& SearchSize) const
+void FAssetRegistryImpl::GetAllocatedSize(bool bLogDetailed, SIZE_T& StateSize, SIZE_T& StaticSize, SIZE_T& SearchSize) const
 {
 	StateSize = State.GetAllocatedSize(bLogDetailed);
 
diff --git a/Engine/Source/Runtime/AssetRegistry/Private/AssetRegistry.h b/Engine/Source/Runtime/AssetRegistry/Private/AssetRegistry.h
--- a/Engine/Source/Runtime/AssetRegistry/Private/AssetRegistry.h
+++ b/Engine/Source/Runtime/AssetRegistry/Private/AssetRegistry.h
@@ -93,7 +93,7 @@ public:
 	virtual void ScanModifiedAssetFiles(const TArray<FString>& InFilePaths) override;
 	virtual void Serialize(FArchive& Ar) override;
 	virtual void AppendState(const FAssetRegistryState& InState) override;
-	virtual uint32 GetAllocatedSize(bool bLogDetailed = false) const override;
+	virtual SIZE_T GetAllocatedSize(bool bLogDetailed = false) const override;
 	virtual void LoadPackageRegistryData(FArchive& Ar, FLoadPackageRegistryData& InOutData) const override;
 	virtual void LoadPackageRegistryData(const FString& PackageFilename, FLoadPackageRegistryData& InOutData) const override;
 	virtual void InitializeTemporaryAssetRegistryState(FAssetRegistryState& OutState, const FAssetRegistrySerializationOptions& Options, bool bRefreshExisting = false, const TMap<FName, FAssetData*>& OverrideData = TMap<FName, FAssetData*>()) const override;
diff --git a/Engine/Source/Runtime/AssetRegistry/Private/AssetRegistryArchive.cpp b/Engine/Source/Runtime/AssetRegistry/Private/AssetRegistryArchive.cpp
--- a/Engine/Source/Runtime/AssetRegistry/Private/AssetRegistryArchive.cpp
+++ b/Engine/Source/Runtime/AssetRegistry/Private/AssetRegistryArchive.cpp
@@ -80,7 +80,7 @@ static TArray<FNameEntryId> FlattenIndex(const TMap<FNameEntryId, uint32>& Names
 FAssetRegistryWriter::~FAssetRegistryWriter()
 {
 	// Save store data and collect FNames
-	int32 BodySize = MemWriter.TotalSize();
+	int64 BodySize = MemWriter.TotalSize();
 	SaveStore(Tags.Finalize(), *this);
 
 	// Save in load-friendly order - names, store then body / tag maps
diff --git a/Engine/Source/Runtime/AssetRegistry/Private/AssetRegistryImpl.h b/Engine/Source/Runtime/AssetRegistry/Private/AssetRegistryImpl.h
--- a/Engine/Source/Runtime/AssetRegistry/Private/AssetRegistryImpl.h
+++ b/Engine/Source/Runtime/AssetRegistry/Private/AssetRegistryImpl.h
@@ -88,7 +88,7 @@ public:
 	void ScanModifiedAssetFiles(Impl::FEventContext& EventContext, const TArray<FString>& InFilePaths);
 	void Serialize(FArchive& Ar, Impl::FEventContext& EventContext);
 	void AppendState(Impl::FEventContext& EventContext, const FAssetRegistryState& InState);
-	void GetAllocatedSize(bool bLogDetailed, uint32& StateSize, uint32& StaticSize, uint32& SearchSize) const;
+	void GetAllocatedSize(bool bLogDetailed, SIZE_T& StateSize, SIZE_T& StaticSize, SIZE_T& SearchSize) const;
 	bool IsLoadingAssets() const;
 	void SetManageReferences(const TMultiMap<FAssetIdentifier, FAssetIdentifier>& ManagerMap,
 		bool bClearExisting, UE::AssetRegistry::EDependencyCategory RecurseType,
diff --git a/Engine/Source/Runtime/AssetRegistry/Private/AssetRegistryState.cpp b/Engine/Source/Runtime/AssetRegistry/Private/AssetRegistryState.cpp
--- a/Engine/Source/Runtime/AssetRegistry/Private/AssetRegistryState.cpp
+++ b/Engine/Source/Runtime/AssetRegistry/Private/AssetRegistryState.cpp
@@ -1307,9 +1307,9 @@ void FAssetRegistryState::LoadDependencies_BeforeFlags(FArchive& Ar, bool bSeria
 	}
 }
 
-uint32 FAssetRegistryState::GetAllocatedSize(bool bLogDetailed) const
+SIZE_T FAssetRegistryState::GetAllocatedSize(bool bLogDetailed) const
 {
-	uint32 MapMemory = CachedAssetsByObjectPath.GetAllocatedSize();
+	SIZE_T MapMemory = CachedAssetsByObjectPath.GetAllocatedSize();
 	MapMemory += CachedAssetsByPackageName.GetAllocatedSize();
 	MapMemory += CachedAssetsByPath.GetAllocatedSize();
 	MapMemory += CachedAssetsByClass.GetAllocatedSize();
@@ -1320,7 +1320,7 @@ uint32 FAssetRegistryState::GetAllocatedSize(bool bLogDetailed) const
 	MapMemory += PreallocatedDependsNodeDataBuffers.GetAllocatedSize();
 	MapMemory += PreallocatedPackageDataBuffers.GetAllocatedSize();
 
-	uint32 MapArrayMemory = 0;
+	SIZE_T MapArrayMemory = 0;
 	auto SubArray = 
 		[&MapArrayMemory](const auto& A)
 	{
@@ -1336,10 +1336,10 @@ uint32 FAssetRegistryState::GetAllocatedSize(bool bLogDetailed) const
 
 	if (bLogDetailed)
 	{
-		UE_LOG(LogAssetRegistry, Log, TEXT("Index Size: %dk"), MapMemory / 1024);
+		UE_LOG(LogAssetRegistry, Log, TEXT("Index Size: %" SIZE_T_FMT "k"), MapMemory / 1024);
 	}
 
-	uint32 AssetDataSize = 0;
+	SIZE_T AssetDataSize = 0;
 	FAssetDataTagMapSharedView::FMemoryCounter TagMemoryUsage;
 
 	for (const TPair<FName, FAssetData*>& AssetDataPair : CachedAssetsByObjectPath)
@@ -1354,13 +1354,13 @@ uint32 FAssetRegistryState::GetAllocatedSize(bool bLogDetailed) const
 	if (bLogDetailed)
 	{
 		UE_LOG(LogAssetRegistry, Log, TEXT("AssetData Count: %d"), CachedAssetsByObjectPath.Num());
-		UE_LOG(LogAssetRegistry, Log, TEXT("AssetData Static Size: %dk"), AssetDataSize / 1024);
-		UE_LOG(LogAssetRegistry, Log, TEXT("Loose Tags: %dk"), TagMemoryUsage.GetLooseSize() / 1024);
-		UE_LOG(LogAssetRegistry, Log, TEXT("Fixed Tags: %dk"), TagMemoryUsage.GetFixedSize() / 1024);
-		UE_LOG(LogAssetRegistry, Log, TEXT("TArray<FAssetData*>: %dk"), MapArrayMemory / 1024);
+		UE_LOG(LogAssetRegistry, Log, TEXT("AssetData Static Size: %" SIZE_T_FMT "k"), AssetDataSize / 1024);
+		UE_LOG(LogAssetRegistry, Log, TEXT("Loose Tags: %" SIZE_T_FMT "k"), TagMemoryUsage.GetLooseSize() / 1024);
+		UE_LOG(LogAssetRegistry, Log, TEXT("Fixed Tags: %" SIZE_T_FMT "k"), TagMemoryUsage.GetFixedSize() / 1024);
+		UE_LOG(LogAssetRegistry, Log, TEXT("TArray<FAssetData*>: %" SIZE_T_FMT "k"), MapArrayMemory / 1024);
 	}
 
-	uint32 DependNodesSize = 0, DependenciesSize = 0;
+	SIZE_T DependNodesSize = 0, DependenciesSize = 0;
 
 	for (const TPair<FAssetIdentifier, FDependsNode*>& DependsNodePair : CachedDependsNodes)
 	{
@@ -1373,19 +1373,19 @@ uint32 FAssetRegistryState::GetAllocatedSize(bool bLogDetailed) const
 	if (bLogDetailed)
 	{
 		UE_LOG(LogAssetRegistry, Log, TEXT("Dependency Node Count: %d"), CachedDependsNodes.Num());
-		UE_LOG(LogAssetRegistry, Log, TEXT("Dependency Node Static Size: %dk"), DependNodesSize / 1024);
-		UE_LOG(LogAssetRegistry, Log, TEXT("Dependency Arrays Size: %dk"), DependenciesSize / 1024);
+		UE_LOG(LogAssetRegistry, Log, TEXT("Dependency Node Static Size: %" SIZE_T_FMT "k"), DependNodesSize / 1024);
+		UE_LOG(LogAssetRegistry, Log, TEXT("Dependency Arrays Size: %" SIZE_T_FMT "k"), DependenciesSize / 1024);
 	}
 
-	uint32 PackageDataSize = CachedPackageData.Num() * sizeof(FAssetPackageData);
+	SIZE_T PackageDataSize = CachedPackageData.Num() * sizeof(FAssetPackageData);
 
-	uint32 TotalBytes = MapMemory + AssetDataSize + TagMemoryUsage.GetFixedSize()  + TagMemoryUsage.GetLooseSize() + DependNodesSize + DependenciesSize + PackageDataSize + MapArrayMemory;
+	SIZE_T TotalBytes = MapMemory + AssetDataSize + TagMemoryUsage.GetFixedSize()  + TagMemoryUsage.GetLooseSize() + DependNodesSize + DependenciesSize + PackageDataSize + MapArrayMemory;
 
 	if (bLogDetailed)
 	{
 		UE_LOG(LogAssetRegistry, Log, TEXT("PackageData Count: %d"), CachedPackageData.Num());
-		UE_LOG(LogAssetRegistry, Log, TEXT("PackageData Static Size: %dk"), PackageDataSize / 1024);
-		UE_LOG(LogAssetRegistry, Log, TEXT("Total State Size: %dk"), TotalBytes / 1024);
+		UE_LOG(LogAssetRegistry, Log, TEXT("PackageData Static Size: %" SIZE_T_FMT "k"), PackageDataSize / 1024);
+		UE_LOG(LogAssetRegistry, Log, TEXT("Total State Size: %" SIZE_T_FMT "k"), TotalBytes / 1024);
 	}
 
 	return TotalBytes;
@@ -2038,7 +2038,7 @@ void FAssetRegistryState::Dump(const TArray<FString>& Arguments, TArray<FString>
 			}
 			if (WinningLineEnd != BufferEnd)
 			{
-				PageEndIndex = WinningLineEnd - PageBuffer.GetData();
+				PageEndIndex = UE_PTRDIFF_TO_INT32(WinningLineEnd - PageBuffer.GetData());
 				NumOverflowLines = WinningSearchIndex;
 			}
 		}
diff --git a/Engine/Source/Runtime/AssetRegistry/Private/DependsNode.h b/Engine/Source/Runtime/AssetRegistry/Private/DependsNode.h
--- a/Engine/Source/Runtime/AssetRegistry/Private/DependsNode.h
+++ b/Engine/Source/Runtime/AssetRegistry/Private/DependsNode.h
@@ -73,7 +73,7 @@ public:
 	/** Returns number of connections this node has, both references and dependencies */
 	int32 GetConnectionCount() const;
 	/** Returns amount of memory used by the arrays */
-	uint32 GetAllocatedSize(void) const
+	SIZE_T GetAllocatedSize(void) const
 	{
 		return PackageDependencies.GetAllocatedSize() + PackageFlags.GetAllocatedSize() + NameDependencies.GetAllocatedSize() + ManageDependencies.GetAllocatedSize() + ManageFlags.GetAllocatedSize() + Referencers.GetAllocatedSize();
 	}
diff --git a/Engine/Source/Runtime/AssetRegistry/Private/DependsNodeTest.cpp b/Engine/Source/Runtime/AssetRegistry/Private/DependsNodeTest.cpp
--- a/Engine/Source/Runtime/AssetRegistry/Private/DependsNodeTest.cpp
+++ b/Engine/Source/Runtime/AssetRegistry/Private/DependsNodeTest.cpp
@@ -970,7 +970,7 @@ bool FDependsNodeTest::RunTest(const FString& Parameters)
 							LinkHaver.AddDependency(LinkNode, EDependencyCategory::SearchableName, EDependencyProperty::None);
 							LinkHaver.AddDependency(LinkNode, EDependencyCategory::Manage, ManageProperties[NodeIndex]);
 						};
-						auto GetNodeIndex = [&LinkNodes](const FDependsNode* Existing) { return Existing - &LinkNodes[0]; };
+						auto GetNodeIndex = [&LinkNodes](const FDependsNode* Existing) { return UE_PTRDIFF_TO_INT32(Existing - &LinkNodes[0]); };
 
 						if (SortTestType != ESortTestType::AlwaysSorted)
 						{
diff --git a/Engine/Source/Runtime/AssetRegistry/Private/DiskCachedAssetData.h b/Engine/Source/Runtime/AssetRegistry/Private/DiskCachedAssetData.h
--- a/Engine/Source/Runtime/AssetRegistry/Private/DiskCachedAssetData.h
+++ b/Engine/Source/Runtime/AssetRegistry/Private/DiskCachedAssetData.h
@@ -49,9 +49,9 @@ public:
 	}
 
 	/** Returns the amount of memory allocated by this container, not including sizeof(*this). */
-	uint32 GetAllocatedSize() const
+	SIZE_T GetAllocatedSize() const
 	{
-		uint32 Result = 0;
+		SIZE_T Result = 0;
 		Result += AssetDataList.GetAllocatedSize();
 		Result += DependencyData.GetAllocatedSize();
 		return Result;
diff --git a/Engine/Source/Runtime/AssetRegistry/Private/NameTableArchive.cpp b/Engine/Source/Runtime/AssetRegistry/Private/NameTableArchive.cpp
--- a/Engine/Source/Runtime/AssetRegistry/Private/NameTableArchive.cpp
+++ b/Engine/Source/Runtime/AssetRegistry/Private/NameTableArchive.cpp
@@ -52,8 +52,8 @@ bool FNameTableArchiveReader::SerializeNameMap()
 		}
 		
 		const int32 MinFNameEntrySize = sizeof(int32);
-		int32 MaxReservation = (TotalSize() - Tell()) / MinFNameEntrySize;
-		NameMap.Reserve(FMath::Min(NameCount, MaxReservation));
+		const int64 MaxReservation = ((TotalSize() - Tell()) / MinFNameEntrySize);
+		NameMap.Reserve((int32)FMath::Min<int64>(NameCount, MaxReservation));
 		for ( int32 NameMapIdx = 0; NameMapIdx < NameCount; ++NameMapIdx )
 		{
 			// Read the name entry from the file.
diff --git a/Engine/Source/Runtime/AssetRegistry/Private/PackageDependencyData.h b/Engine/Source/Runtime/AssetRegistry/Private/PackageDependencyData.h
--- a/Engine/Source/Runtime/AssetRegistry/Private/PackageDependencyData.h
+++ b/Engine/Source/Runtime/AssetRegistry/Private/PackageDependencyData.h
@@ -59,9 +59,9 @@ public:
 	}
 
 	/** Returns the amount of memory allocated by this container, not including sizeof(*this). */
-	uint32 GetAllocatedSize() const
+	SIZE_T GetAllocatedSize() const
 	{
-		uint32 Result = FLinkerTables::GetAllocatedSize();
+		SIZE_T Result = FLinkerTables::GetAllocatedSize();
 		Result += ImportUsedInGame.GetAllocatedSize();
 		Result += SoftPackageUsedInGame.GetAllocatedSize();
 		return Result;
diff --git a/Engine/Source/Runtime/AssetRegistry/Public/AssetRegistry/AssetRegistryState.h b/Engine/Source/Runtime/AssetRegistry/Public/AssetRegistry/AssetRegistryState.h
--- a/Engine/Source/Runtime/AssetRegistry/Public/AssetRegistry/AssetRegistryState.h
+++ b/Engine/Source/Runtime/AssetRegistry/Public/AssetRegistry/AssetRegistryState.h
@@ -384,7 +384,7 @@ public:
 	bool Load(FArchive& Ar, const FAssetRegistryLoadOptions& Options = FAssetRegistryLoadOptions());
 
 	/** Returns memory size of entire registry, optionally logging sizes */
-	uint32 GetAllocatedSize(bool bLogDetailed = false) const;
+	SIZE_T GetAllocatedSize(bool bLogDetailed = false) const;
 
 	/** Checks a filter to make sure there are no illegal entries */
 	static bool IsFilterValid(const FARCompiledFilter& Filter);
diff --git a/Engine/Source/Runtime/AssetRegistry/Public/AssetRegistry/IAssetRegistry.h b/Engine/Source/Runtime/AssetRegistry/Public/AssetRegistry/IAssetRegistry.h
--- a/Engine/Source/Runtime/AssetRegistry/Public/AssetRegistry/IAssetRegistry.h
+++ b/Engine/Source/Runtime/AssetRegistry/Public/AssetRegistry/IAssetRegistry.h
@@ -525,7 +525,7 @@ public:
 	virtual void AppendState(const FAssetRegistryState& InState) = 0;
 
 	/** Returns memory size of entire registry, optionally logging sizes */
-	virtual uint32 GetAllocatedSize(bool bLogDetailed = false) const = 0;
+	virtual SIZE_T GetAllocatedSize(bool bLogDetailed = false) const = 0;
 
 	/**
 	 * Fills in a AssetRegistryState with a copy of the data in the internal cache, overriding some
diff --git a/Engine/Source/Runtime/CoreUObject/Private/AssetRegistry/AssetBundleData.cpp b/Engine/Source/Runtime/CoreUObject/Private/AssetRegistry/AssetBundleData.cpp
--- a/Engine/Source/Runtime/CoreUObject/Private/AssetRegistry/AssetBundleData.cpp
+++ b/Engine/Source/Runtime/CoreUObject/Private/AssetRegistry/AssetBundleData.cpp
@@ -75,7 +75,7 @@ bool FAssetBundleEntry::ImportTextItem(const TCHAR*& Buffer, int32 PortFlags, UO
 	{
 		if (const TCHAR* BundleNameEnd = FCString::Strchr(BundleNameBegin, TCHAR('"')))
 		{
-			FName Name(BundleNameEnd - BundleNameBegin, BundleNameBegin);
+			FName Name(UE_PTRDIFF_TO_INT32(BundleNameEnd - BundleNameBegin), BundleNameBegin);
 
 			const TCHAR* PathIt = BundleNameEnd + 1;
 			if (SkipPrefix(/* in-out */ PathIt, BundleAssetsPrefix))
diff --git a/Engine/Source/Runtime/CoreUObject/Private/AssetRegistry/AssetDataTagMap.cpp b/Engine/Source/Runtime/CoreUObject/Private/AssetRegistry/AssetDataTagMap.cpp
--- a/Engine/Source/Runtime/CoreUObject/Private/AssetRegistry/AssetDataTagMap.cpp
+++ b/Engine/Source/Runtime/CoreUObject/Private/AssetRegistry/AssetDataTagMap.cpp
@@ -1388,9 +1388,9 @@ bool operator==(const FAssetDataTagMapSharedView& A, const FAssetDataTagMapShare
 	}
 }
 
-static uint32 GetLooseMapMemoryUsage(const FAssetDataTagMap& Map)
+static SIZE_T GetLooseMapMemoryUsage(const FAssetDataTagMap& Map)
 {
-	uint32 Out = sizeof(FAssetDataTagMap) + Map.GetAllocatedSize();
+	SIZE_T Out = sizeof(FAssetDataTagMap) + Map.GetAllocatedSize();
 	for (const TPair<FName, FString> & Pair : Map)
 	{
 		Out += Pair.Value.GetAllocatedSize();
@@ -1411,9 +1411,9 @@ void FAssetDataTagMapSharedView::FMemoryCounter::Include(const FAssetDataTagMapS
 	}
 }
 
-uint32 FAssetDataTagMapSharedView::FMemoryCounter::GetFixedSize() const
+SIZE_T FAssetDataTagMapSharedView::FMemoryCounter::GetFixedSize() const
 {
-	uint32 Out = 0;
+	SIZE_T Out = 0;
 	for (uint32 StoreIndex : FixedStoreIndices)
 	{
 		Out += sizeof(FixedTagPrivate::FStore);
diff --git a/Engine/Source/Runtime/CoreUObject/Public/AssetRegistry/AssetDataTagMap.h b/Engine/Source/Runtime/CoreUObject/Public/AssetRegistry/AssetDataTagMap.h
--- a/Engine/Source/Runtime/CoreUObject/Public/AssetRegistry/AssetDataTagMap.h
+++ b/Engine/Source/Runtime/CoreUObject/Public/AssetRegistry/AssetDataTagMap.h
@@ -456,11 +456,11 @@ public:
 	class COREUOBJECT_API FMemoryCounter
 	{
 		TSet<uint32> FixedStoreIndices;
-		uint32 LooseBytes = 0;
+		SIZE_T LooseBytes = 0;
 	public:
 		void Include(const FAssetDataTagMapSharedView& Tags);
-		uint32 GetLooseSize() const { return LooseBytes; }
-		uint32 GetFixedSize() const;		
+		SIZE_T GetLooseSize() const { return LooseBytes; }
+		SIZE_T GetFixedSize() const;		
 	};
 };
 
