Asset Registry: Fixing a couple of Clang-only float truncation warnings

diff --git a/Engine/Source/Runtime/AssetRegistry/Private/AssetDataGatherer.cpp b/Engine/Source/Runtime/AssetRegistry/Private/AssetDataGatherer.cpp
--- a/Engine/Source/Runtime/AssetRegistry/Private/AssetDataGatherer.cpp
+++ b/Engine/Source/Runtime/AssetRegistry/Private/AssetDataGatherer.cpp
@@ -2580,7 +2580,7 @@ bool FAssetDataGatherer::Init()
 
 uint32 FAssetDataGatherer::Run()
 {
-	constexpr double IdleSleepTime = 0.1;
+	constexpr float IdleSleepTime = 0.1f;
 	while (!IsStopped)
 	{
 		InnerTickLoop(false /* bInIsSynchronousTick */, true /* bContributeToCacheSave */);
diff --git a/Engine/Source/Runtime/AssetRegistry/Private/AssetRegistry.cpp b/Engine/Source/Runtime/AssetRegistry/Private/AssetRegistry.cpp
--- a/Engine/Source/Runtime/AssetRegistry/Private/AssetRegistry.cpp
+++ b/Engine/Source/Runtime/AssetRegistry/Private/AssetRegistry.cpp
@@ -95,7 +95,7 @@ static FAssetRegistryConsoleCommands ConsoleCommands; // Registers its various c
 namespace UE::AssetRegistry::Impl
 {
 	/** The max time to spend in UAssetRegistryImpl::Tick */
-	const float MaxSecondsPerFrame = 0.04;
+	const float MaxSecondsPerFrame = 0.04f;
 }
 
 /**
