Fixed another Clang warning about float truncation into FPlatformProcess::Sleep
[Third time's a charm]

diff --git a/Engine/Source/Runtime/AssetRegistry/Private/AssetDataGatherer.cpp b/Engine/Source/Runtime/AssetRegistry/Private/AssetDataGatherer.cpp
--- a/Engine/Source/Runtime/AssetRegistry/Private/AssetDataGatherer.cpp
+++ b/Engine/Source/Runtime/AssetRegistry/Private/AssetDataGatherer.cpp
@@ -1324,7 +1324,7 @@ bool FAssetDataDiscovery::Init()
 
 uint32 FAssetDataDiscovery::Run()
 {
-	constexpr double IdleSleepTime = 0.1;
+	constexpr float IdleSleepTime = 0.1f;
 	{
 		FGathererScopeLock ResultsScopeLock(&ResultsLock);
 		DiscoverStartTime = FPlatformTime::Seconds();
