AssetRegistry::GetAssetsByPath: handle a trailing / terminator in packagepaths. This was already handled for subfolders, but not for the top-level requested folder.

diff --git a/Engine/Source/Runtime/AssetRegistry/Private/AssetRegistry.cpp b/Engine/Source/Runtime/AssetRegistry/Private/AssetRegistry.cpp
--- a/Engine/Source/Runtime/AssetRegistry/Private/AssetRegistry.cpp
+++ b/Engine/Source/Runtime/AssetRegistry/Private/AssetRegistry.cpp
@@ -2299,7 +2299,11 @@ void FAssetRegistryImpl::CompileFilter(Impl::FClassInheritanceContext& Inheritan
 {
 	OutCompiledFilter.Clear();
 	OutCompiledFilter.PackageNames.Append(InFilter.PackageNames);
-	OutCompiledFilter.PackagePaths.Append(InFilter.PackagePaths);
+	OutCompiledFilter.PackagePaths.Reserve(InFilter.PackagePaths.Num());
+	for (FName PackagePath : InFilter.PackagePaths)
+	{
+		OutCompiledFilter.PackagePaths.Add(FPathTree::NormalizePackagePath(PackagePath));
+	}
 	OutCompiledFilter.ObjectPaths.Append(InFilter.ObjectPaths);
 	OutCompiledFilter.ClassNames.Append(InFilter.ClassNames);
 	OutCompiledFilter.TagsAndValues = InFilter.TagsAndValues;
diff --git a/Engine/Source/Runtime/AssetRegistry/Private/PathTree.cpp b/Engine/Source/Runtime/AssetRegistry/Private/PathTree.cpp
--- a/Engine/Source/Runtime/AssetRegistry/Private/PathTree.cpp
+++ b/Engine/Source/Runtime/AssetRegistry/Private/PathTree.cpp
@@ -222,6 +222,17 @@ bool FPathTree::GetSubPaths(FName BasePath, TSet<FName>& OutPaths, bool bRecurse
 	}, bRecurse) && OutPaths.Num() > OutPathsOriginalNum;
 }
 
+FName FPathTree::NormalizePackagePath(FName In)
+{
+	FNameBuilder InStr(In);
+	if (InStr.Len() == 0 || InStr.LastChar() != TEXT('/'))
+	{
+		return In;
+	}
+	InStr.RemoveSuffix(1);
+	return FName(InStr);
+}
+
 bool FPathTree::EnumerateSubPaths(FName BasePath, TFunctionRef<bool(FName)> Callback, bool bRecurse) const
 {
 	if (BasePath.IsNone())
@@ -233,14 +244,11 @@ bool FPathTree::EnumerateSubPaths(FName BasePath, TFunctionRef<bool(FName)> Call
 	if (!ChildPathsPtr)
 	{
 		// Paths are cached without their trailing slash, so if the given path has a trailing slash, test it again now as it may already be cached
-		// We do this after the initial map test as: a) Most paths are well formed, b) This avoids string allocations until we know we need them
-		FNameBuilder BasePathStr(BasePath);
-		if (BasePathStr.LastChar() == TEXT('/'))
+		// We do this after the initial map test as: a) Most paths are well formed, b) This avoids string tests until we know it's possibly needed
+		FName NormalizedBasePath = NormalizePackagePath(BasePath);
+		if (NormalizedBasePath != BasePath)
 		{
-			BasePathStr.RemoveSuffix(1);
-			BasePath = *BasePathStr;
-
-			ChildPathsPtr = ParentPathToChildPaths.Find(BasePath);
+			ChildPathsPtr = ParentPathToChildPaths.Find(NormalizedBasePath);
 			if (!ChildPathsPtr)
 			{
 				return false;
diff --git a/Engine/Source/Runtime/AssetRegistry/Public/AssetRegistry/PathTree.h b/Engine/Source/Runtime/AssetRegistry/Public/AssetRegistry/PathTree.h
--- a/Engine/Source/Runtime/AssetRegistry/Public/AssetRegistry/PathTree.h
+++ b/Engine/Source/Runtime/AssetRegistry/Public/AssetRegistry/PathTree.h
@@ -45,6 +45,9 @@ public:
 		return AllocatedSize;
 	}
 
+	/** Normalize the given PackagePath (/Game/SubDir) by removing a trailing slash if it exists. */
+	static FName NormalizePackagePath(FName In);
+
 private:
 	/** A one-to-many mapping between a parent path and its child paths. */
 	TMap<FName, TSet<FName>> ParentPathToChildPaths;
