AssetRegistry: Fix crash in package reader if a content directory is unmounted in between queing the package and reading it.
#rnx

diff --git a/Engine/Source/Runtime/AssetRegistry/Private/PackageReader.cpp b/Engine/Source/Runtime/AssetRegistry/Private/PackageReader.cpp
--- a/Engine/Source/Runtime/AssetRegistry/Private/PackageReader.cpp
+++ b/Engine/Source/Runtime/AssetRegistry/Private/PackageReader.cpp
@@ -245,7 +245,11 @@ bool FPackageReader::ReadAssetDataFromThumbnailCache(TArray<FAssetData*>& AssetD
 	}
 
 	// Determine the package name and path
-	FString PackageName = FPackageName::FilenameToLongPackageName(PackageFilename);
+	FString PackageName;
+	if (!FPackageName::TryConvertFilenameToLongPackageName(PackageFilename, PackageName))
+	{
+		return false;
+	}
 	FString PackagePath = FPackageName::GetLongPackagePath(PackageName);
 
 	// Load the thumbnail count
@@ -298,8 +302,12 @@ bool FPackageReader::ReadAssetRegistryDataIfCookedPackage(TArray<FAssetData*>& A
 {
 	if (!!(GetPackageFlags() & PKG_FilterEditorOnly))
 	{
-		const FString PackageName = FPackageName::FilenameToLongPackageName(PackageFilename);
-		
+		FString PackageName;
+		if (!FPackageName::TryConvertFilenameToLongPackageName(PackageFilename, PackageName))
+		{
+			return false;
+		}
+
 		bool bFoundAtLeastOneAsset = false;
 
 		// If the packaged is saved with the right version we have the information
@@ -758,9 +766,7 @@ FArchive& FPackageReader::operator<<( FName& Name )
 	return *this;
 }
 
-namespace UE
-{
-namespace AssetRegistry
+namespace UE::AssetRegistry
 {
 	// See the corresponding WritePackageData defined in SavePackageUtilities.cpp in CoreUObject module
 	bool ReadPackageDataMain(FArchive& BinaryArchive, const FString& PackageName, const FPackageFileSummary& PackageFileSummary, int64& OutDependencyDataOffset, TArray<FAssetData*>& OutAssetDataList, EReadPackageDataMainErrorCode& OutError)
@@ -890,4 +896,3 @@ namespace AssetRegistry
 		return !BinaryArchive.IsError();
 	}
 }
-}
