AssetRegistry: Update the FAssetData when an asset is saved, so that any future disk-only queries will read the up-to-date data.
ObjectSaveContext: Fix ObjectSaveContext never having IsUpdatingLoadPath set in the in OnPreSaveWorld, PreSavePackageWithContext, and OnPostSaveWorld events.
#rnx

diff --git a/Engine/Source/Editor/UnrealEd/Private/EditorEngine.cpp b/Engine/Source/Editor/UnrealEd/Private/EditorEngine.cpp
--- a/Engine/Source/Editor/UnrealEd/Private/EditorEngine.cpp
+++ b/Engine/Source/Editor/UnrealEd/Private/EditorEngine.cpp
@@ -9,7 +9,8 @@
 #include "Misc/ScopedSlowTask.h"
 #include "Misc/CoreDelegates.h"
 #include "Misc/App.h"
-#include "Misc/AssetRegistryInterface.h"
+#include "AssetRegistry/AssetRegistryModule.h"
+#include "AssetRegistry/IAssetRegistry.h"
 #include "Modules/ModuleManager.h"
 #include "UObject/MetaData.h"
 #include "UObject/ObjectSaveContext.h"
@@ -4355,6 +4356,20 @@ FSavePackageResultStruct UEditorEngine::Save(UPackage* InOuter, UObject* InAsset
 		}
 	}
 
+	if (ObjectSaveContext.bUpdatingLoadedPath)
+	{
+		// Notify the asset registry
+		IAssetRegistry& AssetRegistry = IAssetRegistry::GetChecked();
+		ForEachObjectWithPackage(InOuter, [&AssetRegistry](UObject* Object)
+			{
+				if (Object->IsAsset() && !UE::AssetRegistry::FFiltering::ShouldSkipAsset(Object))
+				{
+					AssetRegistry.AssetSaved(*Object);
+				}
+				return true;
+			}, false /*bIncludeNestedObjects*/);
+	}
+
 	return Result;
 }
 
diff --git a/Engine/Source/Runtime/AssetRegistry/Private/AssetRegistry.cpp b/Engine/Source/Runtime/AssetRegistry/Private/AssetRegistry.cpp
--- a/Engine/Source/Runtime/AssetRegistry/Private/AssetRegistry.cpp
+++ b/Engine/Source/Runtime/AssetRegistry/Private/AssetRegistry.cpp
@@ -2791,6 +2791,19 @@ void UAssetRegistryImpl::AssetRenamed(const UObject* RenamedAsset, const FString
 	}
 }
 
+void UAssetRegistryImpl::AssetSaved(const UObject& SavedAsset)
+{
+#if WITH_EDITOR
+	if (!ensure(SavedAsset.IsAsset()))
+	{
+		return;
+	}
+
+	FWriteScopeLock InterfaceScopeLock(InterfaceLock);
+	GuardedData.OnAssetSaved(SavedAsset);
+#endif
+}
+
 void UAssetRegistryImpl::PackageDeleted(UPackage* DeletedPackage)
 {
 	checkf(GIsEditor, TEXT("Updating the AssetRegistry is only available in editor"));
@@ -4152,7 +4165,7 @@ void UAssetRegistryImpl::ProcessLoadedAssetsToUpdateCache(UE::AssetRegistry::Imp
 
 
 	constexpr int32 BatchSize = 16;
-	TArray<UObject*> BatchObjects;
+	TArray<const UObject*> BatchObjects;
 	TArray<FAssetData, TInlineAllocator<BatchSize>> BatchAssetDatas;
 
 	{
@@ -4175,7 +4188,7 @@ void UAssetRegistryImpl::ProcessLoadedAssetsToUpdateCache(UE::AssetRegistry::Imp
 		int32 Index = 0;
 		while (Index < CurrentBatchSize)
 		{
-			UObject* LoadedObject = BatchObjects[Index++];
+			const UObject* LoadedObject = BatchObjects[Index++];
 			if (!LoadedObject->IsAsset())
 			{
 				// If the object has changed and is no longer an asset, ignore it. This can happen when an Actor is modified during cooking to no longer have an external package
@@ -4194,7 +4207,7 @@ void UAssetRegistryImpl::ProcessLoadedAssetsToUpdateCache(UE::AssetRegistry::Imp
 
 		FWriteScopeLock InterfaceScopeLock(InterfaceLock);
 		GuardedData.PushProcessLoadedAssetsBatch(EventContext, BatchAssetDatas,
-			TArrayView<UObject*>(BatchObjects).Slice(Index, CurrentBatchSize-Index));
+			TArrayView<const UObject*>(BatchObjects).Slice(Index, CurrentBatchSize-Index));
 		if (bTimedOut)
 		{
 			break;
@@ -4211,7 +4224,12 @@ void FAssetRegistryImpl::OnAssetLoaded(UObject* AssetLoaded)
 	LoadedAssetsToProcess.Add(AssetLoaded);
 }
 
-void FAssetRegistryImpl::GetProcessLoadedAssetsBatch(TArray<UObject*>& OutLoadedAssets, uint32 BatchSize)
+void FAssetRegistryImpl::OnAssetSaved(const UObject& AssetSaved)
+{
+	LoadedAssetsToProcess.Add(&AssetSaved);
+}
+
+void FAssetRegistryImpl::GetProcessLoadedAssetsBatch(TArray<const UObject*>& OutLoadedAssets, uint32 BatchSize)
 {
 	if (!GlobalGatherer.IsValid() || !bUpdateDiskCacheAfterLoad)
 	{
@@ -4222,7 +4240,7 @@ void FAssetRegistryImpl::GetProcessLoadedAssetsBatch(TArray<UObject*>& OutLoaded
 	OutLoadedAssets.Reset(BatchSize);
 	while (!LoadedAssetsToProcess.IsEmpty() && OutLoadedAssets.Num() < static_cast<int32>(BatchSize))
 	{
-		UObject* LoadedAsset = LoadedAssetsToProcess.PopFrontValue().Get();
+		const UObject* LoadedAsset = LoadedAssetsToProcess.PopFrontValue().Get();
 		if (!LoadedAsset)
 		{
 			// This could be null, in which case it already got freed, ignore
@@ -4235,12 +4253,7 @@ void FAssetRegistryImpl::GetProcessLoadedAssetsBatch(TArray<UObject*>& OutLoaded
 			continue;
 		}
 
-		const FName ObjectPath = FName(*LoadedAsset->GetPathName());
-		if (AssetDataObjectPathsUpdatedOnLoad.Contains(ObjectPath))
-		{
-			// Already processed once, don't process again even if it loads a second time
-			continue;
-		}
+		// Take a new snapshot of the asset's data every time it loads or saves
 
 		UPackage* InMemoryPackage = LoadedAsset->GetOutermost();
 		if (InMemoryPackage->IsDirty())
@@ -4254,7 +4267,7 @@ void FAssetRegistryImpl::GetProcessLoadedAssetsBatch(TArray<UObject*>& OutLoaded
 }
 
 void FAssetRegistryImpl::PushProcessLoadedAssetsBatch(Impl::FEventContext& EventContext,
-	TArrayView<FAssetData> LoadedAssetDatas, TArrayView<UObject*> UnprocessedFromBatch)
+	TArrayView<FAssetData> LoadedAssetDatas, TArrayView<const UObject*> UnprocessedFromBatch)
 {
 	// Add or update existing for all of the AssetDatas created by the batch
 	for (FAssetData& NewAssetData : LoadedAssetDatas)
diff --git a/Engine/Source/Runtime/AssetRegistry/Private/AssetRegistry.h b/Engine/Source/Runtime/AssetRegistry/Private/AssetRegistry.h
--- a/Engine/Source/Runtime/AssetRegistry/Private/AssetRegistry.h
+++ b/Engine/Source/Runtime/AssetRegistry/Private/AssetRegistry.h
@@ -119,6 +119,7 @@ public:
 	virtual void AssetCreated(UObject* NewAsset) override;
 	virtual void AssetDeleted(UObject* DeletedAsset) override;
 	virtual void AssetRenamed(const UObject* RenamedAsset, const FString& OldObjectPath) override;
+	virtual void AssetSaved(const UObject& SavedAsset) override;
 
 	virtual void PackageDeleted(UPackage* DeletedPackage) override;
 
diff --git a/Engine/Source/Runtime/AssetRegistry/Private/AssetRegistryImpl.h b/Engine/Source/Runtime/AssetRegistry/Private/AssetRegistryImpl.h
--- a/Engine/Source/Runtime/AssetRegistry/Private/AssetRegistryImpl.h
+++ b/Engine/Source/Runtime/AssetRegistry/Private/AssetRegistryImpl.h
@@ -98,6 +98,7 @@ public:
 #if WITH_EDITOR
 	void OnDirectoryChanged(Impl::FEventContext& EventContext, TArray<FFileChangeData>& FileChangesProcessed);
 	void OnAssetLoaded(UObject* AssetLoaded);
+	void OnAssetSaved(const UObject& AssetSaved);
 #endif
 	void OnContentPathMounted(Impl::FEventContext& EventContext, const FString& InAssetPath,
 		const FString& AssetPathWithTrailingSlash, const FString& FileSystemPath);
@@ -130,9 +131,9 @@ public:
 	/** Look for and load a single AssetData result from the gatherer. */
 	void TickGatherPackage(Impl::FEventContext& EventContext, const FString& PackageName, const FString& LocalPath);
 #if WITH_EDITOR
-	void GetProcessLoadedAssetsBatch(TArray<UObject*>& OutLoadedAssets, uint32 BatchSize);
+	void GetProcessLoadedAssetsBatch(TArray<const UObject*>& OutLoadedAssets, uint32 BatchSize);
 	void PushProcessLoadedAssetsBatch(Impl::FEventContext& EventContext,
-		TArrayView<FAssetData> LoadedAssetDatas, TArrayView<UObject*> UnprocessedFromBatch);
+		TArrayView<FAssetData> LoadedAssetDatas, TArrayView<const UObject*> UnprocessedFromBatch);
 #endif
 
 	/** Adds an asset to the empty package list which contains packages that have no assets left in them */
@@ -315,8 +316,8 @@ private:
 	TArray<FAssetRegistryPackageRedirect> PackageRedirects;
 
 #if WITH_EDITOR
-	/** List of loaded objects that need to be processed */
-	TRingBuffer<TWeakObjectPtr<UObject>> LoadedAssetsToProcess;
+	/** List of objects that need to be processed because they were loaded or saved */
+	TRingBuffer<TWeakObjectPtr<const UObject>> LoadedAssetsToProcess;
 
 	/** The set of object paths that have had their disk cache updated from the in memory version */
 	TSet<FName> AssetDataObjectPathsUpdatedOnLoad;
diff --git a/Engine/Source/Runtime/AssetRegistry/Public/AssetRegistry/AssetRegistryModule.h b/Engine/Source/Runtime/AssetRegistry/Public/AssetRegistry/AssetRegistryModule.h
--- a/Engine/Source/Runtime/AssetRegistry/Public/AssetRegistry/AssetRegistryModule.h
+++ b/Engine/Source/Runtime/AssetRegistry/Public/AssetRegistry/AssetRegistryModule.h
@@ -55,6 +55,11 @@ public:
 		IAssetRegistry::GetChecked().AssetRenamed(RenamedAsset, OldObjectPath);
 	}
 
+	static void AssetSaved(const UObject& SavedAsset)
+	{
+		IAssetRegistry::GetChecked().AssetSaved(SavedAsset);
+	}
+
 	static void PackageDeleted(UPackage* DeletedPackage)
 	{
 		IAssetRegistry::GetChecked().PackageDeleted(DeletedPackage);
diff --git a/Engine/Source/Runtime/AssetRegistry/Public/AssetRegistry/IAssetRegistry.h b/Engine/Source/Runtime/AssetRegistry/Public/AssetRegistry/IAssetRegistry.h
--- a/Engine/Source/Runtime/AssetRegistry/Public/AssetRegistry/IAssetRegistry.h
+++ b/Engine/Source/Runtime/AssetRegistry/Public/AssetRegistry/IAssetRegistry.h
@@ -471,13 +471,16 @@ public:
 	virtual FPathRemovedEvent& OnPathRemoved() = 0;
 
 	/** Informs the asset registry that an in-memory asset has been created */
-	virtual void AssetCreated (UObject* NewAsset) = 0;
+	virtual void AssetCreated(UObject* NewAsset) = 0;
 
 	/** Informs the asset registry that an in-memory asset has been deleted */
-	virtual void AssetDeleted (UObject* DeletedAsset) = 0;
+	virtual void AssetDeleted(UObject* DeletedAsset) = 0;
 
 	/** Informs the asset registry that an in-memory asset has been renamed */
-	virtual void AssetRenamed (const UObject* RenamedAsset, const FString& OldObjectPath) = 0;
+	virtual void AssetRenamed(const UObject* RenamedAsset, const FString& OldObjectPath) = 0;
+
+	/** Informs the asset registry that an in-memory asset has been saved */
+	virtual void AssetSaved(const UObject& SavedAsset) = 0;
 
 	/** Informs the asset registry that an in-memory package has been deleted, and all associated assets should be removed */
 	virtual void PackageDeleted (UPackage* DeletedPackage) = 0;
diff --git a/Engine/Source/Runtime/CoreUObject/Private/UObject/SavePackage/SavePackageUtilities.cpp b/Engine/Source/Runtime/CoreUObject/Private/UObject/SavePackage/SavePackageUtilities.cpp
--- a/Engine/Source/Runtime/CoreUObject/Private/UObject/SavePackage/SavePackageUtilities.cpp
+++ b/Engine/Source/Runtime/CoreUObject/Private/UObject/SavePackage/SavePackageUtilities.cpp
@@ -943,7 +943,7 @@ FObjectSaveContextData::FObjectSaveContextData(UPackage* Package, const ITargetP
 
 void FObjectSaveContextData::Set(UPackage* Package, const ITargetPlatform* InTargetPlatform, const TCHAR* InTargetFilename, uint32 InSaveFlags)
 {
-	FPackagePath PackagePath(FPackagePath::FromLocalPath(TargetFilename));
+	FPackagePath PackagePath(FPackagePath::FromLocalPath(InTargetFilename));
 	if (PackagePath.GetHeaderExtension() == EPackageExtension::Unspecified)
 	{
 		PackagePath.SetHeaderExtension(EPackageExtension::EmptyString);
