Fix WaitForCompletion to wait until the AssetRegistry finishes the current gather, rather than returning immediately any time after the first gather has completed.
#rnx

diff --git a/Engine/Source/Runtime/AssetRegistry/Private/AssetDataGatherer.cpp b/Engine/Source/Runtime/AssetRegistry/Private/AssetDataGatherer.cpp
--- a/Engine/Source/Runtime/AssetRegistry/Private/AssetDataGatherer.cpp
+++ b/Engine/Source/Runtime/AssetRegistry/Private/AssetDataGatherer.cpp
@@ -3077,7 +3077,8 @@ void FAssetDataGatherer::AddToCache(FName PackageName, FDiskCachedAssetData* Dis
 	}
 }
 
-void FAssetDataGatherer::GetAndTrimSearchResults(bool& bOutIsSearching, TRingBuffer<FAssetData*>& OutAssetResults, TRingBuffer<FString>& OutPathResults, TRingBuffer<FPackageDependencyData>& OutDependencyResults,
+void FAssetDataGatherer::GetAndTrimSearchResults(bool& bOutIsSearching, bool& bOutAbleToProgress,
+	TRingBuffer<FAssetData*>& OutAssetResults, TRingBuffer<FString>& OutPathResults, TRingBuffer<FPackageDependencyData>& OutDependencyResults,
 	TRingBuffer<FString>& OutCookedPackageNamesWithoutAssetDataResults, TArray<double>& OutSearchTimes, int32& OutNumFilesToSearch, int32& OutNumPathsToSearch, bool& OutIsDiscoveringFiles)
 {
 	FGathererScopeLock ResultsScopeLock(&ResultsLock);
@@ -3112,6 +3113,7 @@ void FAssetDataGatherer::GetAndTrimSearchResults(bool& bOutIsSearching, TRingBuf
 		Shrink();
 	}
 	bOutIsSearching = !bLocalIsComplete;
+	bOutAbleToProgress = !bIsIdle;
 }
 
 void FAssetDataGatherer::GetPackageResults(TRingBuffer<FAssetData*>& OutAssetResults, TRingBuffer<FPackageDependencyData>& OutDependencyResults)
@@ -4139,4 +4141,4 @@ int32 FFilesToSearch::FTreeNode::NumFiles() const
 	return Num;
 }
 
-}
+}
\ No newline at end of file
diff --git a/Engine/Source/Runtime/AssetRegistry/Private/AssetDataGatherer.h b/Engine/Source/Runtime/AssetRegistry/Private/AssetDataGatherer.h
--- a/Engine/Source/Runtime/AssetRegistry/Private/AssetDataGatherer.h
+++ b/Engine/Source/Runtime/AssetRegistry/Private/AssetDataGatherer.h
@@ -85,7 +85,8 @@ public:
 	// Receiving Results (possibly while tick is running)
 
 	/** Gets search results from the data gatherer. */
-	void GetAndTrimSearchResults(bool& bOutIsSearching, TRingBuffer<FAssetData*>& OutAssetResults, TRingBuffer<FString>& OutPathResults,
+	void GetAndTrimSearchResults(bool& bOutIsSearching, bool& bOutAbleToProgress,
+		TRingBuffer<FAssetData*>& OutAssetResults, TRingBuffer<FString>& OutPathResults,
 		TRingBuffer<FPackageDependencyData>& OutDependencyResults, TRingBuffer<FString>& OutCookedPackageNamesWithoutAssetDataResults,
 		TArray<double>& OutSearchTimes, int32& OutNumFilesToSearch, int32& OutNumPathsToSearch, bool& OutIsDiscoveringFiles);
 	/** Gets just the AssetResults and DependencyResults from the data gatherer. */
diff --git a/Engine/Source/Runtime/AssetRegistry/Private/AssetRegistry.cpp b/Engine/Source/Runtime/AssetRegistry/Private/AssetRegistry.cpp
--- a/Engine/Source/Runtime/AssetRegistry/Private/AssetRegistry.cpp
+++ b/Engine/Source/Runtime/AssetRegistry/Private/AssetRegistry.cpp
@@ -347,7 +347,7 @@ void FAssetRegistryImpl::Initialize(Impl::FInitializeContext& Context)
 
 	bInitialSearchStarted = false;
 	bInitialSearchCompleted = true;
-	bGatherIdle = false;
+	GatherStatus = Impl::EGatherStatus::Active;
 	bSearchAllAssets = false;
 	AmortizeStartTime = 0;
 	TotalAmortizeTime = 0;
@@ -1037,7 +1037,9 @@ void FAssetRegistryImpl::SearchAllAssetsInitialAsync(Impl::FEventContext& EventC
 
 void UAssetRegistryImpl::SearchAllAssets(bool bSynchronousSearch)
 {
-	UE::AssetRegistry::Impl::FEventContext EventContext;
+	using namespace UE::AssetRegistry::Impl;
+
+	FEventContext EventContext;
 	{
 		FWriteScopeLock InterfaceScopeLock(InterfaceLock);
 		GuardedData.SearchAllAssets(EventContext, bSynchronousSearch);
@@ -1045,7 +1047,7 @@ void UAssetRegistryImpl::SearchAllAssets(bool bSynchronousSearch)
 #if WITH_EDITOR
 	if (bSynchronousSearch)
 	{
-		ProcessLoadedAssetsToUpdateCache(EventContext, -1., true /* bIdle */);
+		ProcessLoadedAssetsToUpdateCache(EventContext, -1., EGatherStatus::Complete);
 	}
 #endif
 	Broadcast(EventContext);
@@ -1093,9 +1095,8 @@ void FAssetRegistryImpl::SearchAllAssets(Impl::FEventContext& EventContext, bool
 	if (bSynchronousSearch)
 	{
 		Gatherer.WaitForIdle();
-		bool bUnusedIdle;
 		bool bUnusedInterrupted;
-		TickGatherer(EventContext, -1., bUnusedIdle, bUnusedInterrupted);
+		Impl::EGatherStatus UnusedStatus = TickGatherer(EventContext, -1., bUnusedInterrupted);
 #if WITH_EDITOR
 		if (!bInitialSearchStarted)
 		{
@@ -1115,23 +1116,25 @@ void FAssetRegistryImpl::SearchAllAssets(Impl::FEventContext& EventContext, bool
 
 void UAssetRegistryImpl::WaitForCompletion()
 {
+	using namespace UE::AssetRegistry::Impl;
+
 	for (;;)
 	{
-		UE::AssetRegistry::Impl::FEventContext EventContext;
-		bool bIsIdle = false;
+		FEventContext EventContext;
+		EGatherStatus Status;
 		{
 			FWriteScopeLock InterfaceScopeLock(InterfaceLock);
-			if (!GuardedData.IsLoadingAssets())
-			{
-				break;
-			}
 			bool bUnusedInterrupted;
-			GuardedData.TickGatherer(EventContext, -1., bIsIdle, bUnusedInterrupted);
+			Status = GuardedData.TickGatherer(EventContext, -1., bUnusedInterrupted);
 		}
 #if WITH_EDITOR
-		ProcessLoadedAssetsToUpdateCache(EventContext, -1., bIsIdle);
+		ProcessLoadedAssetsToUpdateCache(EventContext, -1., Status);
 #endif
 		Broadcast(EventContext);
+		if (Status != EGatherStatus::Active)
+		{
+			break;
+		}
 
 		FThreadHeartBeat::Get().HeartBeat();
 		FPlatformProcess::SleepNoStats(0.0001f);
@@ -2593,7 +2596,7 @@ void UAssetRegistryImpl::ScanPathsSynchronousInternal(const TArray<FString>& InD
 	}
 
 #if WITH_EDITOR
-	ProcessLoadedAssetsToUpdateCache(EventContext, -1., Context.bIdle);
+	ProcessLoadedAssetsToUpdateCache(EventContext, -1., Context.Status);
 #endif
 	Broadcast(EventContext);
 
@@ -2836,7 +2839,7 @@ void UAssetRegistryImpl::Tick(float DeltaTime)
 {
 	checkf(IsInGameThread(), TEXT("The tick function executes deferred loads and events and must be on the game thread to do so."));
 
-	bool bIdle = false;
+	UE::AssetRegistry::Impl::EGatherStatus Status = UE::AssetRegistry::Impl::EGatherStatus::Active;
 	double TickStartTime = -1; // Force a full flush if DeltaTime < 0
 	if (DeltaTime >= 0)
 	{
@@ -2861,7 +2864,7 @@ void UAssetRegistryImpl::Tick(float DeltaTime)
 			if (EventContext.IsEmpty())
 			{
 				// Tick the Gatherer
-				GuardedData.TickGatherer(EventContext, TickStartTime, bIdle, bInterrupted);
+				Status = GuardedData.TickGatherer(EventContext, TickStartTime, bInterrupted);
 			}
 			else
 			{
@@ -2873,7 +2876,7 @@ void UAssetRegistryImpl::Tick(float DeltaTime)
 #if WITH_EDITOR
 		if (!bInterrupted)
 		{
-			ProcessLoadedAssetsToUpdateCache(EventContext, TickStartTime, bIdle);
+			ProcessLoadedAssetsToUpdateCache(EventContext, TickStartTime, Status);
 		}
 #endif
 		Broadcast(EventContext);
@@ -2894,37 +2897,55 @@ void FAssetRegistryImpl::TickDeletes()
 	PRAGMA_ENABLE_DEPRECATION_WARNINGS
 }
 
-void FAssetRegistryImpl::TickGatherer(Impl::FEventContext& EventContext, const double TickStartTime, bool& bOutIdle, bool& bOutInterrupted, TOptional<FAssetsFoundCallback> AssetsFoundCallback)
+Impl::EGatherStatus FAssetRegistryImpl::TickGatherer(Impl::FEventContext& EventContext,
+	const double TickStartTime, bool& bOutInterrupted, TOptional<FAssetsFoundCallback> AssetsFoundCallback)
 {
-	bOutIdle = true;
+	using namespace UE::AssetRegistry::Impl;
+
+	EGatherStatus OutStatus = EGatherStatus::Complete;
 	bOutInterrupted = false;
 	if (!GlobalGatherer.IsValid())
 	{
-		return;
+		return OutStatus;
 	}
 
 	// Gather results from the background search
 	bool bIsSearching = false;
+	bool bAbleToProgress = false;
 	TArray<double> SearchTimes;
 	int32 NumFilesToSearch = 0;
 	int32 NumPathsToSearch = 0;
 	bool bIsDiscoveringFiles = false;
-	GlobalGatherer->GetAndTrimSearchResults(bIsSearching, BackgroundAssetResults, BackgroundPathResults, BackgroundDependencyResults, BackgroundCookedPackageNamesWithoutAssetDataResults, SearchTimes, NumFilesToSearch, NumPathsToSearch, bIsDiscoveringFiles);
+	GlobalGatherer->GetAndTrimSearchResults(bIsSearching, bAbleToProgress, BackgroundAssetResults, BackgroundPathResults, BackgroundDependencyResults,
+		BackgroundCookedPackageNamesWithoutAssetDataResults, SearchTimes, NumFilesToSearch, NumPathsToSearch, bIsDiscoveringFiles);
 	// Report the search times
 	for (int32 SearchTimeIdx = 0; SearchTimeIdx < SearchTimes.Num(); ++SearchTimeIdx)
 	{
 		UE_LOG(LogAssetRegistry, Verbose, TEXT("### Background search completed in %0.4f seconds"), SearchTimes[SearchTimeIdx]);
 	}
 	const bool bHadAssetsToProcess = BackgroundAssetResults.Num() > 0 || BackgroundDependencyResults.Num() > 0;
-	auto UpdateStatus = [&bIsSearching, &bHadAssetsToProcess, &bIsDiscoveringFiles, &EventContext, &bOutIdle, &NumFilesToSearch, &NumPathsToSearch, this](bool bInterrupted)
+	auto UpdateStatus = [bIsSearching, bAbleToProgress, bHadAssetsToProcess, bIsDiscoveringFiles, &EventContext,
+		&OutStatus, &NumFilesToSearch, &NumPathsToSearch, this](bool bInterrupted)
 	{
 		// Compute total pending, plus highest pending for this run so we can show a good progress bar
 		int32 NumPending = NumFilesToSearch + NumPathsToSearch + BackgroundPathResults.Num() + BackgroundAssetResults.Num() + BackgroundDependencyResults.Num() + BackgroundCookedPackageNamesWithoutAssetDataResults.Num();
 		HighestPending = FMath::Max(this->HighestPending, NumPending);
 
-		bOutIdle = !bInterrupted && !bIsSearching && NumPending == 0;
+		if (!bInterrupted && !bIsSearching && NumPending == 0)
+		{
+			OutStatus = EGatherStatus::Complete;
+		}
+		else if (!bInterrupted && !bAbleToProgress)
+		{
+			OutStatus = EGatherStatus::UnableToProgress;
+		}
+		else
+		{
+			OutStatus = EGatherStatus::Active;
+		}
 		// Notify the status change, only when something changed, or when sending the final result before going idle
-		if (bIsSearching || bHadAssetsToProcess || (bOutIdle && !this->bGatherIdle))
+		if (bIsSearching || bHadAssetsToProcess ||
+			(OutStatus == EGatherStatus::Complete && this->GatherStatus != EGatherStatus::Complete))
 		{
 			EventContext.ProgressUpdateData.Emplace(
 				HighestPending,					// NumTotalAssets
@@ -2933,7 +2954,7 @@ void FAssetRegistryImpl::TickGatherer(Impl::FEventContext& EventContext, const d
 				bIsDiscoveringFiles				// bIsDiscoveringAssetFiles
 			);
 		}
-		this->bGatherIdle = bOutIdle;
+		this->GatherStatus = OutStatus;
 	};
 
 
@@ -2978,14 +2999,14 @@ void FAssetRegistryImpl::TickGatherer(Impl::FEventContext& EventContext, const d
 		if (bOutInterrupted)
 		{
 			UpdateStatus(true /* bInterrupted */);
-			return;
+			return OutStatus;
 		}
 	}
 
 	// If completing an initial search, refresh the content browser
 	UpdateStatus(false /* bInterrupted */);
 
-	if (bOutIdle)
+	if (OutStatus == EGatherStatus::Complete)
 	{
 		HighestPending = 0;
 
@@ -3003,6 +3024,8 @@ void FAssetRegistryImpl::TickGatherer(Impl::FEventContext& EventContext, const d
 			EventContext.bFileLoadedEventBroadcast = true;
 		}
 	}
+
+	return OutStatus;
 }
 
 void FAssetRegistryImpl::TickGatherPackage(Impl::FEventContext& EventContext, const FString& PackageName, const FString& LocalPath)
@@ -3473,7 +3496,7 @@ void FAssetRegistryImpl::ScanPathsSynchronous(Impl::FScanPathContext& Context)
 	};
 
 	bool bUnusedInterrupted;
-	TickGatherer(Context.EventContext, -1., Context.bIdle, bUnusedInterrupted, FAssetsFoundCallback(AssetsFoundCallback));
+	Context.Status = TickGatherer(Context.EventContext, -1., bUnusedInterrupted, FAssetsFoundCallback(AssetsFoundCallback));
 }
 
 namespace Utils
@@ -4146,7 +4169,7 @@ void UAssetRegistryImpl::OnAssetLoaded(UObject *AssetLoaded)
 }
 
 void UAssetRegistryImpl::ProcessLoadedAssetsToUpdateCache(UE::AssetRegistry::Impl::FEventContext& EventContext,
-	const double TickStartTime, bool bIdle)
+	const double TickStartTime, UE::AssetRegistry::Impl::EGatherStatus Status)
 {
 	// Note this function can be reentered due to arbitrary code execution in construction of FAssetData
 	if (!IsInGameThread())
@@ -4155,8 +4178,9 @@ void UAssetRegistryImpl::ProcessLoadedAssetsToUpdateCache(UE::AssetRegistry::Imp
 		return;
 	}
 
+	// Early exit to save cputime if we're still processing cache data
 	const bool bFlushFullBuffer = TickStartTime < 0;
-	if (!bIdle && !bFlushFullBuffer)
+	if (Status == UE::AssetRegistry::Impl::EGatherStatus::Active && !bFlushFullBuffer)
 	{
 		return;
 	}
@@ -5399,4 +5423,4 @@ UAssetRegistryImpl::FFileLoadProgressUpdatedEvent& UAssetRegistryImpl::OnFileLoa
 {
 	checkf(IsInGameThread(), TEXT("Registering to AssetRegistry events is not supported from multiple threads."));
 	return FileLoadProgressUpdatedEvent;
-}
+}
\ No newline at end of file
diff --git a/Engine/Source/Runtime/AssetRegistry/Private/AssetRegistry.h b/Engine/Source/Runtime/AssetRegistry/Private/AssetRegistry.h
--- a/Engine/Source/Runtime/AssetRegistry/Private/AssetRegistry.h
+++ b/Engine/Source/Runtime/AssetRegistry/Private/AssetRegistry.h
@@ -208,7 +208,7 @@ private:
 #if WITH_EDITOR
 	/** Create FAssetData from any loaded UObject assets and store the updated AssetData in the state */
 	void ProcessLoadedAssetsToUpdateCache(UE::AssetRegistry::Impl::FEventContext& EventContext,
-		const double TickStartTime, bool bIsIdle);
+		const double TickStartTime, UE::AssetRegistry::Impl::EGatherStatus Status);
 #endif
 	/**
 	 * Remain under the given lock and return an InheritanceContext based on the appropriate choice of the persistent
diff --git a/Engine/Source/Runtime/AssetRegistry/Private/AssetRegistryImpl.h b/Engine/Source/Runtime/AssetRegistry/Private/AssetRegistryImpl.h
--- a/Engine/Source/Runtime/AssetRegistry/Private/AssetRegistryImpl.h
+++ b/Engine/Source/Runtime/AssetRegistry/Private/AssetRegistryImpl.h
@@ -54,6 +54,14 @@ namespace Impl
 		/** Report size of dynamic allocations. */
 		uint32 GetAllocatedSize() const;
 	};
+
+	/** Status of gathering, returned from the Tick function */
+	enum EGatherStatus : uint8
+	{
+		Active,
+		Complete,
+		UnableToProgress,
+	};
 }
 
 /**
@@ -125,8 +133,8 @@ public:
 	void TickDeletes();
 	/** Callback type for TickGatherer */
 	typedef TFunctionRef<void(const TRingBuffer<FAssetData*>&)> FAssetsFoundCallback;
-	/** Read any results from the gatherer, if it is running asynchronously */
-	void TickGatherer(Impl::FEventContext& EventContext, const double TickStartTime, bool& bOutIdle, bool& bOutInterrupted,
+	/** Consume any results from the gatherer and return its status */
+	Impl::EGatherStatus TickGatherer(Impl::FEventContext& EventContext, const double TickStartTime, bool& bOutInterrupted,
 		TOptional<FAssetsFoundCallback> AssetsFoundCallback = TOptional<FAssetsFoundCallback>());
 	/** Look for and load a single AssetData result from the gatherer. */
 	void TickGatherPackage(Impl::FEventContext& EventContext, const FString& PackageName, const FString& LocalPath);
@@ -288,8 +296,8 @@ private:
 	bool bInitialSearchStarted;
 	/** Flag to indicate if the initial background search has completed */
 	bool bInitialSearchCompleted;
-	/** Flag that indicates the background search is idle, so we can take actions when idle changes */
-	bool bGatherIdle;
+	/** Status of the background search, so we can take actions when it changes to or from idle */
+	Impl::EGatherStatus GatherStatus;
 
 	/**
 	 * Enables extra check to make sure path still mounted before adding.
@@ -402,7 +410,7 @@ struct FScanPathContext
 	int32 NumFoundAssets = 0;
 	bool bForceRescan = false;
 	bool bIgnoreDenyListScanFilters = false;
-	bool bIdle = false;
+	EGatherStatus Status = EGatherStatus::Active;
 };
 
 }
diff --git a/Engine/Source/Runtime/AssetRegistry/Public/AssetRegistry/IAssetRegistry.h b/Engine/Source/Runtime/AssetRegistry/Public/AssetRegistry/IAssetRegistry.h
--- a/Engine/Source/Runtime/AssetRegistry/Public/AssetRegistry/IAssetRegistry.h
+++ b/Engine/Source/Runtime/AssetRegistry/Public/AssetRegistry/IAssetRegistry.h
@@ -446,7 +446,10 @@ public:
 	UFUNCTION(BlueprintCallable, Category = "AssetRegistry")
 	virtual bool IsSearchAsync() const = 0;
 
-	/** Wait for scan to be complete */
+	/**
+	 * Wait for scan to be complete. If called during editor startup before OnPostEngineInit, and there are any assets that use classes in 
+	 * not-yet-loaded plugin modules, WaitForCompletion will return silently with those assets still ungathered.
+	 */
 	UFUNCTION(BlueprintCallable, Category = "AssetRegistry")
 	virtual void WaitForCompletion() = 0;
 
