AssetRegistry Gather: Enable even during synchronous scans the deferral rather than permanent failure for packages that use a custom version from a not-yet-loaded plugin.
#rnx

diff --git a/Engine/Source/Runtime/AssetRegistry/Private/AssetDataGatherer.cpp b/Engine/Source/Runtime/AssetRegistry/Private/AssetDataGatherer.cpp
--- a/Engine/Source/Runtime/AssetRegistry/Private/AssetDataGatherer.cpp
+++ b/Engine/Source/Runtime/AssetRegistry/Private/AssetDataGatherer.cpp
@@ -3001,10 +3001,19 @@ bool FAssetDataGatherer::ReadAssetFile(const FString& AssetFilename, TArray<FAss
 		// If we're missing a custom version, we might be able to load this package later once the module containing that version is loaded...
 		//   -	We can only attempt a retry in editors (not commandlets) that haven't yet finished initializing (!GIsRunning), as we 
 		//		have no guarantee that a commandlet or an initialized editor is going to load any more modules/plugins
-		//   -	Likewise, we can only attempt a retry for asynchronous scans, as during a synchronous scan we won't be loading any 
-		//		modules/plugins so it would last forever
-		const bool bAllowRetry = GIsEditor && !bInitialPluginsLoaded && !bIsSynchronousTick;
-		OutCanRetry = bAllowRetry && OpenPackageResult == FPackageReader::EOpenPackageResult::CustomVersionMissing;
+		const bool bAllowRetry = GIsEditor && !bInitialPluginsLoaded;
+		if (OpenPackageResult == FPackageReader::EOpenPackageResult::CustomVersionMissing)
+		{
+			OutCanRetry = bAllowRetry;
+			if (!bAllowRetry)
+			{
+				UE_LOG(LogAssetRegistry, Warning, TEXT("Package %s uses an unknown custom version and cannot be loaded for the AssetRegistry"), *AssetFilename);
+			}
+		}
+		else
+		{
+			OutCanRetry = false;
+		}
 		return false;
 	}
 	else
@@ -4141,4 +4150,4 @@ int32 FFilesToSearch::FTreeNode::NumFiles() const
 	return Num;
 }
 
-}
\ No newline at end of file
+}
