Downgrade AssetRegistry unknown custom version warning to a display until we have fixed the problem packages.
#rnx

diff --git a/Engine/Source/Runtime/AssetRegistry/Private/AssetDataGatherer.cpp b/Engine/Source/Runtime/AssetRegistry/Private/AssetDataGatherer.cpp
--- a/Engine/Source/Runtime/AssetRegistry/Private/AssetDataGatherer.cpp
+++ b/Engine/Source/Runtime/AssetRegistry/Private/AssetDataGatherer.cpp
@@ -3007,7 +3007,7 @@ bool FAssetDataGatherer::ReadAssetFile(const FString& AssetFilename, TArray<FAss
 			OutCanRetry = bAllowRetry;
 			if (!bAllowRetry)
 			{
-				UE_LOG(LogAssetRegistry, Warning, TEXT("Package %s uses an unknown custom version and cannot be loaded for the AssetRegistry"), *AssetFilename);
+				UE_LOG(LogAssetRegistry, Display, TEXT("Package %s uses an unknown custom version and cannot be loaded for the AssetRegistry"), *AssetFilename);
 			}
 		}
 		else
