AssetRegistry dependencies: Fix dependencies to packages that have CoreRedirects.
#rnx

diff --git a/Engine/Source/Runtime/AssetRegistry/Private/AssetRegistry.cpp b/Engine/Source/Runtime/AssetRegistry/Private/AssetRegistry.cpp
--- a/Engine/Source/Runtime/AssetRegistry/Private/AssetRegistry.cpp
+++ b/Engine/Source/Runtime/AssetRegistry/Private/AssetRegistry.cpp
@@ -3671,13 +3671,16 @@ void FAssetRegistryImpl::DependencyDataGathered(const double TickStartTime, TRin
 		check(Result.ImportUsedInGame.Num() == Result.ImportMap.Num());
 		for (int32 ImportIdx = 0; ImportIdx < Result.ImportMap.Num(); ++ImportIdx)
 		{
-			const FName AssetReference = Result.GetImportPackageName(ImportIdx);
+			FName AssetReference = Result.GetImportPackageName(ImportIdx);
 
 			// Should we skip this because it's too common?
 			if (ScriptPackagesToSkip.Contains(AssetReference))
 			{
 				continue;
 			}
+			FName Redirected = FCoreRedirects::GetRedirectedName(ECoreRedirectFlags::Type_Package,
+				FCoreRedirectObjectName(NAME_None, NAME_None, AssetReference)).PackageName;
+			AssetReference = Redirected;
 
 			EDependencyProperty DependencyProperty = EDependencyProperty::Build | EDependencyProperty::Hard;
 			DependencyProperty |= Result.ImportUsedInGame[ImportIdx] ? EDependencyProperty::Game : EDependencyProperty::None;
@@ -3688,6 +3691,9 @@ void FAssetRegistryImpl::DependencyDataGathered(const double TickStartTime, TRin
 		for (int32 SoftPackageIdx = 0; SoftPackageIdx < Result.SoftPackageReferenceList.Num(); ++SoftPackageIdx)
 		{
 			FName AssetReference = Result.SoftPackageReferenceList[SoftPackageIdx];
+			FName Redirected = FCoreRedirects::GetRedirectedName(ECoreRedirectFlags::Type_Package,
+				FCoreRedirectObjectName(NAME_None, NAME_None, AssetReference)).PackageName;
+			AssetReference = Redirected;
 
 			EDependencyProperty DependencyProperty = UE::AssetRegistry::EDependencyProperty::Build;
 			DependencyProperty |= (Result.SoftPackageUsedInGame[SoftPackageIdx] ? EDependencyProperty::Game : EDependencyProperty::None);
@@ -5423,4 +5429,4 @@ UAssetRegistryImpl::FFileLoadProgressUpdatedEvent& UAssetRegistryImpl::OnFileLoa
 {
 	checkf(IsInGameThread(), TEXT("Registering to AssetRegistry events is not supported from multiple threads."));
 	return FileLoadProgressUpdatedEvent;
-}
\ No newline at end of file
+}
