Added additional verbose logging for diagnosing plugin detection and asset registry scanning at runtime
- Added expected parameters to console command hints for asset registry

diff --git a/Engine/Source/Runtime/AssetRegistry/Private/AssetRegistry.cpp b/Engine/Source/Runtime/AssetRegistry/Private/AssetRegistry.cpp
--- a/Engine/Source/Runtime/AssetRegistry/Private/AssetRegistry.cpp
+++ b/Engine/Source/Runtime/AssetRegistry/Private/AssetRegistry.cpp
@@ -237,6 +237,7 @@ private:
 		Options.ParallelWorkers = FMath::Clamp(MaxWorkers, 0, 16);
 
 		const bool bLoaded = LoadAssetRegistry(GetPath(), Options, State);
+		UE_CLOG(bLoaded, LogAssetRegistry, Verbose, TEXT("AssetRegistry.bin loaded from '%s'"), GetPath())
 		checkf(bLoaded, TEXT("Failed to load %s"), GetPath());
 	}
 
@@ -3490,6 +3491,8 @@ void FAssetRegistryImpl::ScanPathsSynchronous(Impl::FScanPathContext& Context)
 
 			if (bIsInRequestedPaths)
 			{
+				UE_LOG(LogAssetRegistry, VeryVerbose, TEXT("FAssetRegistryImpl::ScanPathsSynchronous: Found Asset: %s"),
+					*AssetData->ObjectPath.ToString());	
 				Context.OutFoundAssets->Add(AssetData->ObjectPath);
 			}
 		}
@@ -3596,6 +3599,10 @@ void FAssetRegistryImpl::AssetSearchDataGathered(Impl::FEventContext& EventConte
 			// Populate the path tree
 			AddAssetPath(EventContext, PackagePath);
 		}
+		else
+		{
+			UE_LOG(LogAssetRegistry, Warning, TEXT("AssetRegistry: An asset has been loaded with an invalid mount point: '%s', Mount Point: '%s'"), *BackgroundResult->ObjectPath.ToString(), *PackagePathString)
+		}
 
 		// Check to see if we have run out of time in this tick
 		if (!bFlushFullBuffer && (FPlatformTime::Seconds() - TickStartTime) > Impl::MaxSecondsPerFrame)
@@ -3627,6 +3634,10 @@ void FAssetRegistryImpl::PathDataGathered(Impl::FEventContext& EventContext, con
 		{
 			AddAssetPath(EventContext, FName(*Path));
 		}
+		else
+		{
+			UE_LOG(LogAssetRegistry, Warning, TEXT("AssetRegistry: A path has been loaded with an invalid mount point: '%s'"), *Path)
+		}
 
 		// Check to see if we have run out of time in this tick
 		if (!bFlushFullBuffer && (FPlatformTime::Seconds() - TickStartTime) > Impl::MaxSecondsPerFrame)
diff --git a/Engine/Source/Runtime/AssetRegistry/Private/AssetRegistryConsoleCommands.h b/Engine/Source/Runtime/AssetRegistry/Private/AssetRegistryConsoleCommands.h
--- a/Engine/Source/Runtime/AssetRegistry/Private/AssetRegistryConsoleCommands.h
+++ b/Engine/Source/Runtime/AssetRegistry/Private/AssetRegistryConsoleCommands.h
@@ -28,27 +28,27 @@ public:
 	FAssetRegistryConsoleCommands()
 		: GetByNameCommand(
 		TEXT( "AssetRegistry.GetByName" ),
-		*LOCTEXT("CommandText_GetByName", "Query the asset registry for assets matching the supplied package name").ToString(),
+		*LOCTEXT("CommandText_GetByName", "<PackageName> //Query the asset registry for assets matching the supplied package name").ToString(),
 		FConsoleCommandWithArgsDelegate::CreateRaw( this, &FAssetRegistryConsoleCommands::GetByName ) )
 	,	GetByPathCommand(
 		TEXT( "AssetRegistry.GetByPath" ),
-		*LOCTEXT("CommandText_GetByPath", "Query the asset registry for assets matching the supplied package path").ToString(),
+		*LOCTEXT("CommandText_GetByPath", "<Path> //Query the asset registry for assets matching the supplied package path").ToString(),
 		FConsoleCommandWithArgsDelegate::CreateRaw( this, &FAssetRegistryConsoleCommands::GetByPath ) )
 	,	GetByClassCommand(
 		TEXT( "AssetRegistry.GetByClass" ),
-		*LOCTEXT("CommandText_GetByClass", "Query the asset registry for assets matching the supplied class").ToString(),
+		*LOCTEXT("CommandText_GetByClass", "<ClassName> //Query the asset registry for assets matching the supplied class").ToString(),
 		FConsoleCommandWithArgsDelegate::CreateRaw( this, &FAssetRegistryConsoleCommands::GetByClass ) )
 	,	GetByTagCommand(
 		TEXT( "AssetRegistry.GetByTag" ),
-		*LOCTEXT("CommandText_GetByTag", "Query the asset registry for assets matching the supplied tag and value").ToString(),
+		*LOCTEXT("CommandText_GetByTag", "<TagName> <TagValue> //Query the asset registry for assets matching the supplied tag and value").ToString(),
 		FConsoleCommandWithArgsDelegate::CreateRaw( this, &FAssetRegistryConsoleCommands::GetByTag ) )
 	,	GetDependenciesCommand(
 		TEXT( "AssetRegistry.GetDependencies" ),
-		*LOCTEXT("CommandText_GetDependencies", "Query the asset registry for dependencies for the specified package").ToString(),
+		*LOCTEXT("CommandText_GetDependencies", "<PackageName> //Query the asset registry for dependencies for the specified package").ToString(),
 		FConsoleCommandWithArgsDelegate::CreateRaw( this, &FAssetRegistryConsoleCommands::GetDependencies ) )
 	,	GetReferencersCommand(
 		TEXT( "AssetRegistry.GetReferencers" ),
-		*LOCTEXT("CommandText_GetReferencers", "Query the asset registry for referencers for the specified package").ToString(),
+		*LOCTEXT("CommandText_GetReferencers", "<ObjectPath> //Query the asset registry for referencers for the specified package").ToString(),
 		FConsoleCommandWithArgsDelegate::CreateRaw( this, &FAssetRegistryConsoleCommands::GetReferencers ) )
 	,	FindInvalidUAssetsCommand(
 		TEXT( "AssetRegistry.Debug.FindInvalidUAssets" ),
@@ -56,7 +56,7 @@ public:
 		FConsoleCommandWithArgsDelegate::CreateRaw( this, &FAssetRegistryConsoleCommands::FindInvalidUAssets ) )
 	, ScanPathCommand(
 		TEXT("AssetRegistry.ScanPath"),
-		*LOCTEXT("CommandText_ScanPath", "Scan the given filename or directoryname for package files and load them into the assetregistry. Extra string parameters: -forcerescan, -ignoreDenyLists, -asfile, -asdir").ToString(),
+		*LOCTEXT("CommandText_ScanPath", "<PathToScan> //Scan the given filename or directoryname for package files and load them into the assetregistry. Extra string parameters: -forcerescan, -ignoreDenyLists, -asfile, -asdir").ToString(),
 		FConsoleCommandWithArgsDelegate::CreateRaw(this, &FAssetRegistryConsoleCommands::ScanPath ) )
 	{}
 
diff --git a/Engine/Source/Runtime/Projects/Private/PluginManager.cpp b/Engine/Source/Runtime/Projects/Private/PluginManager.cpp
--- a/Engine/Source/Runtime/Projects/Private/PluginManager.cpp
+++ b/Engine/Source/Runtime/Projects/Private/PluginManager.cpp
@@ -189,6 +189,29 @@ const TSharedPtr<FJsonObject>& FPlugin::GetDescriptorJson()
 }
 #endif // WITH_EDITOR
 
+const FString& EnumToString(const EPluginType& InPluginType)
+{
+	// Const enum strings, special case no error.
+	static const FString Engine(TEXT("Engine"));
+	static const FString Enterprise(TEXT("Enterprise"));
+	static const FString Project(TEXT("Project"));
+	static const FString External(TEXT("External"));
+	static const FString Mod(TEXT("Mod"));
+	static const FString ErrorString(TEXT("Unknown plugin type, has EPluginType been changed?"));
+
+	switch (InPluginType)
+	{
+	case EPluginType::Engine: return Engine;
+	case EPluginType::Enterprise: return Enterprise;
+	case EPluginType::Project: return Project;
+	case EPluginType::External: return External;
+	case EPluginType::Mod: return Mod;
+	default:
+		check(!"Unknown plugin type, has EPluginType been changed?");
+		return ErrorString;
+	}
+}
+
 FPluginManager::FPluginManager()
 {
 	SCOPED_BOOT_TIMING("DiscoverAllPlugins");
@@ -338,6 +361,8 @@ void FPluginManager::ReadAllPlugins(TMap<FString, TSharedRef<FPlugin>>& Plugins,
 	// If we didn't find any manifests, do a recursive search for plugins
 	if (ManifestFileNames.Num() == 0)
 	{
+		UE_LOG(LogPluginManager, Verbose, TEXT("No *.upluginmanifest files found, looking for *.uplugin files instead."))
+		
 		// Find "built-in" plugins.  That is, plugins situated right within the Engine directory.
 		TArray<FString> EnginePluginDirs = FPaths::GetExtensionDirs(FPaths::EngineDir(), TEXT("Plugins"));
 		for (const FString& EnginePluginDir : EnginePluginDirs)
@@ -1038,7 +1063,8 @@ bool FPluginManager::ConfigureEnabledPlugins()
 			{
 				FString PlatformName = FPlatformProperties::PlatformName();
 				FPlugin& Plugin = *PluginsArray[Index];
-				UE_LOG(LogPluginManager, Log, TEXT("Mounting plugin %s"), *Plugin.GetName());
+				UE_LOG(LogPluginManager, Log, TEXT("Mounting %s plugin %s"), *EnumToString(Plugin.Type), *Plugin.GetName());
+				UE_LOG(LogPluginManager, Verbose, TEXT("Plugin path: %s"), *Plugin.FileName);
 
 				// Load <PluginName>.ini config file if it exists
 				FString PluginConfigDir = FPaths::GetPath(Plugin.FileName) / TEXT("Config/");
