AssetRegistry CachedBPInheritanceMap: Fix Blueprint instances not being registered if they are scanned before their Blueprint class from another plugin is scanned.

Fix submitted by [REDACTED_PERSON].
#rnx

diff --git a/Engine/Source/Runtime/AssetRegistry/Private/AssetRegistry.cpp b/Engine/Source/Runtime/AssetRegistry/Private/AssetRegistry.cpp
--- a/Engine/Source/Runtime/AssetRegistry/Private/AssetRegistry.cpp
+++ b/Engine/Source/Runtime/AssetRegistry/Private/AssetRegistry.cpp
@@ -885,22 +885,55 @@ void InitializeSerializationOptionsFromIni(FAssetRegistrySerializationOptions& O
 void FAssetRegistryImpl::CollectCodeGeneratorClasses()
 {
 	// Only refresh the list if our registered classes have changed
-	if (ClassGeneratorNamesRegisteredClassesVersionNumber != GetRegisteredClassesVersionNumber())
+	if (ClassGeneratorNamesRegisteredClassesVersionNumber == GetRegisteredClassesVersionNumber())
 	{
-		// Work around the fact we don't reference Engine module directly
-		UClass* BlueprintCoreClass = Cast<UClass>(StaticFindObject(UClass::StaticClass(), ANY_PACKAGE, TEXT("BlueprintCore")));
-		if (BlueprintCoreClass)
+		return;
+	}
+	ClassGeneratorNamesRegisteredClassesVersionNumber = GetRegisteredClassesVersionNumber();
+
+	// Work around the fact we don't reference Engine module directly
+	UClass* BlueprintCoreClass = Cast<UClass>(StaticFindObject(UClass::StaticClass(), ANY_PACKAGE, TEXT("BlueprintCore")));
+	if (!BlueprintCoreClass)
+	{
+		return;
+	}
+
+	ClassGeneratorNames.Add(BlueprintCoreClass->GetFName());
+
+	TArray<UClass*> BlueprintCoreDerivedClasses;
+	GetDerivedClasses(BlueprintCoreClass, BlueprintCoreDerivedClasses);
+	for (UClass* BPCoreClass : BlueprintCoreDerivedClasses)
+	{
+		bool bAlreadyRecorded;
+		FName BPCoreClassName = BPCoreClass->GetFName();
+		ClassGeneratorNames.Add(BPCoreClassName, &bAlreadyRecorded);
+		if (bAlreadyRecorded)
 		{
-			ClassGeneratorNames.Add(BlueprintCoreClass->GetFName());
+			continue;
+		}
 
-			TArray<UClass*> BlueprintCoreDerivedClasses;
-			GetDerivedClasses(BlueprintCoreClass, BlueprintCoreDerivedClasses);
-			for (UClass* BPCoreClass : BlueprintCoreDerivedClasses)
+		// For new generator classes, add all instances of them to CachedBPInheritanceMap. This is usually done
+		// when AddAssetData is called for those instances, but when we add a new generator class we have to recheck all
+		// instances of the class since they would have failed to detect they were Blueprint classes before.
+		// This can happen if blueprints in plugin B are scanned before their blueprint class from plugin A is scanned.
+		for (const FAssetData* AssetData : State.GetAssetsByClassName(BPCoreClassName))
+		{
+			const FString GeneratedClass = AssetData->GetTagValueRef<FString>(FBlueprintTags::GeneratedClassPath);
+			const FString ParentClass = AssetData->GetTagValueRef<FString>(FBlueprintTags::ParentClassPath);
+			if (!GeneratedClass.IsEmpty() && !ParentClass.IsEmpty())
 			{
-				ClassGeneratorNames.Add(BPCoreClass->GetFName());
+				const FName GeneratedClassFName = *Utils::ExportTextPathToObjectName(GeneratedClass);
+				const FName ParentClassFName = *Utils::ExportTextPathToObjectName(ParentClass);
+
+				if (!CachedBPInheritanceMap.Contains(GeneratedClassFName))
+				{
+					CachedBPInheritanceMap.Add(GeneratedClassFName, ParentClassFName);
+
+					// Invalidate caching because CachedBPInheritanceMap got modified
+					TempCachedInheritanceBuffer.bDirty = true;
+				}
 			}
 		}
-		ClassGeneratorNamesRegisteredClassesVersionNumber = GetRegisteredClassesVersionNumber();
 	}
 }
 
