Adding a way to directly add asset registry tags without having to recreate the entire AssetData object.

diff --git a/Engine/Source/Runtime/AssetRegistry/Private/AssetRegistryState.cpp b/Engine/Source/Runtime/AssetRegistry/Private/AssetRegistryState.cpp
--- a/Engine/Source/Runtime/AssetRegistry/Private/AssetRegistryState.cpp
+++ b/Engine/Source/Runtime/AssetRegistry/Private/AssetRegistryState.cpp
@@ -1596,6 +1596,38 @@ void FAssetRegistryState::AddAssetData(FAssetData* AssetData)
 	}
 }
 
+void FAssetRegistryState::AddTagsToAssetData(const FSoftObjectPath& InObjectPath, FAssetDataTagMap&& InTagsAndValues)
+{
+	if (InTagsAndValues.IsEmpty())
+	{
+		return;
+	}
+
+	FAssetData* AssetData = CachedAssetsByObjectPath.FindRef(InObjectPath.GetAssetPathName());
+	if (AssetData == nullptr)
+	{
+		UE_LOG(LogAssetRegistry, Error, TEXT("AddTagsToAssetData called with asset data that doesn't exist! Tags not added. ObjectPath: %s"), *InObjectPath.GetAssetPathName().ToString());
+		return;
+	}
+
+	// Update the tag cache map with the new tags.
+	for (auto TagIt : InTagsAndValues)
+	{
+		const FName FNameKey = TagIt.Key;
+
+		if (!AssetData->TagsAndValues.Contains(FNameKey))
+		{
+			TArray<FAssetData*>& NewTagAssets = CachedAssetsByTag.FindOrAdd(FNameKey);
+			NewTagAssets.Add(AssetData);
+		}
+	}
+
+	FAssetDataTagMap OldTags = AssetData->TagsAndValues.CopyMap();
+	OldTags.Append(MoveTemp(InTagsAndValues));
+	AssetData->TagsAndValues = FAssetDataTagMapSharedView(MoveTemp(OldTags));
+}
+
+
 void FAssetRegistryState::UpdateAssetData(const FAssetData& NewAssetData)
 {
 	if (FAssetData* AssetData = CachedAssetsByObjectPath.FindRef(NewAssetData.ObjectPath))
diff --git a/Engine/Source/Runtime/AssetRegistry/Public/AssetRegistry/AssetRegistryState.h b/Engine/Source/Runtime/AssetRegistry/Public/AssetRegistry/AssetRegistryState.h
--- a/Engine/Source/Runtime/AssetRegistry/Public/AssetRegistry/AssetRegistryState.h
+++ b/Engine/Source/Runtime/AssetRegistry/Public/AssetRegistry/AssetRegistryState.h
@@ -326,6 +326,9 @@ public:
 	/** Adds the asset data to the lookup maps */
 	void AddAssetData(FAssetData* AssetData);
 
+	/** Add the given tags/values to the asset data associated with the given object path, if it exists */
+	void AddTagsToAssetData(const FSoftObjectPath& InObjectPath, FAssetDataTagMap&& InTagsAndValues);
+
 	/** Finds an existing asset data based on object path and updates it with the new value and updates lookup maps */
 	void UpdateAssetData(const FAssetData& NewAssetData);
 
