- Fixed the cookededitor loading the wrong premade assetregistry

diff --git a/Engine/Source/Runtime/AssetRegistry/Private/AssetRegistry.cpp b/Engine/Source/Runtime/AssetRegistry/Private/AssetRegistry.cpp
--- a/Engine/Source/Runtime/AssetRegistry/Private/AssetRegistry.cpp
+++ b/Engine/Source/Runtime/AssetRegistry/Private/AssetRegistry.cpp
@@ -260,7 +260,13 @@ private:
 	{
 		if (Path.IsEmpty())
 		{
-			Path = FPaths::ProjectDir() + TEXT("AssetRegistry.bin");
+#if ASSETREGISTRY_ENABLE_PREMADE_REGISTRY_IN_EDITOR
+			Path = FPaths::ProjectDir() + TEXT("EditorClientAssetRegistry.bin");
+			if (!IFileManager::Get().FileExists(*Path))
+#endif
+			{
+				Path = FPaths::ProjectDir() + TEXT("AssetRegistry.bin");
+			}
 		}
 	
 		return *Path;
