FAssetDataGatherer/FPackageReader: Fix assets created by FPackageReader to match the LongPackageName the AssetDataGatherer has for them. This is important in cases of mounting ambiguity, such as when running UE5 without a project, and /Game and /Engine mount to the same localpath ../../../Engine/Content.
#rnx

diff --git a/Engine/Source/Runtime/AssetRegistry/Private/AssetDataGatherer.cpp b/Engine/Source/Runtime/AssetRegistry/Private/AssetDataGatherer.cpp
--- a/Engine/Source/Runtime/AssetRegistry/Private/AssetDataGatherer.cpp
+++ b/Engine/Source/Runtime/AssetRegistry/Private/AssetDataGatherer.cpp
@@ -2877,7 +2877,7 @@ void FAssetDataGatherer::TickInternal(bool& bOutIsTickInterrupt)
 				ReadContext.bCanceled = true;
 				return;
 			}
-			ReadContext.bResult = ReadAssetFile(ReadContext.AssetFileData.LocalAbsPath, ReadContext.AssetDataFromFile, ReadContext.DependencyData, ReadContext.CookedPackageNamesWithoutAssetData, ReadContext.bCanAttemptAssetRetry);
+			ReadContext.bResult = ReadAssetFile(ReadContext.AssetFileData.LongPackageName, ReadContext.AssetFileData.LocalAbsPath, ReadContext.AssetDataFromFile, ReadContext.DependencyData, ReadContext.CookedPackageNamesWithoutAssetData, ReadContext.bCanAttemptAssetRetry);
 		},
 		EParallelForFlags::Unbalanced | EParallelForFlags::BackgroundPriority
 	);
@@ -2991,7 +2991,7 @@ void FAssetDataGatherer::IngestDiscoveryResults()
 	Discovery->GetAndTrimSearchResults(bDiscoveryIsComplete, DiscoveredPaths, *FilesToSearch, NumPathsToSearchAtLastSyncPoint);
 }
 
-bool FAssetDataGatherer::ReadAssetFile(const FString& AssetFilename, TArray<FAssetData*>& AssetDataList, FPackageDependencyData& DependencyData, TArray<FString>& CookedPackageNamesWithoutAssetData, bool& OutCanRetry) const
+bool FAssetDataGatherer::ReadAssetFile(const FString& AssetLongPackageName, const FString& AssetFilename, TArray<FAssetData*>& AssetDataList, FPackageDependencyData& DependencyData, TArray<FString>& CookedPackageNamesWithoutAssetData, bool& OutCanRetry) const
 {
 	TRACE_CPUPROFILER_EVENT_SCOPE(FAssetDataGatherer::ReadAssetFile);
 	OutCanRetry = false;
@@ -3000,7 +3000,7 @@ bool FAssetDataGatherer::ReadAssetFile(const FString& AssetFilename, TArray<FAss
 	FPackageReader PackageReader;
 
 	FPackageReader::EOpenPackageResult OpenPackageResult;
-	if (!PackageReader.OpenPackageFile(AssetFilename, &OpenPackageResult))
+	if (!PackageReader.OpenPackageFile(AssetLongPackageName, AssetFilename, &OpenPackageResult))
 	{
 		// If we're missing a custom version, we might be able to load this package later once the module containing that version is loaded...
 		//   -	We can only attempt a retry in editors (not commandlets) that haven't yet finished initializing (!GIsRunning), as we 
diff --git a/Engine/Source/Runtime/AssetRegistry/Private/AssetDataGatherer.h b/Engine/Source/Runtime/AssetRegistry/Private/AssetDataGatherer.h
--- a/Engine/Source/Runtime/AssetRegistry/Private/AssetDataGatherer.h
+++ b/Engine/Source/Runtime/AssetRegistry/Private/AssetDataGatherer.h
@@ -210,7 +210,8 @@ private:
 	/**
 	 * Reads FAssetData information out of a file
 	 *
-	 * @param AssetFilename the name of the file to read
+	 * @param AssetLongPackageName the package name of the file to read
+	 * @param AssetFilename the local path of the file to read
 	 * @param AssetDataList the FAssetData for every asset found in the file
 	 * @param DependencyData the FPackageDependencyData for every asset found in the file
 	 * @param CookedPackagesToLoadUponDiscovery the list of cooked packages to be loaded if any
@@ -218,7 +219,8 @@ private:
 	 *
 	 * @return true if the file was successfully read
 	 */
-	bool ReadAssetFile(const FString& AssetFilename, TArray<FAssetData*>& AssetDataList, FPackageDependencyData& DependencyData,
+	bool ReadAssetFile(const FString& AssetLongPackageName, const FString& AssetFilename,
+		TArray<FAssetData*>& AssetDataList, FPackageDependencyData& DependencyData,
 		TArray<FString>& CookedPackagesToLoadUponDiscovery, bool& OutCanRetry) const;
 
 	/** Helper for LoadCacheFile and internal callsites; occurs inside the proper critical section. */
diff --git a/Engine/Source/Runtime/AssetRegistry/Private/PackageReader.cpp b/Engine/Source/Runtime/AssetRegistry/Private/PackageReader.cpp
--- a/Engine/Source/Runtime/AssetRegistry/Private/PackageReader.cpp
+++ b/Engine/Source/Runtime/AssetRegistry/Private/PackageReader.cpp
@@ -31,9 +31,15 @@ FPackageReader::~FPackageReader()
 	}
 }
 
-bool FPackageReader::OpenPackageFile(const FString& InPackageFilename, EOpenPackageResult* OutErrorCode)
+bool FPackageReader::OpenPackageFile(FStringView InPackageFilename, EOpenPackageResult* OutErrorCode)
+{
+	return OpenPackageFile(FStringView(), InPackageFilename, OutErrorCode);
+}
+
+bool FPackageReader::OpenPackageFile(FStringView InLongPackageName, FStringView InPackageFilename, EOpenPackageResult* OutErrorCode)
 {
 	check(!Loader);
+	LongPackageName = InLongPackageName;
 	PackageFilename = InPackageFilename;
 	Loader = IFileManager::Get().CreateFileReader(*PackageFilename);
 	bLoaderOwner = true;
@@ -45,6 +51,7 @@ bool FPackageReader::OpenPackageFile(FArchive* InLoader, EOpenPackageResult* Out
 	check(!Loader);
 	Loader = InLoader;
 	bLoaderOwner = false;
+	LongPackageName.Empty();
 	PackageFilename = Loader->GetArchiveName();
 	return OpenPackageFile(OutErrorCode);
 }
@@ -152,6 +159,20 @@ bool FPackageReader::OpenPackageFile(EOpenPackageResult* OutErrorCode)
 	SetPackageErrorCode(EOpenPackageResult::Success);
 	return true;
 }
+
+bool FPackageReader::TryGetLongPackageName(FString& OutLongPackageName) const
+{
+	if (!LongPackageName.IsEmpty())
+	{
+		OutLongPackageName = LongPackageName;
+		return true;
+	}
+	else
+	{
+		return FPackageName::TryConvertFilenameToLongPackageName(PackageFilename, OutLongPackageName);
+	}
+}
+
 bool FPackageReader::StartSerializeSection(int64 Offset)
 {
 	check(Loader);
@@ -182,7 +203,7 @@ bool FPackageReader::ReadAssetRegistryData(TArray<FAssetData*>& AssetDataList)
 
 	// Determine the package name and path
 	FString PackageName;
-	if (!FPackageName::TryConvertFilenameToLongPackageName(PackageFilename, PackageName))
+	if (!TryGetLongPackageName(PackageName))
 	{
 		// Path was possibly unmounted
 		return false;
@@ -278,7 +299,7 @@ bool FPackageReader::ReadAssetDataFromThumbnailCache(TArray<FAssetData*>& AssetD
 
 	// Determine the package name and path
 	FString PackageName;
-	if (!FPackageName::TryConvertFilenameToLongPackageName(PackageFilename, PackageName))
+	if (!TryGetLongPackageName(PackageName))
 	{
 		return false;
 	}
@@ -335,7 +356,7 @@ bool FPackageReader::ReadAssetRegistryDataIfCookedPackage(TArray<FAssetData*>& A
 	if (!!(GetPackageFlags() & PKG_FilterEditorOnly))
 	{
 		FString PackageName;
-		if (!FPackageName::TryConvertFilenameToLongPackageName(PackageFilename, PackageName))
+		if (!TryGetLongPackageName(PackageName))
 		{
 			return false;
 		}
@@ -403,7 +424,7 @@ bool FPackageReader::ReadAssetRegistryDataIfCookedPackage(TArray<FAssetData*>& A
 bool FPackageReader::ReadDependencyData(FPackageDependencyData& OutDependencyData)
 {
 	FString PackageNameString;
-	if (!FPackageName::TryConvertFilenameToLongPackageName(PackageFilename, PackageNameString))
+	if (!TryGetLongPackageName(PackageNameString))
 	{
 		// Path was possibly unmounted
 		return false;
diff --git a/Engine/Source/Runtime/AssetRegistry/Private/PackageReader.h b/Engine/Source/Runtime/AssetRegistry/Private/PackageReader.h
--- a/Engine/Source/Runtime/AssetRegistry/Private/PackageReader.h
+++ b/Engine/Source/Runtime/AssetRegistry/Private/PackageReader.h
@@ -37,7 +37,8 @@ public:
 	};
 
 	/** Creates a loader for the filename */
-	bool OpenPackageFile(const FString& PackageFilename, EOpenPackageResult* OutErrorCode = nullptr);
+	bool OpenPackageFile(FStringView PackageFilename, EOpenPackageResult* OutErrorCode = nullptr);
+	bool OpenPackageFile(FStringView LongPackageName, FStringView PackageFilename, EOpenPackageResult* OutErrorCode = nullptr);
 	bool OpenPackageFile(FArchive* Loader, EOpenPackageResult* OutErrorCode = nullptr);
 	bool OpenPackageFile(EOpenPackageResult* OutErrorCode = nullptr);
 
@@ -79,8 +80,10 @@ public:
 	}
 
 private:
+	bool TryGetLongPackageName(FString& OutLongPackageName) const;
 	bool StartSerializeSection(int64 Offset);
 
+	FString LongPackageName;
 	FString PackageFilename;
 	/* Loader is the interface used to read the bytes from the package's repository. All interpretation of the bytes is done by serializing into *this, which is also an FArchive. */
 	FArchive* Loader;
