AssetRegistry Premade Preloader - when in CookedEditor, allow deferring the preload after Initialize and move adding it to the AssestRegistry onto an async thread, for editor startup performance.
#rnx

diff --git a/Engine/Source/Runtime/AssetRegistry/Private/AssetRegistry.cpp b/Engine/Source/Runtime/AssetRegistry/Private/AssetRegistry.cpp
--- a/Engine/Source/Runtime/AssetRegistry/Private/AssetRegistry.cpp
+++ b/Engine/Source/Runtime/AssetRegistry/Private/AssetRegistry.cpp
@@ -125,21 +125,8 @@ FAssetRegistryInterface GAssetRegistryInterface;
 // Caching is permanently enabled in editor because memory is not that constrained, disabled by default otherwise
 #define ASSETREGISTRY_CACHE_ALWAYS_ENABLED (WITH_EDITOR)
 
-// Disable premade asset registry until the target specifies it in UBT
-#ifndef ASSETREGISTRY_ENABLE_PREMADE_REGISTRY_IN_EDITOR
-	#define ASSETREGISTRY_ENABLE_PREMADE_REGISTRY_IN_EDITOR 0
-#endif
-
 DEFINE_LOG_CATEGORY(LogAssetRegistry);
 
-#if ASSETREGISTRY_ENABLE_PREMADE_REGISTRY_IN_EDITOR 
-int32 LoadPremadeAssetRegistryInEditor = WITH_IOSTORE_IN_EDITOR ? 1 : 0;
-static FAutoConsoleVariableRef CVarLoadPremadeRegistryInEditor(
-	TEXT("AssetRegistry.LoadPremadeRegistryInEditor"),
-	LoadPremadeAssetRegistryInEditor,
-	TEXT(""));
-#endif
-
 static bool LoadAssetRegistry(const TCHAR* Path, const FAssetRegistryLoadOptions& Options, FAssetRegistryState& Out)
 {
 	check(Path);
@@ -160,125 +147,452 @@ static bool LoadAssetRegistry(const TCHAR* Path, const FAssetRegistryLoadOptions
 	return false;
 }
 
+namespace UE::AssetRegistry::Premade
+{
+
+/** Returns whether the given executable configuration supports AssetRegistry Preloading. Called before Main. */
+static bool IsEnabled()
+{
+	return FPlatformProperties::RequiresCookedData() || ASSETREGISTRY_ENABLE_PREMADE_REGISTRY_IN_EDITOR;
+}
+
+static bool CanLoadAsync()
+{
+	// TaskGraphSystemReady callback doesn't really mean it's running
+	return FPlatformProcess::SupportsMultithreading() && FTaskGraphInterface::IsRunning();
+}
+
+static bool CanConsumeAsync()
+{
+	// Async Consumption (copying the preloaded Premade FAssetRegistryState into the global AR) is only
+	// available when async is available, and is not currently used in the cooked game because we have
+	// to block on it anyway.
+	return CanLoadAsync() && !FPlatformProperties::RequiresCookedData();
+}
+
+/** Returns the paths to possible Premade AssetRegistry files, ordered from highest priority to lowest. */
+TArray<FString, TInlineAllocator<2>> GetPriorityPaths()
+{
+	TArray<FString, TInlineAllocator<2>> Paths;
+#if ASSETREGISTRY_ENABLE_PREMADE_REGISTRY_IN_EDITOR
+	Paths.Add(FPaths::Combine(FPaths::ProjectDir(), TEXT("EditorClientAssetRegistry.bin")));
+#endif
+	Paths.Add(FPaths::Combine(FPaths::ProjectDir(), TEXT("AssetRegistry.bin")));
+	return Paths;
+}
+
 // Loads cooked AssetRegistry.bin using an async preload task if available and sync otherwise
-static class FCookedAssetRegistryPreloader
+class FPreloader
 {
 public:
-	FCookedAssetRegistryPreloader()
-		: bLoadOnce(FPlatformProperties::RequiresCookedData() || ASSETREGISTRY_ENABLE_PREMADE_REGISTRY_IN_EDITOR)
+	enum class EConsumeResult
 	{
-		if (bLoadOnce)
+		Succeeded,
+		Failed,
+		Deferred
+	};
+	using FConsumeFunction = TFunction<void(bool bSucceeded, FAssetRegistryState&& ARState)>;
+
+	FPreloader()
+	{
+		if (UE::AssetRegistry::Premade::IsEnabled())
 		{
+			LoadState = EState::NotFound;
+
+			// run DelayedInitialize when TaskGraph system is ready
 			OnTaskGraphReady.Emplace(STATS ? EDelayedRegisterRunPhase::StatSystemReady :
 											 EDelayedRegisterRunPhase::TaskGraphSystemReady,
 				[this] () 
 				{
-					if (bLoadOnce && CanLoadAsync())
+				DelayedInitialize();
+			}); 
+		}
+	}
+	~FPreloader()
+	{
+		// We are destructed after Main exits, which means that our AsyncThread was either never called
+		// or it was waited on to complete by TaskGraph. Therefore we do not need to handle waiting for it ourselves.
+		Shutdown();
+	}
+
+	/**
+	 * Block on any pending async load, load if synchronous, and call ConsumeFunction with the results before returning.
+	 * If Consume has been called previously, the current ConsumeFunction is ignored and this call returns false.
+	 * 
+	 * @return Whether the load succeeded (this information is also passed to the ConsumeFunction).
+	 */
+	bool Consume(FConsumeFunction&& ConsumeFunction)
+	{
+		EConsumeResult Result = ConsumeInternal(MoveTemp(ConsumeFunction), FConsumeFunction());
+		check(Result != EConsumeResult::Deferred);
+		return Result == EConsumeResult::Succeeded;
+	}
+
+	/**
+	 * If a load is pending, store ConsumeAsynchronous for later calling and return EConsumeResult::Deferred.
+	 * If load is complete, or failed, or needs to run synchronously, load if necessary and call ConsumeSynchronous with results before returning.
+	 * Note if this function returns EConsumeResult::Deferred, the ConsumeAsynchronous will be called from another thread,
+	 * possibly before this call returns.
+	 * If Consume has been called previously, this call is ignored and returns EConsumeResult::Failed.
+	 *
+	 * @return Whether the load succeeded (this information is also passed to the ConsumeFunction).
+	 */
+	EConsumeResult ConsumeOrDefer(FConsumeFunction&& ConsumeSynchronous, FConsumeFunction&& ConsumeAsynchronous)
+	{
+		return ConsumeInternal(MoveTemp(ConsumeSynchronous), MoveTemp(ConsumeAsynchronous));
+	}
+
+private:
+	enum class EState : uint8
+	{
+		WillNeverPreload,
+		LoadSynchronous,
+		NotFound,
+		Loading,
+		Loaded,
+		Consumed,
+	};
+
+
+	bool TrySetPath()
+	{
+		for (FString& LocalPath : GetPriorityPaths())
+		{
+			if (IFileManager::Get().FileExists(*LocalPath))
+			{
+				ARPath = MoveTemp(LocalPath);
+				return true;
+			}
+		}
+		return false;
+	}
+
+	bool TrySetPath(const IPakFile& Pak)
+	{
+		for (FString& LocalPath : GetPriorityPaths())
+		{
+			if (Pak.PakContains(LocalPath))
+			{
+				ARPath = MoveTemp(LocalPath);
+				return true;
+			}
+		}
+		return false;
+	}
+
+	bool TryLoad()
+	{
+		checkf(!ARPath.IsEmpty(), TEXT("TryLoad must not be called until after TrySetPath has succeeded."));
+
+		FAssetRegistryLoadOptions Options;
+		const int32 ThreadReduction = 2; // This thread + main thread already has work to do 
+		int32 MaxWorkers = CanLoadAsync() ? FPlatformMisc::NumberOfCoresIncludingHyperthreads() - ThreadReduction : 0;
+		Options.ParallelWorkers = FMath::Clamp(MaxWorkers, 0, 16);
+		bLoadSucceeded = LoadAssetRegistry(*ARPath, Options, Payload);
+		UE_CLOG(!bLoadSucceeded, LogAssetRegistry, Warning, TEXT("Premade AssetRegistry path %s existed but failed to load."), *ARPath);
+		UE_CLOG(bLoadSucceeded, LogAssetRegistry, Log, TEXT("Premade AssetRegistry loaded from '%s'"), *ARPath);
+		return bLoadSucceeded;
+	}
+
+	void DelayedInitialize()
 					{
-						if (IFileManager::Get().FileExists(GetPath()))
+		// This function will run before any UObject (ie UAssetRegistryImpl) code can run, so we don't need to do any thread safety
+		// CanLoadAsync - we have to check this after the task graph is ready
+		if (!CanLoadAsync())
+		{
+			LoadState = EState::LoadSynchronous;
+			return;
+		}
+
+		// PreloadReady is in Triggered state until the Async thread is created. It is Reset in KickPreload.
+		PreloadReady = FPlatformProcess::GetSynchEventFromPool(true /* bIsManualReset */);
+		PreloadReady->Trigger();
+
+		if (TrySetPath())
 						{
 							KickPreload();
 						}
 						else
 						{
+			// set to NotFound, although PakMounted may set it to found later
+			LoadState = EState::NotFound;
+
 							// The PAK with the main registry isn't mounted yet
-							FCoreDelegates::OnPakFileMounted2.AddLambda([this](const IPakFile& Pak)
+			PakMountedDelegate = FCoreDelegates::OnPakFileMounted2.AddLambda([this](const IPakFile& Pak)
 								{
-									if (bLoadOnce && Pak.PakContains(GetPath()))
+					FScopeLock Lock(&StateLock);
+					if (LoadState == EState::NotFound && TrySetPath(Pak))
 									{
 										KickPreload();
+						// Remove the callback from OnPakFileMounted2 to avoid wasting time in all future PakFile mounts
+						// Do not access any of the lambda captures after the call to Remove, because deallocating the 
+						// DelegateHandle also deallocates our lambda captures
+						FDelegateHandle LocalPakMountedDelegate = PakMountedDelegate;
+						PakMountedDelegate.Reset();
+						FCoreDelegates::OnPakFileMounted2.Remove(LocalPakMountedDelegate);
 									}
 								});
 						}
 					}
-				});
+
+	void KickPreload()
+	{
+		// Called from Within the Lock
+		check(LoadState == EState::NotFound && !ARPath.IsEmpty());
+		LoadState = EState::Loading;
+		PreloadReady->Reset();
+		Async(EAsyncExecution::TaskGraph, [this]() { TryLoadAsync(); });
+	}
+
+	void TryLoadAsync()
+	{
+		// This function is active only after State has been set to Loading and PreloadReady has been Reset
+		// Until this function triggers PreloadReady, it has exclusive ownership of bLoadSucceeded and Payload
+		// Load outside the lock so that ConsumeOrDefer does not have to wait for the Load before it can defer and exit
+		bool bSucceeded = TryLoad();
+		// Trigger outside the lock so that a locked Consume function that is waiting on PreloadReady can wait inside the lock.
+		PreloadReady->Trigger();
+
+		FConsumeFunction LocalConsumeCallback;
+		{
+			FScopeLock Lock(&StateLock);
+			// The consume function may have woken up after the trigger and already consumed and changed LoadState to Consumed
+			if (LoadState == EState::Loading)
+			{
+				LoadState = EState::Loaded;
+				if (ConsumeCallback)
+				{
+					LocalConsumeCallback = MoveTemp(ConsumeCallback);
+					ConsumeCallback.Reset();
+					LoadState = EState::Consumed;
+				}
+			}
+		}
+
+		if (LocalConsumeCallback)
+		{
+			// No further threads will read/write payload at this point (until destructor, which is called after all async threads are complete
+			// so we can use it outside the lock
+			LocalConsumeCallback(bSucceeded, MoveTemp(Payload));
+			Shutdown();
 		}
 	}
 
-	bool Consume(FAssetRegistryState& Out)
+	EConsumeResult ConsumeInternal(FConsumeFunction&& ConsumeSynchronous, FConsumeFunction&& ConsumeAsynchronous)
 	{
 		SCOPED_BOOT_TIMING("FCookedAssetRegistryPreloader::Consume");
 
-		if (StateReady.IsValid())
+		FScopeLock Lock(&StateLock);
+		// Report failure if constructor decided not to preload or this has already been Consumed
+		if (LoadState == EState::WillNeverPreload || LoadState == EState::Consumed || ConsumeCallback)
 		{
-			StateReady.Wait();
+			Lock.Unlock(); // Unlock before calling external code in Consume callback
+			ConsumeSynchronous(false, FAssetRegistryState());
+			return EConsumeResult::Failed;
+		}
 
-			Out = MoveTemp(State);
-			return true;
+		if (LoadState == EState::LoadSynchronous)
+		{
+			bool bSucceeded = TrySetPath() && TryLoad();
+			LoadState = EState::Consumed;
+			Lock.Unlock(); // Unlock before calling external code in Consume callback
+			ConsumeSynchronous(bSucceeded, MoveTemp(Payload));
+			Shutdown(); // Shutdown can be called outside the lock since AsyncThread doesn't exist
+			return bSucceeded ? EConsumeResult::Succeeded : EConsumeResult::Failed;
 		}
-		else if (bLoadOnce && IFileManager::Get().FileExists(GetPath()))
+
+		// Cancel any further searching in Paks since we will no longer accept preloads starting after this point
+		FCoreDelegates::OnPakFileMounted2.Remove(PakMountedDelegate);
+		PakMountedDelegate.Reset();
+
+		if (ConsumeAsynchronous && LoadState == EState::Loading)
 		{
-			bLoadOnce = false;
-			Load();
+			// The load might have completed and the TryAsyncLoad thread is waiting to enter the lock, but we will still defer since Consume won the race
+			ConsumeCallback = MoveTemp(ConsumeAsynchronous);
+			return EConsumeResult::Deferred;
+		}
 
-			Out = MoveTemp(State);
-			return true;
+		{
+			SCOPED_BOOT_TIMING("BlockingConsume");
+			// If the load is in progress, wait for it to finish (which it does outside the lock)
+			PreloadReady->Wait();
 		}
 
-		return false;
+		// TryAsyncLoad might not yet have set state to Loaded
+		check(LoadState == EState::Loaded || LoadState == EState::Loading || LoadState == EState::NotFound);
+		bool bSucceeded = LoadState == EState::NotFound ? false : bLoadSucceeded;
+		LoadState = EState::Consumed;
+
+		// No further async threads exist that will read/write payload at this point so we can use it outside the lock
+		Lock.Unlock(); // Unlock before calling external code in Consume callback
+		ConsumeSynchronous(bSucceeded, MoveTemp(Payload));
+		Shutdown(); // Shutdown can be called outside the lock since we have set state to Consumed and the Async thread will notice and exit
+		return bSucceeded ? EConsumeResult::Succeeded : EConsumeResult::Failed;;
 	}
 
-	void CleanUp()
+	/** Called when the Preloader has no further work to do, to free resources early since destruction occurs at end of process. */
+	void Shutdown()
+	{
+		OnTaskGraphReady.Reset();
+		if (PreloadReady)
 	{
-		bLoadOnce = false;
-		StateReady.Reset();
+			FPlatformProcess::ReturnSynchEventToPool(PreloadReady);
+			PreloadReady = nullptr;
+	}
+		ARPath.Reset();
+		Payload.Reset();
 	}
 
-private:
-	static bool CanLoadAsync()
+	/** simple way to trigger a callback at a specific time that TaskGraph is usable. */
+	TOptional<FDelayedAutoRegisterHelper> OnTaskGraphReady;
+
+	/** Lock that guards members on this (see notes on each member). */
+	FCriticalSection StateLock;
+	/** Trigger for blocking Consume to wait upon TryLoadAsync. This Trigger is only allocated when in the states NotFound, Loaded, Loading. */
+	FEvent* PreloadReady = nullptr;
+
+	/** Path discovered for the AssetRegistry; Read/Write only within the Lock. */
+	FString ARPath;
+
+	/**
+	 * The ARState loaded from disk. Owned exclusively by either the first Consume or by TryAsyncLoad.
+	 * If LoadState is never set to Loading, this state is read/written only by the first thread to call Consume.
+	 * If LoadState is set to Loading (which happens before threading starts), the thread running TryAsyncLoad
+	 * owns this payload until it triggers PayloadReady, after which ownership returns to the first thread to call Consume.
+	 */
+	FAssetRegistryState Payload;
+
+	/** Delegate handle for the callback added to OnPakFileMounted2.After threading starts, Read / Write only within the lock. */
+	FDelegateHandle PakMountedDelegate;
+
+	/** Callback from ConsumeOrDefer that is set so TryLoadAsync can trigger the Consume when it completes.Read / Write only within the lock. */
+	FConsumeFunction ConsumeCallback;
+
+	/** State machine state. Read/Write only within the lock (or before threading starts). */
+	EState LoadState = EState::WillNeverPreload;
+
+	/** Result of TryLoad.Thread ownership rules are the same as the rules for Payload. */
+	bool bLoadSucceeded = false;
+}
+GPreloader;
+
+#if ASSETREGISTRY_ENABLE_PREMADE_REGISTRY_IN_EDITOR
+FAsyncConsumer::~FAsyncConsumer()
+{
+	if (Consumed)
 	{
-		// TaskGraphSystemReady callback doesn't really mean it's running
-		return FPlatformProcess::SupportsMultithreading() && FTaskGraphInterface::IsRunning();
+		FPlatformProcess::ReturnSynchEventToPool(Consumed);
+		Consumed = nullptr;
+	}
 	}
 
-	void Load()
+void FAsyncConsumer::PrepareForConsume()
+{
+	// Called within the lock
+	check(!Consumed);
+	Consumed = FPlatformProcess::GetSynchEventFromPool(true /* bIsManualReset */);
+	++ReferenceCount;
+};
+
+void FAsyncConsumer::Wait(UAssetRegistryImpl& UARI, FWriteScopeLock& ScopeLock)
+{
+	// Called within the lock
+	if (ReferenceCount == 0)
 	{
-		FAssetRegistryLoadOptions Options;
-		const int32 ThreadReduction = 2; // This thread + main thread already has work to do 
-		int32 MaxWorkers = CanLoadAsync() ? FPlatformMisc::NumberOfCoresIncludingHyperthreads() - ThreadReduction : 0;
-		Options.ParallelWorkers = FMath::Clamp(MaxWorkers, 0, 16);
+		return;
+	}
+	++ReferenceCount;
 
-		const bool bLoaded = LoadAssetRegistry(GetPath(), Options, State);
-		UE_CLOG(bLoaded, LogAssetRegistry, Verbose, TEXT("AssetRegistry.bin loaded from '%s'"), GetPath())
-		checkf(bLoaded, TEXT("Failed to load %s"), GetPath());
+	// Wait outside of the lock so that the AsyncThread can enter the lock to call Consume
+	{
+		UARI.InterfaceLock.WriteUnlock();
+		ON_SCOPE_EXIT{ UARI.InterfaceLock.WriteLock(); };
+		check(Consumed != nullptr);
+		Consumed->Wait();
 	}
 
-	void KickPreload()
+	--ReferenceCount;
+	if (ReferenceCount == 0)
 	{
-		check(!StateReady.IsValid() && bLoadOnce);
+		// We're the last one to drop the refcount, so delete Consumed
+		check(Consumed != nullptr);
+		FPlatformProcess::ReturnSynchEventToPool(Consumed);
+		Consumed = nullptr;
+	}
+}
 
-		bLoadOnce = false;
+void FAsyncConsumer::Consume(UAssetRegistryImpl& UARI, UE::AssetRegistry::Impl::FEventContext& EventContext, bool bSucceeded, FAssetRegistryState&& ARState)
+{
+	// Called within the lock
+	UARI.GuardedData.LoadPremadeAssetRegistry(EventContext, bSucceeded, MoveTemp(ARState), FAssetRegistryState::EInitializationMode::OnlyUpdateNew);
+	check(ReferenceCount >= 1);
+	check(Consumed != nullptr);
+	Consumed->Trigger();
+	--ReferenceCount;
+	if (ReferenceCount == 0)
+	{
+		// We're the last one to drop the refcount, so delete Consumed
+		FPlatformProcess::ReturnSynchEventToPool(Consumed);
+		Consumed = nullptr;
+	}
+}
 
-		StateReady = Async(EAsyncExecution::TaskGraph, [this]() { Load(); });
+#endif
 		
-		// Free FEvent held by TFuture in case of early shut down
-		TFunction<void()> OnTaskGraphShutdown = [this]() { CleanUp(); };
-		FTaskGraphInterface::Get().AddShutdownCallback(OnTaskGraphShutdown);
 	}
 
-	const TCHAR* GetPath()
+namespace UE::AssetRegistry
 	{
-		if (Path.IsEmpty())
+void FAssetRegistryImpl::ConditionalLoadPremadeAssetRegistry(UAssetRegistryImpl& UARI, Impl::FEventContext& EventContext, FWriteScopeLock& ScopeLock)
 		{
 #if ASSETREGISTRY_ENABLE_PREMADE_REGISTRY_IN_EDITOR
-			Path = FPaths::ProjectDir() + TEXT("EditorClientAssetRegistry.bin");
-			if (!IFileManager::Get().FileExists(*Path))
+	AsyncConsumer.Wait(UARI, ScopeLock);
 #endif
+}
+
+void FAssetRegistryImpl::ConsumeOrDeferPreloadedPremade(UAssetRegistryImpl& UARI, Impl::FEventContext& EventContext)
 			{
-				Path = FPaths::ProjectDir() + TEXT("AssetRegistry.bin");
+	// Called from inside WriteLock on InterfaceLock
+	using namespace UE::AssetRegistry::Premade;
+	if (!UE::AssetRegistry::Premade::IsEnabled())
+	{
+		// if we aren't doing any preloading, then we can set the initial search is done right away.
+		// Otherwise, it is set from LoadPremadeAssetRegistry
+		bCanFinishInitialSearch = true;
+		return;
 			}
+#if ASSETREGISTRY_ENABLE_PREMADE_REGISTRY_IN_EDITOR
+	FPreloader::FConsumeFunction ConsumeFromAsyncThread = [this, &UARI](bool bSucceeded, FAssetRegistryState&& ARState)
+	{
+		Impl::FEventContext EventContext;
+		{
+			FWriteScopeLock InterfaceScopeLock(UARI.InterfaceLock);
+			AsyncConsumer.Consume(UARI, EventContext, bSucceeded, MoveTemp(ARState));
+		}
+		UARI.Broadcast(EventContext);
+	};
+	auto ConsumeOnCurrentThread = [ConsumeFromAsyncThread](bool bSucceeded, FAssetRegistryState&& ARState) mutable
+	{
+		Async(EAsyncExecution::TaskGraph, [bSucceeded, ARState=MoveTemp(ARState), ConsumeFromAsyncThread=MoveTemp(ConsumeFromAsyncThread)]() mutable
+		{
+			ConsumeFromAsyncThread(bSucceeded, MoveTemp(ARState));
+		});
+	};
+	if (CanConsumeAsync())
+	{
+		AsyncConsumer.PrepareForConsume();
+		GPreloader.ConsumeOrDefer(MoveTemp(ConsumeOnCurrentThread), MoveTemp(ConsumeFromAsyncThread));
+	}
+	else
+#endif
+	{
+		GPreloader.Consume([this, &EventContext](bool bSucceeded, FAssetRegistryState&& ARState)
+			{
+				LoadPremadeAssetRegistry(EventContext, bSucceeded, MoveTemp(ARState), FAssetRegistryState::EInitializationMode::Rebuild);
+			});
 		}
-	
-		return *Path;
 	}
-	
-	bool bLoadOnce;
-	TOptional<FDelayedAutoRegisterHelper> OnTaskGraphReady;
-	FString Path;
-	TFuture<void> StateReady; // void since TFuture lack move support
-	FAssetRegistryState State;
 }
-GCookedAssetRegistryPreloader;
 
 /** Returns the appropriate ChunkProgressReportingType for the given Asset enum */
 EChunkProgressReportingType::Type GetChunkAvailabilityProgressType(EAssetAvailabilityProgressReportingType::Type ReportType)
@@ -315,6 +629,7 @@ namespace UE::AssetRegistry::Impl
 
 struct FInitializeContext
 {
+	UAssetRegistryImpl& UARI;
 	FEventContext Events;
 	TArray<FString> RootContentPaths;
 	bool bRedirectorsNeedSubscribe = false;
@@ -328,7 +643,7 @@ UAssetRegistryImpl::UAssetRegistryImpl(const FObjectInitializer& ObjectInitializ
 {
 	SCOPED_BOOT_TIMING("UAssetRegistryImpl::UAssetRegistryImpl");
 
-	UE::AssetRegistry::Impl::FInitializeContext Context;
+	UE::AssetRegistry::Impl::FInitializeContext Context{ *this };
 
 	{
 		FWriteScopeLock InterfaceScopeLock(InterfaceLock);
@@ -350,6 +665,54 @@ FAssetRegistryImpl::FAssetRegistryImpl()
 {
 }
 
+void FAssetRegistryImpl::LoadPremadeAssetRegistry(Impl::FEventContext& EventContext,
+	bool bSucceeded, FAssetRegistryState&& ARState, FAssetRegistryState::EInitializationMode Mode)
+{
+	SCOPED_BOOT_TIMING("LoadPremadeAssetRegistry");
+	if (SerializationOptions.bSerializeAssetRegistry)
+	{
+		if (bSucceeded)
+		{
+			if (Mode == FAssetRegistryState::EInitializationMode::Rebuild)
+			{
+				State = MoveTemp(ARState);
+				CachePathsFromState(EventContext, State);
+			}
+			else
+			{
+				State.InitializeFromExisting(ARState, SerializationOptions, Mode);
+				CachePathsFromState(EventContext, ARState);
+			}
+		}
+		else
+		{
+			UE_CLOG(FPlatformProperties::RequiresCookedData(), LogAssetRegistry, Fatal, TEXT("Failed to load premade asset registry."));
+		}
+	}
+
+	TArray<TSharedRef<IPlugin>> ContentPlugins = IPluginManager::Get().GetEnabledPluginsWithContent();
+	for (const TSharedRef<IPlugin>& ContentPlugin : ContentPlugins)
+	{
+		if (ContentPlugin->CanContainContent())
+		{
+			FArrayReader SerializedAssetData;
+			FString PluginAssetRegistry = ContentPlugin->GetBaseDir() / TEXT("AssetRegistry.bin");
+			if (IFileManager::Get().FileExists(*PluginAssetRegistry) && FFileHelper::LoadFileToArray(SerializedAssetData, *PluginAssetRegistry))
+			{
+				SerializedAssetData.Seek(0);
+				FAssetRegistryState PluginState;
+				PluginState.Load(SerializedAssetData);
+
+				State.InitializeFromExisting(PluginState, SerializationOptions, FAssetRegistryState::EInitializationMode::Append);
+				CachePathsFromState(EventContext, PluginState);
+			}
+		}
+	}
+
+	// let Tick know that it can finalize the initial search
+	bCanFinishInitialSearch = true;
+}
+
 void FAssetRegistryImpl::Initialize(Impl::FInitializeContext& Context)
 {
 	LLM_SCOPE(ELLMTag::AssetRegistry);
@@ -411,64 +774,8 @@ void FAssetRegistryImpl::Initialize(Impl::FInitializeContext& Context)
 			SearchAllAssetsInitialAsync(Context.Events);
 		}
 	}
-#if ASSETREGISTRY_ENABLE_PREMADE_REGISTRY_IN_EDITOR 
-	if (GIsEditor && !!LoadPremadeAssetRegistryInEditor)
-	{
-		// if the preload failed, load one now
-		if (GCookedAssetRegistryPreloader.Consume(/* Out */ State) == false)
-		{
-			SCOPED_BOOT_TIMING("ConsumeFailed");
-			FAssetRegistryLoadOptions Options;
-			const int32 ThreadReduction = 2; // This thread + main thread already has work to do 
-			const bool bCanLoadAsync = FPlatformProcess::SupportsMultithreading() && FTaskGraphInterface::IsRunning();
-			int32 MaxWorkers = bCanLoadAsync ? FPlatformMisc::NumberOfCoresIncludingHyperthreads() - ThreadReduction : 0;
-			Options.ParallelWorkers = FMath::Clamp(MaxWorkers, 0, 16);
 
-			if (LoadAssetRegistry(*(FPaths::ProjectDir() / TEXT("EditorClientAssetRegistry.bin")), Options, State))
-			{
-				UE_LOG(LogAssetRegistry, Log, TEXT("Loaded premade editor client asset registry"));
-			}
-			else if (LoadAssetRegistry(*(FPaths::ProjectDir() / TEXT("AssetRegistry.bin")), Options, State))
-			{
-				UE_LOG(LogAssetRegistry, Log, TEXT("Loaded premade asset registry"));
-			}
-			else
-			{
-				UE_LOG(LogAssetRegistry, Log, TEXT("Failed to load premade asset registry"));
-			}
-		}
-		CachePathsFromState(Context.Events, State);
-#else
-	else if (FPlatformProperties::RequiresCookedData())
-	{
-		if (SerializationOptions.bSerializeAssetRegistry &&
-			GCookedAssetRegistryPreloader.Consume(/* Out */ State))
-		{
-			CachePathsFromState(Context.Events, State);
-		}
-#endif // ASSETREGISTRY_ENABLE_PREMADE_REGISTRY_IN_EDITOR 
-
-		TArray<TSharedRef<IPlugin>> ContentPlugins = IPluginManager::Get().GetEnabledPluginsWithContent();
-		for (TSharedRef<IPlugin> ContentPlugin : ContentPlugins)
-		{
-			if (ContentPlugin->CanContainContent())
-			{
-				FArrayReader SerializedAssetData;
-				FString PluginAssetRegistry = ContentPlugin->GetBaseDir() / TEXT("AssetRegistry.bin");
-				if (IFileManager::Get().FileExists(*PluginAssetRegistry) && FFileHelper::LoadFileToArray(SerializedAssetData, *PluginAssetRegistry))
-				{
-					SerializedAssetData.Seek(0);
-					FAssetRegistryState PluginState;
-					PluginState.Load(SerializedAssetData);
-
-					State.InitializeFromExisting(PluginState, SerializationOptions, FAssetRegistryState::EInitializationMode::Append);
-					CachePathsFromState(Context.Events, PluginState);
-				}
-			}
-		}
-	}
-
-	GCookedAssetRegistryPreloader.CleanUp();
+	ConsumeOrDeferPreloadedPremade(Context.UARI, Context.Events);
 
 	// Report startup time. This does not include DirectoryWatcher startup time.
 	UE_LOG(LogAssetRegistry, Log, TEXT("FAssetRegistry took %0.4f seconds to start up"), FPlatformTime::Seconds() - StartupStartTime);
@@ -1090,6 +1397,11 @@ void UAssetRegistryImpl::SearchAllAssets(bool bSynchronousSearch)
 	FEventContext EventContext;
 	{
 		FWriteScopeLock InterfaceScopeLock(InterfaceLock);
+		if (bSynchronousSearch)
+		{
+			// make sure any outstanding async preload is complete
+			GuardedData.ConditionalLoadPremadeAssetRegistry(*this, EventContext, InterfaceScopeLock);
+		}
 		GuardedData.SearchAllAssets(EventContext, bSynchronousSearch);
 	}
 #if WITH_EDITOR
@@ -2639,6 +2951,8 @@ void UAssetRegistryImpl::ScanPathsSynchronousInternal(const TArray<FString>& InD
 
 	{
 		FWriteScopeLock InterfaceScopeLock(InterfaceLock);
+		// make sure any outstanding async preload is complete
+		GuardedData.ConditionalLoadPremadeAssetRegistry(*this, EventContext, InterfaceScopeLock);
 		GuardedData.ScanPathsSynchronous(Context);
 	}
 	if (Context.LocalPaths.IsEmpty())
@@ -3062,7 +3376,7 @@ Impl::EGatherStatus FAssetRegistryImpl::TickGatherer(Impl::FEventContext& EventC
 	{
 		HighestPending = 0;
 
-		if (!bInitialSearchCompleted)
+		if (!bInitialSearchCompleted && bCanFinishInitialSearch)
 		{
 #if WITH_EDITOR
 			// update redirectors
@@ -3844,6 +4158,7 @@ void FAssetRegistryImpl::DependencyDataGathered(const double TickStartTime, TRin
 				}
 			}
 		}
+		Node->SetIsDependenciesInitialized(true);
 
 		// Check to see if we have run out of time in this tick
 		if (!bFlushFullBuffer && (FPlatformTime::Seconds() - TickStartTime) > Impl::MaxSecondsPerFrame)
diff --git a/Engine/Source/Runtime/AssetRegistry/Private/AssetRegistry.h b/Engine/Source/Runtime/AssetRegistry/Private/AssetRegistry.h
--- a/Engine/Source/Runtime/AssetRegistry/Private/AssetRegistry.h
+++ b/Engine/Source/Runtime/AssetRegistry/Private/AssetRegistry.h
@@ -8,6 +8,7 @@
 #include "AssetRegistry.generated.h"
 
 class FRWScopeLock;
+namespace UE::AssetRegistry::Premade { struct FAsyncConsumer; }
 
 /**
  * The AssetRegistry singleton gathers information about .uasset files in the background so things
@@ -276,5 +277,8 @@ private:
 	FFileLoadProgressUpdatedEvent FileLoadProgressUpdatedEvent;
 
 	UE::AssetRegistry::Impl::FEventContext DeferredEvents;
+
+	friend class UE::AssetRegistry::FAssetRegistryImpl;
+	friend struct UE::AssetRegistry::Premade::FAsyncConsumer;
 };
 PRAGMA_ENABLE_DEPRECATION_WARNINGS
diff --git a/Engine/Source/Runtime/AssetRegistry/Private/AssetRegistryImpl.h b/Engine/Source/Runtime/AssetRegistry/Private/AssetRegistryImpl.h
--- a/Engine/Source/Runtime/AssetRegistry/Private/AssetRegistryImpl.h
+++ b/Engine/Source/Runtime/AssetRegistry/Private/AssetRegistryImpl.h
@@ -14,11 +14,52 @@
 #include "Templates/Function.h"
 #include "Templates/UniquePtr.h"
 
+// Disable premade asset registry until the target specifies it in UBT
+#ifndef ASSETREGISTRY_ENABLE_PREMADE_REGISTRY_IN_EDITOR
+#define ASSETREGISTRY_ENABLE_PREMADE_REGISTRY_IN_EDITOR 0
+#endif
+
 class FDependsNode;
 struct FARFilter;
 class FAssetDataGatherer;
 struct FFileChangeData;
 class FPackageReader;
+class UAssetRegistryImpl;
+namespace UE::AssetRegistry::Impl { struct FEventContext; }
+namespace UE::AssetRegistry::Premade { struct FAsyncConsumer; }
+
+#if ASSETREGISTRY_ENABLE_PREMADE_REGISTRY_IN_EDITOR
+namespace UE::AssetRegistry::Premade
+{
+
+/**
+ * A struct to consume a Premade asset registry on an async thread. It supports a cheap Wait() call so that it can be Waited on
+ * by frequent AssetRegistry interface calls for the rest of the process.
+ */
+struct FAsyncConsumer
+{
+	~FAsyncConsumer();
+	/** Sets the Consumer into need-to-wait mode; Wait will block until Consume is called. */
+	void PrepareForConsume();
+	/**
+	 * Does not return until after the Premade AssetRegistryState has been added into the target AssetRegistry, or it has been decided not to be used.
+	 * For performance reasons, this must be called within the WriteScopeLock, and the caller must handle the possibility that it leaves and reenters the lock.
+	 */
+	void Wait(UAssetRegistryImpl& UARI, FWriteScopeLock& ScopeLock);
+	/** Callback from the async thread to consume the Premade ARState */
+	void Consume(UAssetRegistryImpl& UARI, UE::AssetRegistry::Impl::FEventContext& EventContext, bool bSucceeded, FAssetRegistryState&& ARState);
+
+private:
+	/**
+	 * ReferenceCounter for the Consumed event. Also used to decide whether Waiting is necessary. Read/Write only inside the lock.
+	 */
+	int32 ReferenceCount = 0;
+	/** Event used to Wait for Consume. Allocated/deallocated within the lock. Can be waited on outside the lock if RefCount is held. */
+	FEvent* Consumed = nullptr;
+};
+
+}
+#endif
 
 namespace UE
 {
@@ -27,7 +68,6 @@ namespace AssetRegistry
 
 namespace Impl
 {
-	struct FEventContext;
 	struct FInitializeContext;
 	struct FScanPathContext;
 	struct FClassInheritanceContext;
@@ -183,6 +223,11 @@ public:
 	/** Finds all class names of classes capable of generating new UClasses */
 	void CollectCodeGeneratorClasses();
 
+	/**
+	 * Block until the Premade AssetRegistry finishes loading, if there is one still loading.
+	 * For performance reasons, this must be called within the WriteScopeLock, and the caller must handle the possibility that it leaves and reenters the lock.
+	 */
+	void ConditionalLoadPremadeAssetRegistry(UAssetRegistryImpl& UARI, UE::AssetRegistry::Impl::FEventContext& EventContext, FWriteScopeLock& ScopeLock);
 
 private:
 
@@ -241,6 +286,12 @@ private:
 	/** Updates OutBuffer from loaded classes and registered-in-CachedBPInheritanceMap blueprint classes */
 	void UpdateInheritanceBuffer(Impl::FClassInheritanceBuffer& OutBuffer) const;
 
+	/** Conditionally loads the Premade AssetRegistry from disk, or queues it for asynchronous loading if not yet loaded. */
+	void ConsumeOrDeferPreloadedPremade(UAssetRegistryImpl& UARI, Impl::FEventContext& EventContext);
+	/** Moves a premade asset registry state into this AR */
+	void LoadPremadeAssetRegistry(Impl::FEventContext& Context,
+		bool bSucceeded, FAssetRegistryState&& ARState, FAssetRegistryState::EInitializationMode Mode);
+
 private:
 
 	/** Internal state of the cached asset registry */
@@ -296,6 +347,8 @@ private:
 	bool bInitialSearchStarted;
 	/** Flag to indicate if the initial background search has completed */
 	bool bInitialSearchCompleted;
+	/** Flag to indicate if the initial background search can be finalized */
+	bool bCanFinishInitialSearch = false;
 	/** Status of the background search, so we can take actions when it changes to or from idle */
 	Impl::EGatherStatus GatherStatus;
 
@@ -336,7 +389,10 @@ private:
 	/** Class names that return true for IsAsset but which should not be treated as assets in cooked packages */
 	TSet<FName> SkipCookedClasses;
 #endif
-
+#if ASSETREGISTRY_ENABLE_PREMADE_REGISTRY_IN_EDITOR
+	UE::AssetRegistry::Premade::FAsyncConsumer AsyncConsumer;
+#endif
+	friend struct UE::AssetRegistry::Premade::FAsyncConsumer;
 	friend struct Impl::FClassInheritanceContext;
 };
 
diff --git a/Engine/Source/Runtime/AssetRegistry/Private/AssetRegistryState.cpp b/Engine/Source/Runtime/AssetRegistry/Private/AssetRegistryState.cpp
--- a/Engine/Source/Runtime/AssetRegistry/Private/AssetRegistryState.cpp
+++ b/Engine/Source/Runtime/AssetRegistry/Private/AssetRegistryState.cpp
@@ -175,7 +175,8 @@ void FAssetRegistryState::FilterTags(const FAssetDataTagMapSharedView& InTagsAnd
 	}
 }
 
-void FAssetRegistryState::InitializeFromExistingAndPrune(const FAssetRegistryState & ExistingState, const TSet<FName>& RequiredPackages, const TSet<FName>& RemovePackages, const TSet<int32> ChunksToKeep, const FAssetRegistrySerializationOptions & Options)
+void FAssetRegistryState::InitializeFromExistingAndPrune(const FAssetRegistryState & ExistingState, const TSet<FName>& RequiredPackages, const TSet<FName>& RemovePackages,
+	const TSet<int32> ChunksToKeep, const FAssetRegistrySerializationOptions& Options)
 {
 	LLM_SCOPE(ELLMTag::AssetRegistry);
 	const bool bIsFilteredByChunkId = ChunksToKeep.Num() != 0;
@@ -296,6 +297,7 @@ void FAssetRegistryState::InitializeFromExistingAndPrune(const FAssetRegistrySta
 				NewDependency->AddReferencer(NewNode);
 			}
 		});
+		NewNode->SetIsDependenciesInitialized(true);
 	}
 
 	// Remove any orphaned depends nodes. This will leave cycles in but those might represent useful data
@@ -318,7 +320,8 @@ void FAssetRegistryState::InitializeFromExistingAndPrune(const FAssetRegistrySta
 	}
 }
 
-void FAssetRegistryState::InitializeFromExisting(const TMap<FName, FAssetData*>& AssetDataMap, const TMap<FAssetIdentifier, FDependsNode*>& DependsNodeMap, const TMap<FName, FAssetPackageData*>& AssetPackageDataMap, const FAssetRegistrySerializationOptions& Options, EInitializationMode InInitializationMode)
+void FAssetRegistryState::InitializeFromExisting(const TMap<FName, FAssetData*>& AssetDataMap, const TMap<FAssetIdentifier, FDependsNode*>& DependsNodeMap,
+	const TMap<FName, FAssetPackageData*>& AssetPackageDataMap, const FAssetRegistrySerializationOptions& Options, EInitializationMode InInitializationMode)
 {
 	LLM_SCOPE(ELLMTag::AssetRegistry);
 	if (InInitializationMode == EInitializationMode::Rebuild)
@@ -343,6 +346,10 @@ void FAssetRegistryState::InitializeFromExisting(const TMap<FName, FAssetData*>&
 		{
 			continue;
 		}
+		if (InInitializationMode == EInitializationMode::OnlyUpdateNew && ExistingData != nullptr)
+		{
+			continue;
+		}
 
 		// Filter asset registry tags now
 		const FAssetData& AssetData = *Pair.Value;
@@ -350,8 +357,6 @@ void FAssetRegistryState::InitializeFromExisting(const TMap<FName, FAssetData*>&
 		FAssetDataTagMap LocalTagsAndValues;
 		FAssetRegistryState::FilterTags(AssetData.TagsAndValues, LocalTagsAndValues, Options.CookFilterlistTagsByClass.Find(AssetData.AssetClass), Options);
 		
-		// in append or onlyupdateexisting we may have some existing data
-		// in rebuild mode existing data should never be found
 		if (ExistingData)
 		{
 			// Bundle tags might have changed even if other tags haven't
@@ -382,7 +387,10 @@ void FAssetRegistryState::InitializeFromExisting(const TMap<FName, FAssetData*>&
 		for (const TPair<FName, FAssetPackageData*>& Pair : AssetPackageDataMap)
 		{
 			bool bIsScriptPackage = FPackageName::IsScriptPackage(Pair.Key.ToString());
-
+			if (InInitializationMode == EInitializationMode::OnlyUpdateNew && CachedPackageData.Find(Pair.Key))
+			{
+				continue;
+			}
 			if (Pair.Value)
 			{
 				// Only add if also in asset data map, or script package
@@ -401,29 +409,55 @@ void FAssetRegistryState::InitializeFromExisting(const TMap<FName, FAssetData*>&
 			}
 		}
 
-		for (const TPair<FAssetIdentifier, FDependsNode*>& Pair : DependsNodeMap)
+		TMap<FAssetIdentifier, FDependsNode*> FilteredDependsNodeMap;
+		const TMap<FAssetIdentifier, FDependsNode*>* DependsNodesToAdd = &DependsNodeMap;
+		if (InInitializationMode == EInitializationMode::OnlyUpdateNew)
+		{
+			// Keep the original DependsNodeMap for reference, but remove from NodesToAdd all nodes that already have dependency data
+			// Also reserve up-front all (unfiltered) nodes we are adding, to avoid reallocating the Referencers array.
+			FilteredDependsNodeMap.Reserve(DependsNodeMap.Num());
+			DependsNodesToAdd = &FilteredDependsNodeMap;
+			for (const TPair<FAssetIdentifier, FDependsNode*>& Pair : DependsNodeMap)
+			{
+				FDependsNode* SourceNode = Pair.Value;
+				FDependsNode* TargetNode = CreateOrFindDependsNode(Pair.Key);
+				if (!TargetNode->IsDependenciesInitialized())
+				{
+					FilteredDependsNodeMap.Add(Pair.Key, SourceNode);
+				}
+				TargetNode->Reserve(SourceNode);
+			}
+		}
+		else
 		{
-			FDependsNode* OldNode = Pair.Value;
-			FDependsNode* NewNode = CreateOrFindDependsNode(Pair.Key);
-			NewNode->Reserve(OldNode);
+			// Reserve up-front all the nodes that we are adding, so we do not reallocate
+			// the Referencers array multiple times on a node as we add nodes that refer to it
+			for (const TPair<FAssetIdentifier, FDependsNode*>& Pair : DependsNodeMap)
+			{
+				FDependsNode* SourceNode = Pair.Value;
+				FDependsNode* TargetNode = CreateOrFindDependsNode(Pair.Key);
+				TargetNode->Reserve(SourceNode);
+			}
 		}
 
-		for (const TPair<FAssetIdentifier, FDependsNode*>& Pair : DependsNodeMap)
+		for (const TPair<FAssetIdentifier, FDependsNode*>& Pair : *DependsNodesToAdd)
 		{
-			FDependsNode* OldNode = Pair.Value;
-			FDependsNode* NewNode = CreateOrFindDependsNode(Pair.Key);
-			Pair.Value->IterateOverDependencies([this, &DependsNodeMap, &ScriptPackages, OldNode, NewNode](FDependsNode* InDependency, UE::AssetRegistry::EDependencyCategory InCategory, UE::AssetRegistry::EDependencyProperty InFlags, bool bDuplicate) {
+			FDependsNode* SourceNode = Pair.Value;
+			FDependsNode* TargetNode = CreateOrFindDependsNode(Pair.Key);
+			SourceNode->IterateOverDependencies([this, &DependsNodeMap, &ScriptPackages, TargetNode]
+			(FDependsNode* InDependency, UE::AssetRegistry::EDependencyCategory InCategory, UE::AssetRegistry::EDependencyProperty InFlags, bool bDuplicate) {
 				const FAssetIdentifier& Identifier = InDependency->GetIdentifier();
 				if (DependsNodeMap.Find(Identifier) || ScriptPackages.Contains(Identifier))
 				{
 					// Only add if this node is in the incoming map
-					FDependsNode* NewDependency = CreateOrFindDependsNode(Identifier);
-					NewNode->SetIsDependencyListSorted(InCategory, false);
-					NewNode->AddDependency(NewDependency, InCategory, InFlags);
-					NewDependency->SetIsReferencersSorted(false);
-					NewDependency->AddReferencer(NewNode);
+					FDependsNode* TargetDependency = CreateOrFindDependsNode(Identifier);
+					TargetNode->SetIsDependencyListSorted(InCategory, false);
+					TargetNode->AddDependency(TargetDependency, InCategory, InFlags);
+					TargetDependency->SetIsReferencersSorted(false);
+					TargetDependency->AddReferencer(TargetDependency);
 				}
 			});
+			TargetNode->SetIsDependenciesInitialized(true);
 		}
 
 		// Restore the sortedness that we turned off for performance when creating each DependsNode
@@ -2381,4 +2415,4 @@ void FAssetRegistryState::Dump(const TArray<FString>& Arguments, TArray<FString>
 	}
 }
 
-#endif // ASSET_REGISTRY_STATE_DUMPING_ENABLED
+#endif // ASSET_REGISTRY_STATE_DUMPING_ENABLED
\ No newline at end of file
diff --git a/Engine/Source/Runtime/AssetRegistry/Private/DependsNode.cpp b/Engine/Source/Runtime/AssetRegistry/Private/DependsNode.cpp
--- a/Engine/Source/Runtime/AssetRegistry/Private/DependsNode.cpp
+++ b/Engine/Source/Runtime/AssetRegistry/Private/DependsNode.cpp
@@ -702,6 +702,8 @@ void FDependsNode::SerializeLoad(FArchive& Ar, const TUniqueFunction<FDependsNod
 	ReadDependencies(NameDependencies, nullptr, 0);
 	ReadDependencies(ManageDependencies, &ManageFlags, ManageFlagSetWidth);
 	ReadDependencies(Referencers, nullptr, 0);
+
+	SetIsDependenciesInitialized(true);
 }
 
 void FDependsNode::SerializeLoad_BeforeFlags(FArchive& Ar, FAssetRegistryVersion::Type Version, FDependsNode* PreallocatedDependsNodeDataBuffer, int32 NumDependsNodes, bool bSerializeDependencies,
@@ -771,6 +773,8 @@ void FDependsNode::SerializeLoad_BeforeFlags(FArchive& Ar, FAssetRegistryVersion
 	SerializeNodeArray(NumSoftManage, ManageDependencies, &ManageFlags, ManageFlagSetWidth, SoftManageBits, true, bSerializeDependencies);
 	SerializeNodeArray(NumHardManage, ManageDependencies, &ManageFlags, ManageFlagSetWidth, HardManageBits, true, bSerializeDependencies);
 	SerializeNodeArray(NumReferencers, Referencers, nullptr, 0, 0, true, true);
+
+	SetIsDependenciesInitialized(true);
 }
 
 void FDependsNode::GetPropertySetBits_BeforeFlags(uint32& HardBits, uint32& SoftBits, uint32& HardManageBits, uint32& SoftManageBits)
@@ -941,3 +945,13 @@ void FDependsNode::SetIsReferencersSorted(bool bValue)
 	}
 	ReferencersIsSorted = bValue ? 1 : 0;
 }
+
+bool FDependsNode::IsDependenciesInitialized() const
+{
+	return DependenciesInitialized;
+}
+
+void FDependsNode::SetIsDependenciesInitialized(bool bValue)
+{
+	DependenciesInitialized = bValue ? 1 : 0;
+}
diff --git a/Engine/Source/Runtime/AssetRegistry/Private/DependsNode.h b/Engine/Source/Runtime/AssetRegistry/Private/DependsNode.h
--- a/Engine/Source/Runtime/AssetRegistry/Private/DependsNode.h
+++ b/Engine/Source/Runtime/AssetRegistry/Private/DependsNode.h
@@ -158,6 +158,8 @@ public:
 	void SetIsDependencyListSorted(UE::AssetRegistry::EDependencyCategory Category, bool bValue);
 	bool IsReferencersSorted() const;
 	void SetIsReferencersSorted(bool bValue);
+	bool IsDependenciesInitialized() const;
+	void SetIsDependenciesInitialized(bool bValue);
 
 private:
 
@@ -172,6 +174,7 @@ private:
 		SearchableNameIsSorted = 1;
 		ManageIsSorted = 1;
 		ReferencersIsSorted = 1;
+		DependenciesInitialized = 0;
 	}
 
 	/** The name of the package/object this node represents */
@@ -186,4 +189,5 @@ private:
 	uint32 SearchableNameIsSorted : 1;
 	uint32 ManageIsSorted : 1;
 	uint32 ReferencersIsSorted : 1;
+	uint32 DependenciesInitialized : 1;
 };
diff --git a/Engine/Source/Runtime/AssetRegistry/Public/AssetRegistry/AssetRegistryState.h b/Engine/Source/Runtime/AssetRegistry/Public/AssetRegistry/AssetRegistryState.h
--- a/Engine/Source/Runtime/AssetRegistry/Public/AssetRegistry/AssetRegistryState.h
+++ b/Engine/Source/Runtime/AssetRegistry/Public/AssetRegistry/AssetRegistryState.h
@@ -117,7 +117,8 @@ public:
 	{
 		Rebuild,
 		OnlyUpdateExisting,
-		Append
+		Append,
+		OnlyUpdateNew,
 	};
 
 	/**
