ReferenceViewer: Fix display of packages with multiple assets to display the UAsset. This fixes redirector packages to Blueprints from displaying the redirector to the BPGC rather than the redirector to the Blueprint.
#rnx

diff --git a/Engine/Plugins/Editor/AssetManagerEditor/Source/AssetManagerEditor/Private/ReferenceViewer/EdGraph_ReferenceViewer.cpp b/Engine/Plugins/Editor/AssetManagerEditor/Source/AssetManagerEditor/Private/ReferenceViewer/EdGraph_ReferenceViewer.cpp
--- a/Engine/Plugins/Editor/AssetManagerEditor/Source/AssetManagerEditor/Private/ReferenceViewer/EdGraph_ReferenceViewer.cpp
+++ b/Engine/Plugins/Editor/AssetManagerEditor/Source/AssetManagerEditor/Private/ReferenceViewer/EdGraph_ReferenceViewer.cpp
@@ -544,22 +544,7 @@ int32 UEdGraph_ReferenceViewer::RecursivelyGatherSizes(bool bReferencers, const
 
 void UEdGraph_ReferenceViewer::GatherAssetData(const TSet<FName>& AllPackageNames, TMap<FName, FAssetData>& OutPackageToAssetDataMap) const
 {
-	// Take a guess to find the asset instead of searching for it. Most packages have a single asset in them with the same name as the package.
-	FAssetRegistryModule& AssetRegistryModule = FModuleManager::LoadModuleChecked<FAssetRegistryModule>("AssetRegistry");
-
-	FARFilter Filter;
-	for ( auto PackageIt = AllPackageNames.CreateConstIterator(); PackageIt; ++PackageIt )
-	{
-		const FString& PackageName = (*PackageIt).ToString();
-		Filter.PackageNames.Add(*PackageIt);
-	}
-
-	TArray<FAssetData> AssetDataList;
-	AssetRegistryModule.Get().GetAssets(Filter, AssetDataList);
-	for ( auto AssetIt = AssetDataList.CreateConstIterator(); AssetIt; ++AssetIt )
-	{
-		OutPackageToAssetDataMap.Add((*AssetIt).PackageName, *AssetIt);
-	}
+	UE::AssetRegistry::GetAssetForPackages(AllPackageNames.Array(), OutPackageToAssetDataMap);
 }
 
 UEdGraphNode_Reference* UEdGraph_ReferenceViewer::RecursivelyConstructNodes(bool bReferencers, UEdGraphNode_Reference* RootNode, const TArray<FAssetIdentifier>& Identifiers, const FIntPoint& NodeLoc, const TMap<FAssetIdentifier, int32>& NodeSizes, const TMap<FName, FAssetData>& PackagesToAssetDataMap, const TSet<FName>& AllowedPackageNames, int32 CurrentDepth, TSet<FAssetIdentifier>& VisitedNames)
diff --git a/Engine/Source/Runtime/AssetRegistry/Private/AssetRegistry.cpp b/Engine/Source/Runtime/AssetRegistry/Private/AssetRegistry.cpp
--- a/Engine/Source/Runtime/AssetRegistry/Private/AssetRegistry.cpp
+++ b/Engine/Source/Runtime/AssetRegistry/Private/AssetRegistry.cpp
@@ -5809,3 +5809,57 @@ UAssetRegistryImpl::FFileLoadProgressUpdatedEvent& UAssetRegistryImpl::OnFileLoa
 	checkf(IsInGameThread(), TEXT("Registering to AssetRegistry events is not supported from multiple threads."));
 	return FileLoadProgressUpdatedEvent;
 }
+
+
+namespace UE::AssetRegistry
+{
+
+void GetAssetForPackages(TConstArrayView<FName> PackageNames, TMap<FName, FAssetData>& OutPackageToAssetData)
+{
+	FARFilter Filter;
+	for (FName PackageName : PackageNames)
+	{
+		Filter.PackageNames.Add(PackageName);
+	}
+
+	TArray<FAssetData> AssetDataList;
+	IAssetRegistry* AssetRegistry = IAssetRegistry::Get();
+	if (!AssetRegistry)
+	{
+		return;
+	}
+	AssetRegistry->GetAssets(Filter, AssetDataList);
+	auto IsHigherDisplayPriority = [](const FAssetData& LHS, const FAssetData& RHS)
+	{
+		bool bLHSIsEmpty = LHS.AssetName.IsNone();
+		bool bRHSIsEmpty = RHS.AssetName.IsNone();
+		if (bLHSIsEmpty || bRHSIsEmpty)
+		{
+			return !bLHSIsEmpty;
+		}
+		bool bLHSIsUAsset = LHS.IsUAsset();
+		bool bRHSIsUAsset = RHS.IsUAsset();
+		if (bLHSIsUAsset != bRHSIsUAsset)
+		{
+			return bLHSIsUAsset;
+		}
+		bool bLHSShouldSkip = FFiltering::ShouldSkipAsset(LHS.AssetClass, LHS.PackageFlags);
+		bool bRHSShouldSkip = FFiltering::ShouldSkipAsset(LHS.AssetClass, LHS.PackageFlags);
+		if (bLHSShouldSkip || bRHSShouldSkip)
+		{
+			return !bLHSShouldSkip;
+		}
+		return LHS.AssetName.FastLess(RHS.AssetName);
+	};
+	for (FAssetData& AssetData : AssetDataList)
+	{
+		// If there are multiple assets in the same package, use the highest priority asset for the display
+		FAssetData& MappedAssetData = OutPackageToAssetData.FindOrAdd(AssetData.PackageName);
+		if (IsHigherDisplayPriority(AssetData, MappedAssetData))
+		{
+			MappedAssetData = MoveTemp(AssetData);
+		}
+	}
+}
+
+}
\ No newline at end of file
diff --git a/Engine/Source/Runtime/AssetRegistry/Public/AssetRegistry/IAssetRegistry.h b/Engine/Source/Runtime/AssetRegistry/Public/AssetRegistry/IAssetRegistry.h
--- a/Engine/Source/Runtime/AssetRegistry/Public/AssetRegistry/IAssetRegistry.h
+++ b/Engine/Source/Runtime/AssetRegistry/Public/AssetRegistry/IAssetRegistry.h
@@ -670,6 +670,13 @@ namespace AssetRegistry
 	// WritePackageData is declared in AssetRegistryInterface.h, in the CoreUObject module, because it is needed by SavePackage in CoreUObject
 	ASSETREGISTRY_API bool ReadPackageDataMain(FArchive& BinaryArchive, const FString& PackageName, const FPackageFileSummary& PackageFileSummary, int64& OutDependencyDataOffset, TArray<FAssetData*>& OutAssetDataList, EReadPackageDataMainErrorCode& OutError);
 	ASSETREGISTRY_API bool ReadPackageDataDependencies(FArchive& BinaryArchive, TBitArray<>& OutImportUsedInGame, TBitArray<>& OutSoftPackageUsedInGame);
+
+	/**
+	 * Given a list of packages, gather the primary assets for each package.
+	 * If multiple assets are in a package, the most important asset will be added.
+	 * If a package does not exist or does not have any assets, no entry will be added for that package name.
+	 */
+	ASSETREGISTRY_API void GetAssetForPackages(TConstArrayView<FName> PackageNames, TMap<FName, FAssetData>& OutPackageToAssetData);
 }
 }
 
