Make the cooked AssetRegistry.bin deterministic.
#rnx

diff --git a/Engine/Source/Runtime/AssetRegistry/Private/AssetRegistryArchive.cpp b/Engine/Source/Runtime/AssetRegistry/Private/AssetRegistryArchive.cpp
--- a/Engine/Source/Runtime/AssetRegistry/Private/AssetRegistryArchive.cpp
+++ b/Engine/Source/Runtime/AssetRegistry/Private/AssetRegistryArchive.cpp
@@ -1,20 +1,46 @@
 // Copyright Epic Games, Inc. All Rights Reserved.
 
 #include "AssetRegistryArchive.h"
+
+#include "Algo/Sort.h"
 #include "AssetRegistryPrivate.h"
 #include "AssetRegistry/AssetRegistryState.h"
 
 
 constexpr uint32 AssetRegistryNumberedNameBit = 0x80000000;
 
-static void SerializeBundleEntries(FArchive& Ar, TArray<FAssetBundleEntry>& Entries)
+static void SaveBundleEntries(FArchive& Ar, TArray<FAssetBundleEntry*>& Entries)
 {
-	for (FAssetBundleEntry& Entry : Entries)
+	for (FAssetBundleEntry* EntryPtr : Entries)
 	{
+		FAssetBundleEntry& Entry = *EntryPtr;
 		Ar << Entry.BundleName;
 
 		int32 Num = Entry.BundleAssets.Num();
 		Ar << Num;
+
+		TArray<FSoftObjectPath*> SortedPaths;
+		SortedPaths.Reserve(Num);
+		for (FSoftObjectPath& Path : Entry.BundleAssets)
+		{
+			SortedPaths.Add(&Path);
+		}
+		Algo::Sort(SortedPaths, [](FSoftObjectPath* A, FSoftObjectPath* B) { return A->LexicalLess(*B); });
+		for (FSoftObjectPath* Path : SortedPaths)
+		{
+			Path->SerializePath(Ar);
+		}
+	}
+}
+
+static void LoadBundleEntries(FArchive& Ar, TArray<FAssetBundleEntry>& Entries)
+{
+	for (FAssetBundleEntry& Entry : Entries)
+	{
+		Ar << Entry.BundleName;
+
+		int32 Num = 0;
+		Ar << Num;
 		Entry.BundleAssets.SetNum(Num);
 
 		for (FSoftObjectPath& Path : Entry.BundleAssets)
@@ -26,13 +52,22 @@ static void SerializeBundleEntries(FArchive& Ar, TArray<FAssetBundleEntry>& Entr
 
 static void SaveBundles(FArchive& Ar, const TSharedPtr<FAssetBundleData, ESPMode::ThreadSafe>& Bundles)
 {
-	TArray<FAssetBundleEntry> Empty;
-	TArray<FAssetBundleEntry>& Entries = Bundles ? Bundles->Bundles : Empty;
+	TArray<FAssetBundleEntry*> SortedEntries;
+	if (Bundles)
+	{
+		TArray<FAssetBundleEntry>& Entries = Bundles->Bundles;
+		SortedEntries.Reserve(Entries.Num());
+		for (FAssetBundleEntry& Entry : Entries)
+		{
+			SortedEntries.Add(&Entry);
+		}
+		Algo::Sort(SortedEntries, [](FAssetBundleEntry* A, FAssetBundleEntry* B) { return A->BundleName.LexicalLess(B->BundleName); });
+	}
 
-	int32 Num = Entries.Num();
+	int32 Num = SortedEntries.Num();
 	Ar << Num;
 
-	SerializeBundleEntries(Ar, Entries);
+	SaveBundleEntries(Ar, SortedEntries);
 }
 
 static TSharedPtr<FAssetBundleData, ESPMode::ThreadSafe> LoadBundles(FArchive& Ar)
@@ -44,7 +79,7 @@ static TSharedPtr<FAssetBundleData, ESPMode::ThreadSafe> LoadBundles(FArchive& A
 	{
 		FAssetBundleData Temp;
 		Temp.Bundles.SetNum(Num);
-		SerializeBundleEntries(Ar, Temp.Bundles);
+		LoadBundleEntries(Ar, Temp.Bundles);
 
 		return MakeShared<FAssetBundleData, ESPMode::ThreadSafe>(MoveTemp(Temp));
 	}
@@ -283,4 +318,4 @@ bool FAssetRegistryTagSerializationTest::RunTest(const FString& Parameters)
 	return true;
 }
 
-#endif //WITH_DEV_AUTOMATION_TESTS
\ No newline at end of file
+#endif //WITH_DEV_AUTOMATION_TESTS
diff --git a/Engine/Source/Runtime/AssetRegistry/Private/AssetRegistryState.cpp b/Engine/Source/Runtime/AssetRegistry/Private/AssetRegistryState.cpp
--- a/Engine/Source/Runtime/AssetRegistry/Private/AssetRegistryState.cpp
+++ b/Engine/Source/Runtime/AssetRegistry/Private/AssetRegistryState.cpp
@@ -2,6 +2,7 @@
 
 #include "AssetRegistry/AssetRegistryState.h"
 
+#include "Algo/Sort.h"
 #include "Misc/CommandLine.h"
 #include "Misc/FileHelper.h"
 #include "Misc/Paths.h"
@@ -1054,10 +1055,11 @@ bool FAssetRegistryState::Save(FArchive& OriginalAr, const FAssetRegistrySeriali
 	Ar << AssetCount;
 
 	// Write asset data first
-	for (TPair<FName, FAssetData*>& Pair : CachedAssetsByObjectPath)
+	TArray<TPair<FName, FAssetData*>> SortedAssetsByObjectPath = CachedAssetsByObjectPath.Array();
+	Algo::Sort(SortedAssetsByObjectPath, [](TPair<FName, FAssetData*>& A, TPair<FName, FAssetData*>& B) { return A.Key.LexicalLess(B.Key); });
+	for (const TPair<FName, FAssetData*>& Pair : SortedAssetsByObjectPath)
 	{
-		FAssetData& AssetData(*Pair.Value);
-		AssetData.SerializeForCache(Ar);
+		Pair.Value->SerializeForCache(Ar);
 	}
 
 	// Serialize Dependencies
@@ -1074,9 +1076,7 @@ bool FAssetRegistryState::Save(FArchive& OriginalAr, const FAssetRegistrySeriali
 	else
 	{
 		TMap<FDependsNode*, FDependsNode*> RedirectCache;
-		TMap<FDependsNode*, int32> DependsIndexMap;
 		TArray<FDependsNode*> Dependencies;
-		DependsIndexMap.Reserve(CachedAssetsByObjectPath.Num());
 
 		// Scan dependency nodes, we won't save all of them if we filter out certain types
 		for (TPair<FAssetIdentifier, FDependsNode*>& Pair : CachedDependsNodes)
@@ -1087,13 +1087,19 @@ bool FAssetRegistryState::Save(FArchive& OriginalAr, const FAssetRegistrySeriali
 				|| (Options.bSerializeSearchableNameDependencies && Node->GetIdentifier().IsValue())
 				|| (Options.bSerializeManageDependencies && Node->GetIdentifier().GetPrimaryAssetId().IsValid()))
 			{
-				DependsIndexMap.Add(Node, Dependencies.Num());
 				Dependencies.Add(Node);
 			}
 		}
-
+		Algo::Sort(Dependencies, [](FDependsNode* A, FDependsNode* B) { return A->GetIdentifier().LexicalLess(B->GetIdentifier()); });
 		int32 NumDependencies = Dependencies.Num();
-		Ar << NumDependencies;
+
+		TMap<FDependsNode*, int32> DependsIndexMap;
+		DependsIndexMap.Reserve(NumDependencies);
+		int32 Index = 0;
+		for (FDependsNode* Node : Dependencies)
+		{
+			DependsIndexMap.Add(Node, Index++);
+		}
 
 		TUniqueFunction<int32(FDependsNode*, bool bAsReferencer)> GetSerializeIndexFromNode = [this, &RedirectCache, &DependsIndexMap](FDependsNode* InDependency, bool bAsReferencer)
 		{
@@ -1114,18 +1120,19 @@ bool FAssetRegistryState::Save(FArchive& OriginalAr, const FAssetRegistrySeriali
 		};
 
 		FDependsNode::FSaveScratch Scratch;
+		Ar << NumDependencies;
 		for (FDependsNode* DependentNode : Dependencies)
 		{
 			DependentNode->SerializeSave(Ar, GetSerializeIndexFromNode, Scratch, Options);
 		}
 	}
-		// Write the real value to the placeholder data for the DependencySectionSize
-		int64 DependencySectionEnd = Ar.Tell();
-		DependencySectionSize = DependencySectionEnd - DependencySectionStart;
-		Ar.Seek(OffsetToDependencySectionSize);
-		Ar << DependencySectionSize;
-		check(Ar.Tell() == DependencySectionStart);
-		Ar.Seek(DependencySectionEnd);	
+	// Write the real value to the placeholder data for the DependencySectionSize
+	int64 DependencySectionEnd = Ar.Tell();
+	DependencySectionSize = DependencySectionEnd - DependencySectionStart;
+	Ar.Seek(OffsetToDependencySectionSize);
+	Ar << DependencySectionSize;
+	check(Ar.Tell() == DependencySectionStart);
+	Ar.Seek(DependencySectionEnd);
 
 
 	// Serialize the PackageData
@@ -1135,7 +1142,9 @@ bool FAssetRegistryState::Save(FArchive& OriginalAr, const FAssetRegistrySeriali
 		PackageDataCount = CachedPackageData.Num();
 		Ar << PackageDataCount;
 
-		for (TPair<FName, FAssetPackageData*>& Pair : CachedPackageData)
+		TArray<TPair<FName, FAssetPackageData*>> SortedPackageData = CachedPackageData.Array();
+		Algo::Sort(SortedPackageData, [](TPair<FName, FAssetPackageData*>& A, TPair<FName, FAssetPackageData*>& B) { return A.Key.LexicalLess(B.Key); });
+		for (TPair<FName, FAssetPackageData*>& Pair : SortedPackageData)
 		{
 			Ar << Pair.Key;
 			Pair.Value->SerializeForCache(Ar);
diff --git a/Engine/Source/Runtime/AssetRegistry/Private/DependsNode.cpp b/Engine/Source/Runtime/AssetRegistry/Private/DependsNode.cpp
--- a/Engine/Source/Runtime/AssetRegistry/Private/DependsNode.cpp
+++ b/Engine/Source/Runtime/AssetRegistry/Private/DependsNode.cpp
@@ -599,14 +599,11 @@ void FDependsNode::SerializeSave(FArchive& Ar, const TUniqueFunction<int32(FDepe
 
 	auto WriteDependencies = [&Ar, &GetSerializeIndexFromNode, &Scratch](const TArray<FDependsNode*>& InDependencies, const TBitArray<>* InFlagBits, int FlagSetWidth, bool bAsReferencer)
 	{
+		TArray<FSaveScratch::FSortInfo>& SortInfos = Scratch.SortInfos;
 		TArray<int32>& OutDependencies = Scratch.OutDependencies;
 		TBitArray<>& OutFlagBits = Scratch.OutFlagBits;
 
-		OutDependencies.Reset();
-		if (InFlagBits)
-		{
-			OutFlagBits.Reset();
-		}
+		SortInfos.Reset(InDependencies.Num());
 		for (int32 ListIndex = 0, End = InDependencies.Num(); ListIndex < End; ++ListIndex)
 		{
 			int32 SerializeIndex = GetSerializeIndexFromNode(InDependencies[ListIndex], bAsReferencer);
@@ -614,17 +611,29 @@ void FDependsNode::SerializeSave(FArchive& Ar, const TUniqueFunction<int32(FDepe
 			{
 				continue;
 			}
-			OutDependencies.Add(SerializeIndex);
-			if (InFlagBits)
-			{
-				OutFlagBits.AddRange(*InFlagBits, FlagSetWidth, ListIndex * FlagSetWidth);
-			}
+			SortInfos.Add({ SerializeIndex, ListIndex });
+		}
+		// Sort the serialized dependencies to make the output deterministic.
+		// OutDependencies and OutFlagBits (if present) are associated arrays and have to be sorted together
+		Algo::Sort(SortInfos, [](const FSaveScratch::FSortInfo& A, const FSaveScratch::FSortInfo& B) { return A.SerializeIndex < B.SerializeIndex; });
+
+		int32 NumOutDependencies = SortInfos.Num();
+		OutDependencies.Reset(NumOutDependencies);
+		for (const FSaveScratch::FSortInfo& SortInfo : SortInfos)
+		{
+			OutDependencies.Add(SortInfo.SerializeIndex);
 		}
 		Ar << OutDependencies;
 		if (InFlagBits)
 		{
+			OutFlagBits.Reset();
+			for (const FSaveScratch::FSortInfo& SortInfo : SortInfos)
+			{
+				OutFlagBits.AddRange(*InFlagBits, FlagSetWidth, SortInfo.ListIndex * FlagSetWidth);
+			}
+
 			// We don't use BitArray::operator<< because we want to avoid the reallocation that operator<< does on load and we want to avoid saving BitArray.Num when it can be derived from NumDependencies
-			int32 NumFlagBits = FlagSetWidth * OutDependencies.Num();
+			int32 NumFlagBits = FlagSetWidth * NumOutDependencies;
 			check(OutFlagBits.Num() == NumFlagBits);
 			int32 NumFlagWords = FBitSet::CalculateNumWords(NumFlagBits);
 			Ar.Serialize(OutFlagBits.GetData(), NumFlagWords * sizeof(uint32));
diff --git a/Engine/Source/Runtime/AssetRegistry/Private/DependsNode.h b/Engine/Source/Runtime/AssetRegistry/Private/DependsNode.h
--- a/Engine/Source/Runtime/AssetRegistry/Private/DependsNode.h
+++ b/Engine/Source/Runtime/AssetRegistry/Private/DependsNode.h
@@ -137,6 +137,12 @@ public:
 
 	struct FSaveScratch
 	{
+		struct FSortInfo
+		{
+			int32 SerializeIndex;
+			int32 ListIndex;
+		};
+		TArray<FSortInfo> SortInfos;
 		TArray<int32> OutDependencies;
 		TBitArray<> OutFlagBits;
 	};
diff --git a/Engine/Source/Runtime/CoreUObject/Private/AssetRegistry/AssetData.cpp b/Engine/Source/Runtime/CoreUObject/Private/AssetRegistry/AssetData.cpp
--- a/Engine/Source/Runtime/CoreUObject/Private/AssetRegistry/AssetData.cpp
+++ b/Engine/Source/Runtime/CoreUObject/Private/AssetRegistry/AssetData.cpp
@@ -2,6 +2,7 @@
 
 #include "AssetRegistry/AssetData.h"
 
+#include "Algo/IsSorted.h"
 #include "Algo/Sort.h"
 #include "AssetRegistry/ARFilter.h"
 #include "Containers/Set.h"
@@ -191,6 +192,33 @@ FPrimaryAssetId FAssetData::GetPrimaryAssetId() const
 	return FPrimaryAssetId();
 }
 
+void FAssetData::SerializeForCacheInternal(FArchive& Ar, void (*SerializeTagsAndBundles)(FArchive& , FAssetData&))
+{
+	// Serialize out the asset info
+	Ar << ObjectPath;
+	Ar << PackagePath;
+	Ar << AssetClass;
+
+	// These are derived from ObjectPath, we manually serialize them because they get pooled
+	Ar << PackageName;
+	Ar << AssetName;
+
+	SerializeTagsAndBundles(Ar , *this);
+
+	if (Ar.IsSaving() && ChunkIDs.Num() > 1)
+	{
+		TArray<int32> SortedChunkIDs(ChunkIDs);
+		Algo::Sort(SortedChunkIDs);
+		Ar << SortedChunkIDs;
+	}
+	else
+	{
+		Ar << ChunkIDs;
+	}
+	Ar << PackageFlags;
+}
+
+
 bool FAssetRegistryVersion::SerializeVersion(FArchive& Ar, FAssetRegistryVersion::Type& Version)
 {
 	FGuid Guid = FAssetRegistryVersion::GUID;
@@ -223,6 +251,55 @@ bool FAssetRegistryVersion::SerializeVersion(FArchive& Ar, FAssetRegistryVersion
 	return !Ar.IsError();
 }
 
+void FAssetPackageData::SerializeForCacheInternal(FArchive& Ar, FAssetPackageData& PackageData, FAssetRegistryVersion::Type Version)
+{
+	Ar << PackageData.DiskSize;
+	PRAGMA_DISABLE_DEPRECATION_WARNINGS
+	Ar << PackageData.PackageGuid;
+	PRAGMA_ENABLE_DEPRECATION_WARNINGS
+	if (Version >= FAssetRegistryVersion::AddedCookedMD5Hash)
+	{
+		Ar << PackageData.CookedHash;
+	}
+	if (Version >= FAssetRegistryVersion::WorkspaceDomain)
+	{
+		if (Version >= FAssetRegistryVersion::PackageFileSummaryVersionChange)
+		{
+			Ar << PackageData.FileVersionUE;
+		}
+		else
+		{
+			int32 UE4Version;
+			Ar << UE4Version;
+
+			PackageData.FileVersionUE = FPackageFileVersion::CreateUE4Version(UE4Version);
+		}
+
+		Ar << PackageData.FileVersionLicenseeUE;
+		Ar << PackageData.Flags;
+		Ar << PackageData.CustomVersions;
+	}
+	if (Version >= FAssetRegistryVersion::PackageImportedClasses)
+	{
+		if (Ar.IsSaving() && !Algo::IsSorted(PackageData.ImportedClasses, FNameLexicalLess()))
+		{
+			Algo::Sort(PackageData.ImportedClasses, FNameLexicalLess());
+		}
+		Ar << PackageData.ImportedClasses;
+	}
+}
+
+COREUOBJECT_API void FAssetPackageData::SerializeForCache(FArchive& Ar)
+{
+	// Calling with hard-coded version and using force-inline on SerializeForCacheInternal eliminates the cost of its if-statements
+	SerializeForCacheInternal(Ar, *this, FAssetRegistryVersion::LatestVersion);
+}
+
+COREUOBJECT_API void FAssetPackageData::SerializeForCacheOldVersion(FArchive& Ar, FAssetRegistryVersion::Type Version)
+{
+	SerializeForCacheInternal(Ar, *this, Version);
+}
+
 namespace UE
 {
 namespace AssetRegistry
diff --git a/Engine/Source/Runtime/CoreUObject/Private/AssetRegistry/AssetDataTagMap.cpp b/Engine/Source/Runtime/CoreUObject/Private/AssetRegistry/AssetDataTagMap.cpp
--- a/Engine/Source/Runtime/CoreUObject/Private/AssetRegistry/AssetDataTagMap.cpp
+++ b/Engine/Source/Runtime/CoreUObject/Private/AssetRegistry/AssetDataTagMap.cpp
@@ -1,6 +1,8 @@
 // Copyright Epic Games, Inc. All Rights Reserved.
 
 #include "AssetRegistry/AssetDataTagMap.h"
+
+#include "Algo/Sort.h"
 #include "AssetRegistry/AssetDataTagMapSerializationDetails.h"
 #include "Misc/PackageName.h"
 #include "Misc/ScopeLock.h"
@@ -644,7 +646,14 @@ namespace FixedTagPrivate
 		Out.Num = static_cast<uint16>(Map.Num());
 		Out.PairBegin = Pairs.Num();
 
+		TArray<TPair<FName, FAssetTagValueRef>> SortedMap;
+		SortedMap.Reserve(Map.Num());
 		for (TPair<FName, FAssetTagValueRef> Pair : Map)
+		{
+			SortedMap.Add(Pair);
+		}
+		Algo::Sort(SortedMap, [](TPair<FName, FAssetTagValueRef>& A, TPair<FName, FAssetTagValueRef>& B) { return A.Key.LexicalLess(B.Key); });
+		for (TPair<FName, FAssetTagValueRef> Pair : SortedMap)
 		{
 			FValueId Value = IndexValue(Pair.Key, Pair.Value);
 			Pairs.Add({Pair.Key, Value});
diff --git a/Engine/Source/Runtime/CoreUObject/Public/AssetRegistry/AssetData.h b/Engine/Source/Runtime/CoreUObject/Public/AssetRegistry/AssetData.h
--- a/Engine/Source/Runtime/CoreUObject/Public/AssetRegistry/AssetData.h
+++ b/Engine/Source/Runtime/CoreUObject/Public/AssetRegistry/AssetData.h
@@ -479,22 +479,13 @@ public:
 	template<class Archive>
 	void SerializeForCache(Archive&& Ar)
 	{
-		// Serialize out the asset info
-		Ar << ObjectPath;
-		Ar << PackagePath;
-		Ar << AssetClass;
-
-		// These are derived from ObjectPath, we manually serialize them because they get pooled
-		Ar << PackageName;
-		Ar << AssetName;
-
-		Ar.SerializeTagsAndBundles(*this);
-
-		Ar << ChunkIDs;
-		Ar << PackageFlags;
+		SerializeForCacheInternal(Ar, [](FArchive& Ar, FAssetData& Ad) {
+			static_cast<Archive&>(Ar).SerializeTagsAndBundles(Ad);
+			});
 	}
-
 private:
+	COREUOBJECT_API void SerializeForCacheInternal(FArchive& Ar, void (*SerializeTagsAndBundles)(FArchive& , FAssetData&));
+
 	static bool DetectIsUAssetByNames(FStringView PackageName, FStringView ObjectPathName)
 	{
 		FStringView PackageBaseName;
@@ -748,59 +739,19 @@ public:
 	 * Serialize as part of the registry cache. This is not meant to be serialized as part of a package so  it does not handle versions normally
 	 * To version this data change FAssetRegistryVersion
 	 */
-	void SerializeForCache(FArchive& Ar)
-	{
-		// Calling with hard-coded version and using force-inline on SerializeForCacheInternal eliminates the cost of its if-statements
-		SerializeForCacheInternal(Ar, FAssetRegistryVersion::LatestVersion);
-	}
-	void SerializeForCacheOldVersion(FArchive& Ar, FAssetRegistryVersion::Type Version)
-	{
-		SerializeForCacheInternal(Ar, Version);
-	}
+	COREUOBJECT_API void SerializeForCache(FArchive& Ar);
+	COREUOBJECT_API void SerializeForCacheOldVersion(FArchive& Ar, FAssetRegistryVersion::Type Version);
 
 private:
+	FORCEINLINE void SerializeForCacheInternal(FArchive& Ar, FAssetPackageData& PackageData, FAssetRegistryVersion::Type Version);
+
 	enum
 	{
 		FLAG_LICENSEE_VERSION			= 1 << 0,
 		FLAG_HAS_VIRTUALIZED_PAYLOADS	= 1 << 1
 	};
-
-	void SerializeForCacheInternal(FArchive& Ar, FAssetRegistryVersion::Type Version);
 };
 PRAGMA_ENABLE_DEPRECATION_WARNINGS
-FORCEINLINE void FAssetPackageData::SerializeForCacheInternal(FArchive& Ar, FAssetRegistryVersion::Type Version)
-{
-	Ar << DiskSize;
-	PRAGMA_DISABLE_DEPRECATION_WARNINGS
-	Ar << PackageGuid;
-	PRAGMA_ENABLE_DEPRECATION_WARNINGS
-	if (Version >= FAssetRegistryVersion::AddedCookedMD5Hash)
-	{
-		Ar << CookedHash;
-	}
-	if (Version >= FAssetRegistryVersion::WorkspaceDomain)
-	{
-		if (Version >= FAssetRegistryVersion::PackageFileSummaryVersionChange)
-		{
-			Ar << FileVersionUE;
-		}
-		else
-		{
-			int32 UE4Version;
-			Ar << UE4Version;
-
-			FileVersionUE = FPackageFileVersion::CreateUE4Version(UE4Version);
-		}
-
-		Ar << FileVersionLicenseeUE;
-		Ar << Flags;
-		Ar << CustomVersions;
-	}
-	if (Version >= FAssetRegistryVersion::PackageImportedClasses)
-	{
-		Ar << ImportedClasses;
-	}
-}
 
 
 /**
@@ -1053,5 +1004,22 @@ struct FAssetIdentifier
 		
 		return Ar;
 	}
+
+	bool LexicalLess(const FAssetIdentifier& Other) const
+	{
+		if (PrimaryAssetType != Other.PrimaryAssetType)
+		{
+			return PrimaryAssetType.LexicalLess(Other.PrimaryAssetType);
+		}
+		if (PackageName != Other.PackageName)
+		{
+			return PackageName.LexicalLess(Other.PackageName);
+		}
+		if (ObjectName != Other.ObjectName)
+		{
+			return ObjectName.LexicalLess(Other.ObjectName);
+		}
+		return ValueName.LexicalLess(Other.ValueName);
+	}
 };
 
diff --git a/Engine/Source/Runtime/CoreUObject/Public/UObject/PrimaryAssetId.h b/Engine/Source/Runtime/CoreUObject/Public/UObject/PrimaryAssetId.h
--- a/Engine/Source/Runtime/CoreUObject/Public/UObject/PrimaryAssetId.h
+++ b/Engine/Source/Runtime/CoreUObject/Public/UObject/PrimaryAssetId.h
@@ -73,6 +73,11 @@ struct FPrimaryAssetType
 		return GetTypeHash(Key.Name);
 	}
 
+	bool LexicalLess(const FPrimaryAssetType& Other) const
+	{
+		return Name.LexicalLess(Other.Name);
+	}
+
 private:
 	friend struct Z_Construct_UScriptStruct_FPrimaryAssetType_Statics;
 
diff --git a/Engine/Source/Runtime/CoreUObject/Public/UObject/SoftObjectPath.h b/Engine/Source/Runtime/CoreUObject/Public/UObject/SoftObjectPath.h
--- a/Engine/Source/Runtime/CoreUObject/Public/UObject/SoftObjectPath.h
+++ b/Engine/Source/Runtime/CoreUObject/Public/UObject/SoftObjectPath.h
@@ -177,6 +177,15 @@ struct COREUOBJECT_API FSoftObjectPath
 		return !AssetPathName.IsNone() && !SubPathString.IsEmpty();
 	}
 
+	bool LexicalLess(const FSoftObjectPath& Other) const
+	{
+		if (AssetPathName != Other.AssetPathName)
+		{
+			return AssetPathName.LexicalLess(Other.AssetPathName);
+		}
+		return SubPathString.Compare(Other.SubPathString, ESearchCase::IgnoreCase) < 0;
+	}
+
 	/** Struct overrides */
 	bool Serialize(FArchive& Ar);
 	bool Serialize(FStructuredArchive::FSlot Slot);
