Fix bad whitespace merge.
#rnx

diff --git a/Engine/Source/Runtime/AssetRegistry/Private/AssetRegistry.cpp b/Engine/Source/Runtime/AssetRegistry/Private/AssetRegistry.cpp
--- a/Engine/Source/Runtime/AssetRegistry/Private/AssetRegistry.cpp
+++ b/Engine/Source/Runtime/AssetRegistry/Private/AssetRegistry.cpp
@@ -202,8 +202,8 @@ public:
 			// run DelayedInitialize when TaskGraph system is ready
 			OnTaskGraphReady.Emplace(STATS ? EDelayedRegisterRunPhase::StatSystemReady :
 											 EDelayedRegisterRunPhase::TaskGraphSystemReady,
-				[this] () 
-				{
+				[this]()
+			{
 				DelayedInitialize();
 			}); 
 		}
@@ -295,7 +295,7 @@ private:
 	}
 
 	void DelayedInitialize()
-					{
+	{
 		// This function will run before any UObject (ie UAssetRegistryImpl) code can run, so we don't need to do any thread safety
 		// CanLoadAsync - we have to check this after the task graph is ready
 		if (!CanLoadAsync())
@@ -309,31 +309,31 @@ private:
 		PreloadReady->Trigger();
 
 		if (TrySetPath())
-						{
-							KickPreload();
-						}
-						else
-						{
+		{
+			KickPreload();
+		}
+		else
+		{
 			// set to NotFound, although PakMounted may set it to found later
 			LoadState = EState::NotFound;
 
-							// The PAK with the main registry isn't mounted yet
+			// The PAK with the main registry isn't mounted yet
 			PakMountedDelegate = FCoreDelegates::OnPakFileMounted2.AddLambda([this](const IPakFile& Pak)
-								{
+				{
 					FScopeLock Lock(&StateLock);
 					if (LoadState == EState::NotFound && TrySetPath(Pak))
-									{
-										KickPreload();
+					{
+						KickPreload();
 						// Remove the callback from OnPakFileMounted2 to avoid wasting time in all future PakFile mounts
 						// Do not access any of the lambda captures after the call to Remove, because deallocating the 
 						// DelegateHandle also deallocates our lambda captures
 						FDelegateHandle LocalPakMountedDelegate = PakMountedDelegate;
 						PakMountedDelegate.Reset();
 						FCoreDelegates::OnPakFileMounted2.Remove(LocalPakMountedDelegate);
-									}
-								});
-						}
 					}
+				});
+		}
+	}
 
 	void KickPreload()
 	{
@@ -435,10 +435,10 @@ private:
 	{
 		OnTaskGraphReady.Reset();
 		if (PreloadReady)
-	{
+		{
 			FPlatformProcess::ReturnSynchEventToPool(PreloadReady);
 			PreloadReady = nullptr;
-	}
+		}
 		ARPath.Reset();
 		Payload.Reset();
 	}
@@ -484,7 +484,7 @@ FAsyncConsumer::~FAsyncConsumer()
 		FPlatformProcess::ReturnSynchEventToPool(Consumed);
 		Consumed = nullptr;
 	}
-	}
+}
 
 void FAsyncConsumer::PrepareForConsume()
 {
@@ -538,20 +538,20 @@ void FAsyncConsumer::Consume(UAssetRegistryImpl& UARI, UE::AssetRegistry::Impl::
 }
 
 #endif
-		
-	}
+
+}
 
 namespace UE::AssetRegistry
-	{
+{
 void FAssetRegistryImpl::ConditionalLoadPremadeAssetRegistry(UAssetRegistryImpl& UARI, Impl::FEventContext& EventContext, FWriteScopeLock& ScopeLock)
-		{
+{
 #if ASSETREGISTRY_ENABLE_PREMADE_REGISTRY_IN_EDITOR
 	AsyncConsumer.Wait(UARI, ScopeLock);
 #endif
 }
 
 void FAssetRegistryImpl::ConsumeOrDeferPreloadedPremade(UAssetRegistryImpl& UARI, Impl::FEventContext& EventContext)
-			{
+{
 	// Called from inside WriteLock on InterfaceLock
 	using namespace UE::AssetRegistry::Premade;
 	if (!UE::AssetRegistry::Premade::IsEnabled())
@@ -560,7 +560,7 @@ void FAssetRegistryImpl::ConsumeOrDeferPreloadedPremade(UAssetRegistryImpl& UARI
 		// Otherwise, it is set from LoadPremadeAssetRegistry
 		bCanFinishInitialSearch = true;
 		return;
-			}
+	}
 #if ASSETREGISTRY_ENABLE_PREMADE_REGISTRY_IN_EDITOR
 	FPreloader::FConsumeFunction ConsumeFromAsyncThread = [this, &UARI](bool bSucceeded, FAssetRegistryState&& ARState)
 	{
@@ -590,9 +590,9 @@ void FAssetRegistryImpl::ConsumeOrDeferPreloadedPremade(UAssetRegistryImpl& UARI
 			{
 				LoadPremadeAssetRegistry(EventContext, bSucceeded, MoveTemp(ARState), FAssetRegistryState::EInitializationMode::Rebuild);
 			});
-		}
 	}
 }
+}
 
 /** Returns the appropriate ChunkProgressReportingType for the given Asset enum */
 EChunkProgressReportingType::Type GetChunkAvailabilityProgressType(EAssetAvailabilityProgressReportingType::Type ReportType)
@@ -3859,7 +3859,7 @@ void FAssetRegistryImpl::ScanPathsSynchronous(Impl::FScanPathContext& Context)
 			if (bIsInRequestedPaths)
 			{
 				UE_LOG(LogAssetRegistry, VeryVerbose, TEXT("FAssetRegistryImpl::ScanPathsSynchronous: Found Asset: %s"),
-					*AssetData->ObjectPath.ToString());	
+					*AssetData->ObjectPath.ToString());
 				Context.OutFoundAssets->Add(AssetData->ObjectPath);
 			}
 		}
