Add DoesPackageExistOnDisk to the AssetRegistry. Use it to check whether the AssetRegistry thinks a package exists without having to copy the entire FAssetPackageData.

In a future change we will add tracking of extension and of package name capitalization to FAssetPackageData, and will use DoesPackageExistOnDisk to replace the more expensive PackageResourceManager::TryMatchCaseOnDisk calls that occur when loading a package.

Change FPackageDatas to use DoesPackageExist, which is based on the AssetPackageDatas, rather than using GetAssetsByPackageName, because GetAssetsByPackageName can have false negatives from packages that are missing assets and can have false positives from in-memory packages that have added their Assets to the AssetRegistryState via ProcessLoadedAssetsToUpdateCache.

#rnx

diff --git a/Engine/Source/Editor/UnrealEd/Private/Cooker/CookPackageData.cpp b/Engine/Source/Editor/UnrealEd/Private/Cooker/CookPackageData.cpp
--- a/Engine/Source/Editor/UnrealEd/Private/Cooker/CookPackageData.cpp
+++ b/Engine/Source/Editor/UnrealEd/Private/Cooker/CookPackageData.cpp
@@ -1860,60 +1860,12 @@ bool FPackageDatas::TryLookupFileNameOnDisk(FName PackageName, FString& OutFileN
 	}
 	else
 	{
-		TArray<FAssetData> Assets;
-		AssetRegistry->GetAssetsByPackageName(PackageName, Assets, /*bIncludeOnlyDiskAssets =*/true);
-
-		if (Assets.Num() <= 0)
+		FString PackageExtension;
+		if (!AssetRegistry->DoesPackageExistOnDisk(PackageName, nullptr, &PackageExtension))
 		{
-			if (FPackageName::DoesPackageExist(PackageNameStr, &OutFileName, false /* InAllowTextFormats */))
-			{
-				if (AssetRegistry->GetAssetPackageDataCopy(PackageName))
-				{
-					// The AssetRegistry knows the package exists, but it has 0 assets.
-				}
-				else
-				{
-					UE_LOG(LogCook, Warning, TEXT("Package %s exists on disk but does not exist in the AssetRegistry"), *PackageNameStr);
-				}
-				return true;
-			}
 			return false;
 		}
-		else
-		{
-			// Temporary fix for packages added during cook: GetAssetsByPackageName is returning true for these generated
-			// files that do not exist on disk, since they get added into the AssetRegistryState in UAssetRegistryImpl::ProcessLoadedAssetsToUpdateCache
-			// The current known cases where this is a problem is when cooking WorldPartition maps, which create temporary World packages in
-			// //Temp. So for /Temp files, use the slower disk check rather than the AssetRegistry.
-			bool bForceDiskCheck = PackageNameStr.StartsWith(TEXT("/Temp/"));
-			if (bForceDiskCheck)
-			{
-				if (!FPackageName::DoesPackageExist(PackageNameStr, &OutFileName, false /* InAllowTextFormats */))
-				{
-					return false;
-				}
-				return true;
-			}
-		}
 
-		FName ClassRedirector = UObjectRedirector::StaticClass()->GetFName();
-		bool bContainsMap = false;
-		bool bContainsRedirector = false;
-		for (const FAssetData& Asset : Assets)
-		{
-			bContainsMap = bContainsMap | ((Asset.PackageFlags & PKG_ContainsMap) != 0);
-			bContainsRedirector = bContainsRedirector | (Asset.AssetClass == ClassRedirector);
-		}
-		if (!bContainsMap && bContainsRedirector)
-		{
-			// presence of map -> .umap
-			// But we can only assume lack of map -> .uasset if we know the type of every object in the package.
-			// If we don't, because there was a redirector, we have to check the package on disk
-			// TODO: Have the AssetRegistry store the extension of the package so that we don't have to look it up
-			// Guessing the extension based on map vs non-map also does not support text assets and maps which have a different extension
-			return FPackageName::DoesPackageExist(PackageNameStr, &OutFileName, false /* InAllowTextFormats */);
-		}
-		const FString& PackageExtension = bContainsMap ? FPackageName::GetMapPackageExtension() : FPackageName::GetAssetPackageExtension();
 		return FPackageName::TryConvertLongPackageNameToFilename(PackageNameStr, OutFileName, PackageExtension);
 	}
 }
diff --git a/Engine/Source/Runtime/AssetRegistry/Private/AssetDataGatherer.cpp b/Engine/Source/Runtime/AssetRegistry/Private/AssetDataGatherer.cpp
--- a/Engine/Source/Runtime/AssetRegistry/Private/AssetDataGatherer.cpp
+++ b/Engine/Source/Runtime/AssetRegistry/Private/AssetDataGatherer.cpp
@@ -32,7 +32,7 @@ namespace AssetDataGathererConstants
 	constexpr int32 SingleThreadFilesPerBatch = 3;
 	constexpr int32 ExpectedMaxBatchSize = 100;
 	constexpr int32 MinSecondsToElapseBeforeCacheWrite = 60;
-	static constexpr uint32 CacheSerializationMagic = 0xCBA78340; // Added integrity checking
+	static constexpr uint32 CacheSerializationMagic = 0x1F20181F; // Versioning and integrity checking
 }
 
 namespace UE
@@ -2549,7 +2549,9 @@ FAssetDataGatherer::FAssetDataGatherer(const TArray<FString>& InLongPackageNames
 {
 	using namespace UE::AssetDataGather::Private;
 
-	bGatherDependsData = (GIsEditor && !FParse::Param(FCommandLine::Get(), TEXT("NoDependsGathering"))) || FParse::Param(FCommandLine::Get(),TEXT("ForceDependsGathering"));
+	bool bForceDependsGathering = FParse::Param(FCommandLine::Get(), TEXT("ForceDependsGathering"));
+	bGatherAssetPackageData = GIsEditor || bForceDependsGathering;
+	bGatherDependsData = (GIsEditor && !FParse::Param(FCommandLine::Get(), TEXT("NoDependsGathering"))) || bForceDependsGathering;
 
 	bCacheEnabled = !FParse::Param(FCommandLine::Get(), TEXT("NoAssetRegistryCache")) && !FParse::Param(FCommandLine::Get(), TEXT("multiprocess"));
 
@@ -2848,16 +2850,19 @@ void FAssetDataGatherer::TickInternal(bool& bOutIsTickInterrupt)
 			// Add the valid cached data to our results, and to the map of data we keep to write out the new version of the cache file
 			++NumCachedFiles;
 
+			// Set the transient flags based on whether our current cache has dependency data.
+			// Note that in editor, bGatherAssetPackageData is always true, no way to turn it off,
+			// and in game it is always equal to bGatherDependsData, so it can share the cache with dependency data.
+			DiskCachedAssetData->DependencyData.bHasPackageData = bGatherAssetPackageData;
+			DiskCachedAssetData->DependencyData.bHasDependencyData = bGatherDependsData;
+
 			LocalAssetResults.Reserve(LocalAssetResults.Num() + DiskCachedAssetData->AssetDataList.Num());
 			for (const FAssetData& AssetData : DiskCachedAssetData->AssetDataList)
 			{
 				LocalAssetResults.Add(new FAssetData(AssetData));
 			}
 
-			if (bGatherDependsData)
-			{
-				LocalDependencyResults.Add(DiskCachedAssetData->DependencyData);
-			}
+			LocalDependencyResults.Add(DiskCachedAssetData->DependencyData);
 
 			AddToCache(PackageName, DiskCachedAssetData);
 		}
@@ -2922,14 +2927,7 @@ void FAssetDataGatherer::TickInternal(bool& bOutIsTickInterrupt)
 				}
 
 				// MoveTemp only used if we don't need DependencyData anymore
-				if (bGatherDependsData)
-				{
-					NewData->DependencyData = ReadContext.DependencyData;
-				}
-				else
-				{
-					NewData->DependencyData = MoveTemp(ReadContext.DependencyData);
-				}
+				NewData->DependencyData = ReadContext.DependencyData;
 
 				NewCachedAssetData.Add(NewData);
 				AddToCache(ReadContext.PackageName, NewData);
@@ -2937,10 +2935,7 @@ void FAssetDataGatherer::TickInternal(bool& bOutIsTickInterrupt)
 
 			// Add the results from the package into our output results
 			LocalAssetResults.Append(MoveTemp(ReadContext.AssetDataFromFile));
-			if (bGatherDependsData)
-			{
-				LocalDependencyResults.Add(MoveTemp(ReadContext.DependencyData));
-			}
+			LocalDependencyResults.Add(MoveTemp(ReadContext.DependencyData));
 		}
 		else if (ReadContext.bCanAttemptAssetRetry)
 		{
@@ -3022,41 +3017,39 @@ bool FAssetDataGatherer::ReadAssetFile(const FString& AssetLongPackageName, cons
 	}
 	else
 	{
-		return ReadAssetFile(PackageReader, AssetDataList, (bGatherDependsData ? &DependencyData : nullptr), CookedPackageNamesWithoutAssetData);
+		return ReadAssetFile(PackageReader, AssetDataList, DependencyData, CookedPackageNamesWithoutAssetData,
+			(bGatherAssetPackageData ? FPackageReader::EReadOptions::PackageData : FPackageReader::EReadOptions::PackageData) |
+			(bGatherDependsData ? FPackageReader::EReadOptions::Dependencies : FPackageReader::EReadOptions::PackageData));
 	}
 }
 
-bool FAssetDataGatherer::ReadAssetFile(FPackageReader& PackageReader, TArray<FAssetData*>& AssetDataList, FPackageDependencyData* DependencyData, TArray<FString>& CookedPackageNamesWithoutAssetData)
+bool FAssetDataGatherer::ReadAssetFile(FPackageReader& PackageReader, TArray<FAssetData*>& AssetDataList,
+	FPackageDependencyData& DependencyData, TArray<FString>& CookedPackageNamesWithoutAssetData, FPackageReader::EReadOptions Options)
 {
-	if (PackageReader.ReadAssetRegistryDataIfCookedPackage(AssetDataList, CookedPackageNamesWithoutAssetData))
+	bool bOutIsCookedWithoutAssetData;
+	if (!PackageReader.ReadAssetRegistryData(AssetDataList, bOutIsCookedWithoutAssetData))
 	{
-		// Cooked data is special. No further data is found in these packages
-		return true;
+		return false;
 	}
-
-	if (!PackageReader.ReadAssetRegistryData(AssetDataList))
+	if (bOutIsCookedWithoutAssetData)
 	{
-		if (!PackageReader.ReadAssetDataFromThumbnailCache(AssetDataList))
-		{
-			// It's ok to keep reading even if the asset registry data doesn't exist yet
-			//return false;
-		}
+		CookedPackageNamesWithoutAssetData.Add(PackageReader.GetLongPackageName());
 	}
 
-	if ( DependencyData )
+	if (!PackageReader.ReadDependencyData(DependencyData, Options))
 	{
-		if ( !PackageReader.ReadDependencyData(*DependencyData) )
-		{
-			return false;
-		}
+		return false;
+	}
 
+	if (EnumHasAnyFlags(Options, FPackageReader::EReadOptions::Dependencies))
+	{
 		// DEPRECATION_TODO: Remove this fixup-on-load once we bump EUnrealEngineObjectUEVersion VER_UE4_ASSETREGISTRY_DEPENDENCYFLAGS and therefore all projects will resave each ObjectRedirector
 		// UObjectRedirectors were originally incorrectly marked as having editor-only imports, since UObjectRedirector is an editor-only class. But UObjectRedirectors are followed during cooking
 		// and so their imports should be considered used-in-game. Mark all dependencies in the package as used in game if the package has a UObjectRedirector object
 		FName RedirectorClassName = UObjectRedirector::StaticClass()->GetFName();
 		if (Algo::AnyOf(AssetDataList, [RedirectorClassName](FAssetData* AssetData) { return AssetData->AssetClass == RedirectorClassName; }))
 		{
-			TBitArray<>& ImportUsedInGame = DependencyData->ImportUsedInGame;
+			TBitArray<>& ImportUsedInGame = DependencyData.ImportUsedInGame;
 			for (int32 ImportNum = ImportUsedInGame.Num(), Index = 0; Index < ImportNum; ++Index)
 			{
 				ImportUsedInGame[Index] = true;
diff --git a/Engine/Source/Runtime/AssetRegistry/Private/AssetDataGatherer.h b/Engine/Source/Runtime/AssetRegistry/Private/AssetDataGatherer.h
--- a/Engine/Source/Runtime/AssetRegistry/Private/AssetDataGatherer.h
+++ b/Engine/Source/Runtime/AssetRegistry/Private/AssetDataGatherer.h
@@ -19,6 +19,7 @@
 #include "Misc/DateTime.h"
 #include "Misc/Optional.h"
 #include "PackageDependencyData.h"
+#include "PackageReader.h"
 #include "UObject/NameTypes.h"
 #include "Templates/UniquePtr.h"
 
@@ -27,7 +28,6 @@ struct FAssetData;
 class FDiskCachedAssetData;
 class FAssetRegistryReader;
 class FAssetRegistryWriter; // Not defined if !ALLOW_NAME_BATCH_SAVING
-class FPackageReader;
 namespace UE::AssetDataGather::Private
 {
 class FAssetDataDiscovery;
@@ -167,10 +167,13 @@ public:
 	 *
 	 * @param PackageReader the previously opened package reader
 	 * @param AssetDataList the FAssetData for every asset found in the file
-	 * @param DependencyData the FPackageDependencyData for every asset found in the file, can be null
+	 * @param DependencyData the FPackageDependencyData for every asset found in the file
 	 * @param CookedPackagesToLoadUponDiscovery the list of cooked packages to be loaded if any
+	 * @param Options Which bits of data to read
 	 */
-	static bool ReadAssetFile(FPackageReader& PackageReader, TArray<FAssetData*>& AssetDataList, FPackageDependencyData* DependencyData, TArray<FString>& CookedPackagesToLoadUponDiscovery);
+	static bool ReadAssetFile(FPackageReader& PackageReader, TArray<FAssetData*>& AssetDataList,
+		FPackageDependencyData& DependencyData, TArray<FString>& CookedPackagesToLoadUponDiscovery,
+		FPackageReader::EReadOptions Options);
 
 private:
 
@@ -294,6 +297,8 @@ private:
 	 * added when Wait functions are called. Constant during threading.
 	 */
 	bool bIsSynchronous;
+	/** True if AssetPackageData should be gathered. Constant during threading. */
+	bool bGatherAssetPackageData;
 	/** True if dependency data should be gathered. Constant during threading. */
 	bool bGatherDependsData;
 	/** True if the current process allows reading/writing of AssetDataGatherer cache files. Constant during threading. */
diff --git a/Engine/Source/Runtime/AssetRegistry/Private/AssetRegistry.cpp b/Engine/Source/Runtime/AssetRegistry/Private/AssetRegistry.cpp
--- a/Engine/Source/Runtime/AssetRegistry/Private/AssetRegistry.cpp
+++ b/Engine/Source/Runtime/AssetRegistry/Private/AssetRegistry.cpp
@@ -2267,6 +2267,115 @@ TOptional<FAssetPackageData> UAssetRegistryImpl::GetAssetPackageDataCopy(FName P
 	return AssetPackageData ? *AssetPackageData : TOptional<FAssetPackageData>();
 }
 
+bool UAssetRegistryImpl::DoesPackageExistOnDisk(FName PackageName, FString* OutCorrectCasePackageName, FString* OutExtension) const
+{
+	auto CalculateExtension = [](const FString& PackageNameStr, TConstArrayView<FAssetData> Assets) -> FString
+	{
+		FName ClassRedirector = UObjectRedirector::StaticClass()->GetFName();
+		bool bContainsMap = false;
+		bool bContainsRedirector = false;
+		for (const FAssetData& Asset : Assets)
+		{
+			bContainsMap |= ((Asset.PackageFlags & PKG_ContainsMap) != 0);
+			bContainsRedirector |= (Asset.AssetClass == ClassRedirector);
+		}
+		if (!bContainsMap && bContainsRedirector)
+		{
+			// presence of map -> .umap
+			// But we can only assume lack of map -> .uasset if we know the type of every object in the package.
+			// If we don't, because there was a redirector, we have to check the package on disk
+			FString FileName;
+			if (FPackageName::DoesPackageExist(PackageNameStr, &FileName, false /* InAllowTextFormats */))
+			{
+				check(!FileName.IsEmpty());
+				return FPaths::GetExtension(FileName, true /* bIncludeDot */);
+			}
+		}
+		return bContainsMap ? FPackageName::GetMapPackageExtension() : FPackageName::GetAssetPackageExtension();
+	};
+
+#if WITH_EDITOR
+	if (GIsEditor)
+	{
+		// The editor always gathers PackageAssetDatas and uses those because they exactly match files on disk, whereas AssetsByPackageName
+		// includes memory-only assets that have added themselves to the AssetRegistry's State.
+		FString PackageNameStr = PackageName.ToString();
+		if (FPackageName::IsScriptPackage(PackageNameStr))
+		{
+			// Script packages are an exception; the AssetRegistry creates AssetPackageData for them but they exist only in memory
+			return false;
+		}
+
+		const FAssetPackageData* AssetPackageData;
+		{
+			FReadScopeLock InterfaceScopeLock(InterfaceLock);
+			AssetPackageData = GuardedData.GetState().GetAssetPackageData(PackageName);
+		}
+		constexpr bool bVerifyNegativeResults = true;
+		if (bVerifyNegativeResults && !AssetPackageData)
+		{
+			FString FileName;
+			if (FPackageName::DoesPackageExist(PackageNameStr, &FileName, false /* InAllowTextFormats */))
+			{
+				UE_LOG(LogAssetRegistry, Warning, TEXT("Package %s exists on disk but does not exist in the AssetRegistry"), *PackageNameStr);
+				if (OutCorrectCasePackageName)
+				{
+					FPackageName::TryConvertLongPackageNameToFilename(FileName, *OutCorrectCasePackageName);
+				}
+				if (OutExtension)
+				{
+					*OutExtension = FPaths::GetExtension(FileName, true /* bIncludeDot */);
+				}
+				return true;
+			}
+		}
+
+		if (!AssetPackageData)
+		{
+			return false;
+		}
+
+		if (OutCorrectCasePackageName)
+		{
+			// TODO: Implement this correctly by saving the information in AssetPackageData if the capitalization does not match the FName capitalization
+			*OutCorrectCasePackageName = PackageNameStr;
+		}
+		if (OutExtension)
+		{
+			// TODO: Save the extension on the AssetPackageData rather than deriving it here
+			// Guessing the extension based on map vs non-map also does not support text assets and maps which have a different extension
+			TArray<FAssetData> Assets;
+			GetAssetsByPackageName(PackageName, Assets, /*bIncludeOnlyDiskAssets*/ true);
+			*OutExtension = CalculateExtension(PackageNameStr, Assets);
+		}
+		return true;
+	}
+	else
+#endif
+	{
+		// Runtime Game and Programs use GetAssetsByPackageName, which will match the files on disk since these configurations do not
+		// add loaded assets to the AssetRegistryState
+		TArray<FAssetData> Assets;
+		GetAssetsByPackageName(PackageName, Assets, /*bIncludeOnlyDiskAssets*/ true);
+		if (Assets.Num() == 0)
+		{
+			return false;
+		}
+		FString PackageNameStr = PackageName.ToString();
+		if (OutCorrectCasePackageName)
+		{
+			// In Game does not handle matching case, but it still needs to return a value for the CorrectCase field if asked
+			*OutCorrectCasePackageName = PackageNameStr;
+		}
+		if (OutExtension)
+		{
+			*OutExtension = CalculateExtension(PackageNameStr, Assets);
+		}
+		return true;
+	}
+}
+
+
 FName UAssetRegistryImpl::GetRedirectedObjectPath(const FName ObjectPath) const
 {
 	FReadScopeLock InterfaceScopeLock(InterfaceLock);
@@ -3624,7 +3733,9 @@ bool ReadAssetFile(FPackageReader& PackageReader, IAssetRegistry::FLoadPackageRe
 	const bool bGetDependencies = (InOutData.bGetDependencies);
 	FPackageDependencyData DependencyData;
 
-	bool bReadOk = FAssetDataGatherer::ReadAssetFile(PackageReader, AssetDataList, (InOutData.bGetDependencies ? &DependencyData : nullptr), CookedPackageNamesWithoutAssetDataGathered);
+	bool bReadOk = FAssetDataGatherer::ReadAssetFile(PackageReader, AssetDataList, DependencyData,
+		CookedPackageNamesWithoutAssetDataGathered,
+		InOutData.bGetDependencies ? FPackageReader::EReadOptions::Dependencies : FPackageReader::EReadOptions::None);
 
 	if (bReadOk)
 	{
@@ -4023,142 +4134,149 @@ void FAssetRegistryImpl::DependencyDataGathered(const double TickStartTime, TRin
 	{
 		FPackageDependencyData Result = DependsResults.PopFrontValue();
 
-		// Update package data
-		FAssetPackageData* PackageData = State.CreateOrGetAssetPackageData(Result.PackageName);
-		*PackageData = Result.PackageData;
-
-		FDependsNode* Node = State.CreateOrFindDependsNode(Result.PackageName);
+		checkf(!GIsEditor || Result.bHasPackageData, TEXT("We rely on PackageData being read for every gathered Asset in the editor."));
+		if (Result.bHasPackageData)
+		{
+			// Update package data
+			FAssetPackageData* PackageData = State.CreateOrGetAssetPackageData(Result.PackageName);
+			*PackageData = Result.PackageData;
+		}
 
-		// We will populate the node dependencies below. Empty the set here in case this file was already read
-		// Also remove references to all existing dependencies, those will be also repopulated below
-		Node->IterateOverDependencies([Node](FDependsNode* InDependency, UE::AssetRegistry::EDependencyCategory Category, UE::AssetRegistry::EDependencyProperty Properties, bool bDuplicate)
+		if (Result.bHasDependencyData)
 		{
-			if (!bDuplicate)
-			{
-				InDependency->RemoveReferencer(Node);
-			}
-		});
+			FDependsNode* Node = State.CreateOrFindDependsNode(Result.PackageName);
 
-		Node->ClearDependencies();
+			// We will populate the node dependencies below. Empty the set here in case this file was already read
+			// Also remove references to all existing dependencies, those will be also repopulated below
+			Node->IterateOverDependencies([Node](FDependsNode* InDependency, UE::AssetRegistry::EDependencyCategory Category, UE::AssetRegistry::EDependencyProperty Properties, bool bDuplicate)
+				{
+					if (!bDuplicate)
+					{
+						InDependency->RemoveReferencer(Node);
+					}
+				});
 
-		// Don't bother registering dependencies on these packages, every package in the game will depend on them
-		static TArray<FName> ScriptPackagesToSkip = TArray<FName>{ TEXT("/Script/CoreUObject"), TEXT("/Script/Engine"), TEXT("/Script/BlueprintGraph"), TEXT("/Script/UnrealEd") };
+			Node->ClearDependencies();
 
-		// Determine the new package dependencies
-		TMap<FName, FDependsNode::FPackageFlagSet> PackageDependencies;
-		check(Result.ImportUsedInGame.Num() == Result.ImportMap.Num());
-		for (int32 ImportIdx = 0; ImportIdx < Result.ImportMap.Num(); ++ImportIdx)
-		{
-			FName AssetReference = Result.GetImportPackageName(ImportIdx);
+			// Don't bother registering dependencies on these packages, every package in the game will depend on them
+			static TArray<FName> ScriptPackagesToSkip = TArray<FName>{ TEXT("/Script/CoreUObject"), TEXT("/Script/Engine"), TEXT("/Script/BlueprintGraph"), TEXT("/Script/UnrealEd") };
 
-			// Should we skip this because it's too common?
-			if (ScriptPackagesToSkip.Contains(AssetReference))
+			// Determine the new package dependencies
+			TMap<FName, FDependsNode::FPackageFlagSet> PackageDependencies;
+			check(Result.ImportUsedInGame.Num() == Result.ImportMap.Num());
+			for (int32 ImportIdx = 0; ImportIdx < Result.ImportMap.Num(); ++ImportIdx)
 			{
-				continue;
-			}
-			FName Redirected = FCoreRedirects::GetRedirectedName(ECoreRedirectFlags::Type_Package,
-				FCoreRedirectObjectName(NAME_None, NAME_None, AssetReference)).PackageName;
-			AssetReference = Redirected;
-
-			EDependencyProperty DependencyProperty = EDependencyProperty::Build | EDependencyProperty::Hard;
-			DependencyProperty |= Result.ImportUsedInGame[ImportIdx] ? EDependencyProperty::Game : EDependencyProperty::None;
-			PackageDependencies.FindOrAdd(AssetReference).Add(FDependsNode::PackagePropertiesToByte(DependencyProperty));
-		}
+				FName AssetReference = Result.GetImportPackageName(ImportIdx);
 
-		check(Result.SoftPackageUsedInGame.Num() == Result.SoftPackageReferenceList.Num());
-		for (int32 SoftPackageIdx = 0; SoftPackageIdx < Result.SoftPackageReferenceList.Num(); ++SoftPackageIdx)
-		{
-			FName AssetReference = Result.SoftPackageReferenceList[SoftPackageIdx];
-			FName Redirected = FCoreRedirects::GetRedirectedName(ECoreRedirectFlags::Type_Package,
-				FCoreRedirectObjectName(NAME_None, NAME_None, AssetReference)).PackageName;
-			AssetReference = Redirected;
+				// Should we skip this because it's too common?
+				if (ScriptPackagesToSkip.Contains(AssetReference))
+				{
+					continue;
+				}
+				FName Redirected = FCoreRedirects::GetRedirectedName(ECoreRedirectFlags::Type_Package,
+					FCoreRedirectObjectName(NAME_None, NAME_None, AssetReference)).PackageName;
+				AssetReference = Redirected;
 
-			EDependencyProperty DependencyProperty = UE::AssetRegistry::EDependencyProperty::Build;
-			DependencyProperty |= (Result.SoftPackageUsedInGame[SoftPackageIdx] ? EDependencyProperty::Game : EDependencyProperty::None);
-			PackageDependencies.FindOrAdd(AssetReference).Add(FDependsNode::PackagePropertiesToByte(DependencyProperty));
-		}
+				EDependencyProperty DependencyProperty = EDependencyProperty::Build | EDependencyProperty::Hard;
+				DependencyProperty |= Result.ImportUsedInGame[ImportIdx] ? EDependencyProperty::Game : EDependencyProperty::None;
+				PackageDependencies.FindOrAdd(AssetReference).Add(FDependsNode::PackagePropertiesToByte(DependencyProperty));
+			}
 
-		// Doubly-link all of the PackageDependencies
-		for (TPair<FName, FDependsNode::FPackageFlagSet>& NewDependsIt : PackageDependencies)
-		{
-			FDependsNode* DependsNode = State.CreateOrFindDependsNode(NewDependsIt.Key);
+			check(Result.SoftPackageUsedInGame.Num() == Result.SoftPackageReferenceList.Num());
+			for (int32 SoftPackageIdx = 0; SoftPackageIdx < Result.SoftPackageReferenceList.Num(); ++SoftPackageIdx)
+			{
+				FName AssetReference = Result.SoftPackageReferenceList[SoftPackageIdx];
+				FName Redirected = FCoreRedirects::GetRedirectedName(ECoreRedirectFlags::Type_Package,
+					FCoreRedirectObjectName(NAME_None, NAME_None, AssetReference)).PackageName;
+				AssetReference = Redirected;
+
+				EDependencyProperty DependencyProperty = UE::AssetRegistry::EDependencyProperty::Build;
+				DependencyProperty |= (Result.SoftPackageUsedInGame[SoftPackageIdx] ? EDependencyProperty::Game : EDependencyProperty::None);
+				PackageDependencies.FindOrAdd(AssetReference).Add(FDependsNode::PackagePropertiesToByte(DependencyProperty));
+			}
 
-			// Handle failure of CreateOrFindDependsNode
-			// And Skip dependencies to self 
-			if (DependsNode != nullptr && DependsNode != Node)
+			// Doubly-link all of the PackageDependencies
+			for (TPair<FName, FDependsNode::FPackageFlagSet>& NewDependsIt : PackageDependencies)
 			{
-				const FAssetIdentifier& Identifier = DependsNode->GetIdentifier();
-				if (DependsNode->GetConnectionCount() == 0 && Identifier.IsPackage())
-				{
-					// This was newly created, see if we need to read the script package Guid
-					const FNameBuilder PackageName(Identifier.PackageName);
+				FDependsNode* DependsNode = State.CreateOrFindDependsNode(NewDependsIt.Key);
 
-					if (FPackageName::IsScriptPackage(PackageName))
+				// Handle failure of CreateOrFindDependsNode
+				// And Skip dependencies to self 
+				if (DependsNode != nullptr && DependsNode != Node)
+				{
+					const FAssetIdentifier& Identifier = DependsNode->GetIdentifier();
+					if (DependsNode->GetConnectionCount() == 0 && Identifier.IsPackage())
 					{
-						// Get the guid off the script package, it is updated when script is changed so we need to refresh it every run
-						UPackage* Package = FindPackage(nullptr, *PackageName);
+						// This was newly created, see if we need to read the script package Guid
+						const FNameBuilder PackageName(Identifier.PackageName);
 
-						if (Package)
+						if (FPackageName::IsScriptPackage(PackageName))
 						{
-							FAssetPackageData* ScriptPackageData = State.CreateOrGetAssetPackageData(Identifier.PackageName);
-							PRAGMA_DISABLE_DEPRECATION_WARNINGS
-							ScriptPackageData->PackageGuid = Package->GetGuid();
-							PRAGMA_ENABLE_DEPRECATION_WARNINGS
+							// Get the guid off the script package, it is updated when script is changed so we need to refresh it every run
+							UPackage* Package = FindPackage(nullptr, *PackageName);
+
+							if (Package)
+							{
+								FAssetPackageData* ScriptPackageData = State.CreateOrGetAssetPackageData(Identifier.PackageName);
+								PRAGMA_DISABLE_DEPRECATION_WARNINGS
+									ScriptPackageData->PackageGuid = Package->GetGuid();
+								PRAGMA_ENABLE_DEPRECATION_WARNINGS
+							}
 						}
 					}
-				}
 
-				Node->AddPackageDependencySet(DependsNode, NewDependsIt.Value);
-				DependsNode->AddReferencer(Node);
+					Node->AddPackageDependencySet(DependsNode, NewDependsIt.Value);
+					DependsNode->AddReferencer(Node);
+				}
 			}
-		}
-
-		for (const TPair<FPackageIndex, TArray<FName>>& SearchableNameList : Result.SearchableNamesMap)
-		{
-			FName ObjectName;
-			FName PackageName;
 
-			// Find object and package name from linker
-			FPackageIndex LinkerIndex = SearchableNameList.Key;
-			if (LinkerIndex.IsExport())
-			{
-				// Package name has to be this package, take a guess at object name
-				PackageName = Result.PackageName;
-				ObjectName = FName(*FPackageName::GetLongPackageAssetName(Result.PackageName.ToString()));
-			}
-			else if (LinkerIndex.IsImport())
+			for (const TPair<FPackageIndex, TArray<FName>>& SearchableNameList : Result.SearchableNamesMap)
 			{
-				FObjectResource* Resource = &Result.ImpExp(LinkerIndex);
-				FPackageIndex OuterLinkerIndex = Resource->OuterIndex;
-				check(OuterLinkerIndex.IsNull() || OuterLinkerIndex.IsImport());
-				if (!OuterLinkerIndex.IsNull())
+				FName ObjectName;
+				FName PackageName;
+
+				// Find object and package name from linker
+				FPackageIndex LinkerIndex = SearchableNameList.Key;
+				if (LinkerIndex.IsExport())
+				{
+					// Package name has to be this package, take a guess at object name
+					PackageName = Result.PackageName;
+					ObjectName = FName(*FPackageName::GetLongPackageAssetName(Result.PackageName.ToString()));
+				}
+				else if (LinkerIndex.IsImport())
 				{
-					ObjectName = Resource->ObjectName;
-					while (!OuterLinkerIndex.IsNull())
+					FObjectResource* Resource = &Result.ImpExp(LinkerIndex);
+					FPackageIndex OuterLinkerIndex = Resource->OuterIndex;
+					check(OuterLinkerIndex.IsNull() || OuterLinkerIndex.IsImport());
+					if (!OuterLinkerIndex.IsNull())
 					{
-						Resource = &Result.ImpExp(OuterLinkerIndex);
-						OuterLinkerIndex = Resource->OuterIndex;
-						check(OuterLinkerIndex.IsNull() || OuterLinkerIndex.IsImport());
+						ObjectName = Resource->ObjectName;
+						while (!OuterLinkerIndex.IsNull())
+						{
+							Resource = &Result.ImpExp(OuterLinkerIndex);
+							OuterLinkerIndex = Resource->OuterIndex;
+							check(OuterLinkerIndex.IsNull() || OuterLinkerIndex.IsImport());
+						}
 					}
+					PackageName = Resource->ObjectName;
 				}
-				PackageName = Resource->ObjectName;
-			}
 
-			for (FName NameReference : SearchableNameList.Value)
-			{
-				FAssetIdentifier AssetId = FAssetIdentifier(PackageName, ObjectName, NameReference);
+				for (FName NameReference : SearchableNameList.Value)
+				{
+					FAssetIdentifier AssetId = FAssetIdentifier(PackageName, ObjectName, NameReference);
 
-				// Add node for all name references
-				FDependsNode* DependsNode = State.CreateOrFindDependsNode(AssetId);
+					// Add node for all name references
+					FDependsNode* DependsNode = State.CreateOrFindDependsNode(AssetId);
 
-				if (DependsNode != nullptr)
-				{
-					Node->AddDependency(DependsNode, EDependencyCategory::SearchableName, EDependencyProperty::None);
-					DependsNode->AddReferencer(Node);
+					if (DependsNode != nullptr)
+					{
+						Node->AddDependency(DependsNode, EDependencyCategory::SearchableName, EDependencyProperty::None);
+						DependsNode->AddReferencer(Node);
+					}
 				}
 			}
+			Node->SetIsDependenciesInitialized(true);
 		}
-		Node->SetIsDependenciesInitialized(true);
 
 		// Check to see if we have run out of time in this tick
 		if (!bFlushFullBuffer && (FPlatformTime::Seconds() - TickStartTime) > Impl::MaxSecondsPerFrame)
diff --git a/Engine/Source/Runtime/AssetRegistry/Private/AssetRegistry.h b/Engine/Source/Runtime/AssetRegistry/Private/AssetRegistry.h
--- a/Engine/Source/Runtime/AssetRegistry/Private/AssetRegistry.h
+++ b/Engine/Source/Runtime/AssetRegistry/Private/AssetRegistry.h
@@ -60,6 +60,7 @@ public:
 	virtual bool GetReferencers(FName PackageName, TArray<FName>& OutReferencers, UE::AssetRegistry::EDependencyCategory Category = UE::AssetRegistry::EDependencyCategory::Package, const UE::AssetRegistry::FDependencyQuery& Flags = UE::AssetRegistry::FDependencyQuery()) const override;
 	virtual const FAssetPackageData* GetAssetPackageData(FName PackageName) const override;
 	virtual TOptional<FAssetPackageData> GetAssetPackageDataCopy(FName PackageName) const override;
+	virtual bool DoesPackageExistOnDisk(FName PackageName, FString* OutCorrectCasePackageName = nullptr, FString* OutExtension = nullptr) const override;
 	virtual FName GetRedirectedObjectPath(const FName ObjectPath) const override;
 	virtual bool GetAncestorClassNames(FName ClassName, TArray<FName>& OutAncestorClassNames) const override;
 	virtual void GetDerivedClassNames(const TArray<FName>& ClassNames, const TSet<FName>& ExcludedClassNames, TSet<FName>& OutDerivedClassNames) const override;
diff --git a/Engine/Source/Runtime/AssetRegistry/Private/PackageDependencyData.h b/Engine/Source/Runtime/AssetRegistry/Private/PackageDependencyData.h
--- a/Engine/Source/Runtime/AssetRegistry/Private/PackageDependencyData.h
+++ b/Engine/Source/Runtime/AssetRegistry/Private/PackageDependencyData.h
@@ -21,6 +21,9 @@ public:
 	TBitArray<> ImportUsedInGame;
 	TBitArray<> SoftPackageUsedInGame;
 
+	bool bHasPackageData = false; // Transient
+	bool bHasDependencyData = false; // Transient
+
 	/**
 	 * Return the package name of the UObject represented by the specified import. 
 	 * 
diff --git a/Engine/Source/Runtime/AssetRegistry/Private/PackageReader.cpp b/Engine/Source/Runtime/AssetRegistry/Private/PackageReader.cpp
--- a/Engine/Source/Runtime/AssetRegistry/Private/PackageReader.cpp
+++ b/Engine/Source/Runtime/AssetRegistry/Private/PackageReader.cpp
@@ -116,8 +116,8 @@ bool FPackageReader::OpenPackageFile(EOpenPackageResult* OutErrorCode)
 
 	if (PackageFileSummary.GetFileVersionLicenseeUE() > GPackageFileLicenseeUEVersion)
 	{
-		UE_LOG(	LogAssetRegistry, Error, TEXT("Package %s is too new. Licensee Version: %i Package Licensee Version: %i"), 
-				*PackageFilename, GPackageFileLicenseeUEVersion, PackageFileSummary.GetFileVersionLicenseeUE());
+		UE_LOG(LogAssetRegistry, Error, TEXT("Package %s is too new. Licensee Version: %i Package Licensee Version: %i"),
+			*PackageFilename, GPackageFileLicenseeUEVersion, PackageFileSummary.GetFileVersionLicenseeUE());
 
 		SetPackageErrorCode(EOpenPackageResult::VersionTooNew);
 		return false;
@@ -173,6 +173,13 @@ bool FPackageReader::TryGetLongPackageName(FString& OutLongPackageName) const
 	}
 }
 
+FString FPackageReader::GetLongPackageName() const
+{
+	FString Result;
+	verify(TryGetLongPackageName(Result));
+	return Result;
+}
+
 bool FPackageReader::StartSerializeSection(int64 Offset)
 {
 	check(Loader);
@@ -194,11 +201,22 @@ bool FPackageReader::StartSerializeSection(int64 Offset)
 		FMessageLog("AssetRegistry").Warning(FText::Format(NSLOCTEXT("AssetRegistry", MessageKey, "Cannot read AssetRegistry Data in {FileName}, skipping it. Error: " MessageKey "."), CorruptPackageWarningArguments)); \
 	} while (false)
 
-bool FPackageReader::ReadAssetRegistryData(TArray<FAssetData*>& AssetDataList)
+bool FPackageReader::ReadAssetRegistryData(TArray<FAssetData*>& AssetDataList, bool& bOutIsCookedWithoutAssetData)
 {
+	bOutIsCookedWithoutAssetData = false;
+	if (!!(GetPackageFlags() & PKG_FilterEditorOnly))
+	{
+		return ReadAssetRegistryDataFromCookedPackage(AssetDataList, bOutIsCookedWithoutAssetData);
+	}
+
 	if (!StartSerializeSection(PackageFileSummary.AssetRegistryDataOffset))
 	{
-		return false;
+		if (!ReadAssetDataFromThumbnailCache(AssetDataList))
+		{
+			// Legacy files without AR data and without a thumbnail cache are treated as having no assets
+			AssetDataList.Reset();
+		}
+		return true;
 	}
 
 	// Determine the package name and path
@@ -259,34 +277,23 @@ bool FPackageReader::SerializeAssetRegistryDependencyData(FPackageDependencyData
 	return true;
 }
 
-bool FPackageReader::SerializePackageTrailer(FPackageDependencyData& DependencyData)
+bool FPackageReader::SerializePackageTrailer(FAssetPackageData& PackageData)
 {
-	bool bHasVirtualizedPayloads = false;
-
-	if (PackageFileSummary.PayloadTocOffset != INDEX_NONE)
+	if (!StartSerializeSection(PackageFileSummary.PayloadTocOffset))
 	{
-		// Store the original offset in the package being read, since we are about to seek to near the end
-		const int64 InitialOffset = Tell();
-		
-		Seek(PackageFileSummary.PayloadTocOffset);
-
-		if (IsError())
-		{
-			return false;
-		}
-
-		UE::FPackageTrailer Trailer;
-		if (Trailer.TryLoad(*this))
-		{
-			bHasVirtualizedPayloads = Trailer.GetNumPayloads(UE::EPayloadFilter::Virtualized) > 0 ? true : false;
-		}
-
-		// Restore the original offset inside the package
-		Seek(InitialOffset);
+		PackageData.SetHasVirtualizedPayloads(false);
+		return true;
 	}
 
-	DependencyData.PackageData.SetHasVirtualizedPayloads(bHasVirtualizedPayloads);
+	UE::FPackageTrailer Trailer;
+	if (!Trailer.TryLoad(*this))
+	{
+		// This is not necessarily corrupt; TryLoad will return false if the trailer is empty
+		PackageData.SetHasVirtualizedPayloads(false);
+		return true;
+	}
 
+	PackageData.SetHasVirtualizedPayloads(Trailer.GetNumPayloads(UE::EPayloadFilter::Virtualized) > 0);
 	return true;
 }
 
@@ -351,77 +358,69 @@ bool FPackageReader::ReadAssetDataFromThumbnailCache(TArray<FAssetData*>& AssetD
 	return true;
 }
 
-bool FPackageReader::ReadAssetRegistryDataIfCookedPackage(TArray<FAssetData*>& AssetDataList, TArray<FString>& CookedPackageNamesWithoutAssetData)
+bool FPackageReader::ReadAssetRegistryDataFromCookedPackage(TArray<FAssetData*>& AssetDataList, bool& bOutIsCookedWithoutAssetData)
 {
-	if (!!(GetPackageFlags() & PKG_FilterEditorOnly))
+	FString PackageName;
+	if (!TryGetLongPackageName(PackageName))
 	{
-		FString PackageName;
-		if (!TryGetLongPackageName(PackageName))
+		return false;
+	}
+
+	bool bFoundAtLeastOneAsset = false;
+
+	// If the packaged is saved with the right version we have the information
+	// which of the objects in the export map as the asset.
+	// Otherwise we need to store a temp minimal data and then force load the asset
+	// to re-generate its registry data
+	if (UEVer() >= VER_UE4_COOKED_ASSETS_IN_EDITOR_SUPPORT)
+	{
+		const FString PackagePath = FPackageName::GetLongPackagePath(PackageName);
+
+		TArray<FObjectImport> ImportMap;
+		TArray<FObjectExport> ExportMap;
+		if (!SerializeNameMap())
 		{
 			return false;
 		}
-
-		bool bFoundAtLeastOneAsset = false;
-
-		// If the packaged is saved with the right version we have the information
-		// which of the objects in the export map as the asset.
-		// Otherwise we need to store a temp minimal data and then force load the asset
-		// to re-generate its registry data
-		if (UEVer() >= VER_UE4_COOKED_ASSETS_IN_EDITOR_SUPPORT)
+		if (!SerializeImportMap(ImportMap))
 		{
-			const FString PackagePath = FPackageName::GetLongPackagePath(PackageName);
-
-			TArray<FObjectImport> ImportMap;
-			TArray<FObjectExport> ExportMap;
-			if (!SerializeNameMap())
-			{
-				return false;
-			}
-			if (!SerializeImportMap(ImportMap))
-			{
-				return false;
-			}
-			if (!SerializeExportMap(ExportMap))
-			{
-				return false;
-			}
-			for (FObjectExport& Export : ExportMap)
+			return false;
+		}
+		if (!SerializeExportMap(ExportMap))
+		{
+			return false;
+		}
+		for (FObjectExport& Export : ExportMap)
+		{
+			if (Export.bIsAsset)
 			{
-				if (Export.bIsAsset)
+				// We need to get the class name from the import/export maps
+				FName ObjectClassName;
+				if (Export.ClassIndex.IsNull())
 				{
-					// We need to get the class name from the import/export maps
-					FName ObjectClassName;
-					if (Export.ClassIndex.IsNull())
-					{
-						ObjectClassName = UClass::StaticClass()->GetFName();
-					}
-					else if (Export.ClassIndex.IsExport())
-					{
-						const FObjectExport& ClassExport = ExportMap[Export.ClassIndex.ToExport()];
-						ObjectClassName = ClassExport.ObjectName;
-					}
-					else if (Export.ClassIndex.IsImport())
-					{
-						const FObjectImport& ClassImport = ImportMap[Export.ClassIndex.ToImport()];
-						ObjectClassName = ClassImport.ObjectName;
-					}
-
-					AssetDataList.Add(new FAssetData(FName(*PackageName), FName(*PackagePath), Export.ObjectName, ObjectClassName, FAssetDataTagMap(), TArray<int32>(), GetPackageFlags()));
-					bFoundAtLeastOneAsset = true;
+					ObjectClassName = UClass::StaticClass()->GetFName();
+				}
+				else if (Export.ClassIndex.IsExport())
+				{
+					const FObjectExport& ClassExport = ExportMap[Export.ClassIndex.ToExport()];
+					ObjectClassName = ClassExport.ObjectName;
+				}
+				else if (Export.ClassIndex.IsImport())
+				{
+					const FObjectImport& ClassImport = ImportMap[Export.ClassIndex.ToImport()];
+					ObjectClassName = ClassImport.ObjectName;
 				}
+
+				AssetDataList.Add(new FAssetData(FName(*PackageName), FName(*PackagePath), Export.ObjectName, ObjectClassName, FAssetDataTagMap(), TArray<int32>(), GetPackageFlags()));
+				bFoundAtLeastOneAsset = true;
 			}
 		}
-		if (!bFoundAtLeastOneAsset)
-		{
-			CookedPackageNamesWithoutAssetData.Add(PackageName);
-		}
-		return true;
 	}
-
-	return false;
+	bOutIsCookedWithoutAssetData = !bFoundAtLeastOneAsset;
+	return true;
 }
 
-bool FPackageReader::ReadDependencyData(FPackageDependencyData& OutDependencyData)
+bool FPackageReader::ReadDependencyData(FPackageDependencyData& OutDependencyData, EReadOptions Options)
 {
 	FString PackageNameString;
 	if (!TryGetLongPackageName(PackageNameString))
@@ -431,15 +430,10 @@ bool FPackageReader::ReadDependencyData(FPackageDependencyData& OutDependencyDat
 	}
 
 	OutDependencyData.PackageName = FName(*PackageNameString);
-	FAssetPackageData& PackageData = OutDependencyData.PackageData;
-	PackageData.DiskSize = PackageFileSize;
-	PRAGMA_DISABLE_DEPRECATION_WARNINGS
-	PackageData.PackageGuid = PackageFileSummary.Guid;
-	PRAGMA_ENABLE_DEPRECATION_WARNINGS
-	PackageData.SetCustomVersions(PackageFileSummary.GetCustomVersionContainer().GetAllVersions());
-	PackageData.FileVersionUE = PackageFileSummary.GetFileVersionUE();
-	PackageData.FileVersionLicenseeUE = PackageFileSummary.GetFileVersionLicenseeUE();
-	PackageData.SetIsLicenseeVersion(PackageFileSummary.SavedByEngineVersion.IsLicenseeVersion());
+	if (!EnumHasAnyFlags(Options, EReadOptions::PackageData | EReadOptions::Dependencies))
+	{
+		return true;
+	}
 
 	if (!SerializeNameMap())
 	{
@@ -449,26 +443,51 @@ bool FPackageReader::ReadDependencyData(FPackageDependencyData& OutDependencyDat
 	{
 		return false;
 	}
-	if (!SerializeImportedClasses(OutDependencyData.ImportMap, PackageData.ImportedClasses))
-	{
-		return false;
-	}
-	if (!SerializeSoftPackageReferenceList(OutDependencyData.SoftPackageReferenceList))
-	{
-		return false;
-	}
-	if (!SerializeSearchableNamesMap(OutDependencyData))
+
+	if (EnumHasAnyFlags(Options, EReadOptions::PackageData))
 	{
-		return false;
+		OutDependencyData.bHasPackageData = true;
+		FAssetPackageData& PackageData = OutDependencyData.PackageData;
+		PackageData.DiskSize = PackageFileSize;
+		PRAGMA_DISABLE_DEPRECATION_WARNINGS;
+		PackageData.PackageGuid = PackageFileSummary.Guid;
+		PRAGMA_ENABLE_DEPRECATION_WARNINGS;
+		PackageData.SetCustomVersions(PackageFileSummary.GetCustomVersionContainer().GetAllVersions());
+		PackageData.FileVersionUE = PackageFileSummary.GetFileVersionUE();
+		PackageData.FileVersionLicenseeUE = PackageFileSummary.GetFileVersionLicenseeUE();
+		PackageData.SetIsLicenseeVersion(PackageFileSummary.SavedByEngineVersion.IsLicenseeVersion());
+
+		if (!SerializeImportedClasses(OutDependencyData.ImportMap, PackageData.ImportedClasses))
+		{
+			return false;
+		}
+		if (!SerializePackageTrailer(PackageData))
+		{
+			return false;
+		}
 	}
-	if (!SerializeAssetRegistryDependencyData(OutDependencyData))
+
+	if (EnumHasAnyFlags(Options, EReadOptions::Dependencies))
 	{
-		return false;
+		OutDependencyData.bHasDependencyData = true;
+		if (!SerializeSoftPackageReferenceList(OutDependencyData.SoftPackageReferenceList))
+		{
+			return false;
+		}
+		if (!SerializeSearchableNamesMap(OutDependencyData))
+		{
+			return false;
+		}
+		if (!SerializeAssetRegistryDependencyData(OutDependencyData))
+		{
+			return false;
+		}
 	}
-	// Note that the trailer is read from the end of the package file, and so should be the last thing read here
-	if (!SerializePackageTrailer(OutDependencyData))
+	else
 	{
-		return false;
+		// We used OutDependencyData.ImportMap to store data for SerializeImportedClasses to use, but 
+		// our contract says we should leave it empty if !EReadOptions::Dependencies
+		OutDependencyData.ImportMap.Reset();
 	}
 
 	checkf(OutDependencyData.IsValid(), TEXT("We should have early exited above rather than creating invalid dependency data"));
diff --git a/Engine/Source/Runtime/AssetRegistry/Private/PackageReader.h b/Engine/Source/Runtime/AssetRegistry/Private/PackageReader.h
--- a/Engine/Source/Runtime/AssetRegistry/Private/PackageReader.h
+++ b/Engine/Source/Runtime/AssetRegistry/Private/PackageReader.h
@@ -4,10 +4,12 @@
 
 #include "CoreMinimal.h"
 #include "Serialization/ArchiveUObject.h"
+#include "Misc/EnumClassFlags.h"
 #include "UObject/ObjectResource.h"
 #include "UObject/PackageFileSummary.h"
 
 struct FAssetData;
+class FAssetPackageData;
 class FPackageDependencyData;
 
 class FPackageReader : public FArchiveUObject
@@ -42,30 +44,24 @@ public:
 	bool OpenPackageFile(FArchive* Loader, EOpenPackageResult* OutErrorCode = nullptr);
 	bool OpenPackageFile(EOpenPackageResult* OutErrorCode = nullptr);
 
-	/** Reads information from the asset registry data table and converts it to FAssetData */
-	bool ReadAssetRegistryData(TArray<FAssetData*>& AssetDataList);
-
-	/** Attempts to get the class name of an object from the thumbnail cache for packages older than VER_UE4_ASSET_REGISTRY_TAGS */
-	bool ReadAssetDataFromThumbnailCache(TArray<FAssetData*>& AssetDataList);
+	/** Returns the LongPackageName from constructor if provided, otherwise calculates it from FPackageName::TryConvertFilenameToLongPackageName. */
+	bool TryGetLongPackageName(FString& OutLongPackageName) const;
+	/** Returns the LongPackageName as in TryGetPackageName; asserts if not found. */
+	FString GetLongPackageName() const;
 
-	/** Creates asset data reconstructing all the required data from cooked package info */
-	bool ReadAssetRegistryDataIfCookedPackage(TArray<FAssetData*>& AssetDataList, TArray<FString>& CookedPackageNamesWithoutAssetData);
+	/** Reads information from the asset registry data table and converts it to FAssetData */
+	bool ReadAssetRegistryData(TArray<FAssetData*>& AssetDataList, bool& bOutIsCookedWithoutAssetData);
 
+	/** Options for what to read in functions that read multiple things at once. */
+	enum class EReadOptions
+	{
+		None = 0,
+		PackageData = 1 << 0,
+		Dependencies = 1 << 1,
+		Default = PackageData | Dependencies,
+	};
 	/** Reads information used by the dependency graph */
-	bool ReadDependencyData(FPackageDependencyData& OutDependencyData);
-
-	/** Serializers for different package maps */
-	bool SerializeNameMap();
-	bool SerializeImportMap(TArray<FObjectImport>& OutImportMap);
-	bool SerializeExportMap(TArray<FObjectExport>& OutExportMap);
-	bool SerializeImportedClasses(const TArray<FObjectImport>& ImportMap, TArray<FName>& OutClassNames);
-	bool SerializeSoftPackageReferenceList(TArray<FName>& OutSoftPackageReferenceList);
-	bool SerializeSearchableNamesMap(FPackageDependencyData& OutDependencyData);
-	bool SerializeAssetRegistryDependencyData(FPackageDependencyData& DependencyData);
-	bool SerializePackageTrailer(FPackageDependencyData& DependencyData);
-
-	/** Returns flags the asset package was saved with */
-	uint32 GetPackageFlags() const;
+	bool ReadDependencyData(FPackageDependencyData& OutDependencyData, EReadOptions Options);
 
 	// Farchive implementation to redirect requests to the Loader
 	virtual void Serialize( void* V, int64 Length ) override;
@@ -80,9 +76,25 @@ public:
 	}
 
 private:
-	bool TryGetLongPackageName(FString& OutLongPackageName) const;
+	/** Attempts to get the class name of an object from the thumbnail cache for packages older than VER_UE4_ASSET_REGISTRY_TAGS */
+	bool ReadAssetDataFromThumbnailCache(TArray<FAssetData*>& AssetDataList);
+	/** Creates asset data reconstructing all the required data from cooked package info */
+	bool ReadAssetRegistryDataFromCookedPackage(TArray<FAssetData*>& AssetDataList, bool& bOutIsCookedWithoutAssetData);
 	bool StartSerializeSection(int64 Offset);
 
+	/** Serializers for different package maps */
+	bool SerializeNameMap();
+	bool SerializeImportMap(TArray<FObjectImport>& OutImportMap);
+	bool SerializeExportMap(TArray<FObjectExport>& OutExportMap);
+	bool SerializeImportedClasses(const TArray<FObjectImport>& ImportMap, TArray<FName>& OutClassNames);
+	bool SerializeSoftPackageReferenceList(TArray<FName>& OutSoftPackageReferenceList);
+	bool SerializeSearchableNamesMap(FPackageDependencyData& OutDependencyData);
+	bool SerializeAssetRegistryDependencyData(FPackageDependencyData& DependencyData);
+	bool SerializePackageTrailer(FAssetPackageData& PackageData);
+
+	/** Returns flags the asset package was saved with */
+	uint32 GetPackageFlags() const;
+
 	FString LongPackageName;
 	FString PackageFilename;
 	/* Loader is the interface used to read the bytes from the package's repository. All interpretation of the bytes is done by serializing into *this, which is also an FArchive. */
@@ -93,3 +105,4 @@ private:
 	int64 AssetRegistryDependencyDataOffset;
 	bool bLoaderOwner;
 };
+ENUM_CLASS_FLAGS(FPackageReader::EReadOptions);
diff --git a/Engine/Source/Runtime/AssetRegistry/Public/AssetRegistry/IAssetRegistry.h b/Engine/Source/Runtime/AssetRegistry/Public/AssetRegistry/IAssetRegistry.h
--- a/Engine/Source/Runtime/AssetRegistry/Public/AssetRegistry/IAssetRegistry.h
+++ b/Engine/Source/Runtime/AssetRegistry/Public/AssetRegistry/IAssetRegistry.h
@@ -320,6 +320,8 @@ public:
 	UE_DEPRECATED(5.0, "Receiving a pointer is not threadsafe. Use GetAssetPackageDataCopy instead.")
 	virtual const FAssetPackageData* GetAssetPackageData(FName PackageName) const = 0;
 
+	virtual bool DoesPackageExistOnDisk(FName PackageName, FString* OutCorrectCasePackageName = nullptr, FString* OutExtension = nullptr) const = 0;
+
 	/** Uses the asset registry to look for ObjectRedirectors. This will follow the chain of redirectors. It will return the original path if no redirectors are found */
 	virtual FName GetRedirectedObjectPath(const FName ObjectPath) const = 0;
 
