AssetRegistry Redirector references fixup: Skip the fixup for package versions that do not need it. The fixup pessimizes the cook for any package that contains both redirectors and other assets with editor-only references.
#rnx

diff --git a/Engine/Source/Runtime/AssetRegistry/Private/AssetDataGatherer.cpp b/Engine/Source/Runtime/AssetRegistry/Private/AssetDataGatherer.cpp
--- a/Engine/Source/Runtime/AssetRegistry/Private/AssetDataGatherer.cpp
+++ b/Engine/Source/Runtime/AssetRegistry/Private/AssetDataGatherer.cpp
@@ -32,7 +32,7 @@ namespace AssetDataGathererConstants
 	constexpr int32 SingleThreadFilesPerBatch = 3;
 	constexpr int32 ExpectedMaxBatchSize = 100;
 	constexpr int32 MinSecondsToElapseBeforeCacheWrite = 60;
-	static constexpr uint32 CacheSerializationMagic = 0x1F20181F; // Versioning and integrity checking
+	static constexpr uint32 CacheSerializationMagic = 0x087DF133; // Versioning and integrity checking
 }
 
 namespace UE
@@ -3041,21 +3041,26 @@ bool FAssetDataGatherer::ReadAssetFile(FPackageReader& PackageReader, TArray<FAs
 		return false;
 	}
 
-	if (EnumHasAnyFlags(Options, FPackageReader::EReadOptions::Dependencies))
+	if (PackageReader.UEVer() >= VER_UE4_ASSETREGISTRY_DEPENDENCYFLAGS && PackageReader.UEVer() < VER_UE4_CORRECT_LICENSEE_FLAG)
 	{
-		// DEPRECATION_TODO: Remove this fixup-on-load once we bump EUnrealEngineObjectUEVersion VER_UE4_ASSETREGISTRY_DEPENDENCYFLAGS and therefore all projects will resave each ObjectRedirector
-		// UObjectRedirectors were originally incorrectly marked as having editor-only imports, since UObjectRedirector is an editor-only class. But UObjectRedirectors are followed during cooking
-		// and so their imports should be considered used-in-game. Mark all dependencies in the package as used in game if the package has a UObjectRedirector object
-		FName RedirectorClassName = UObjectRedirector::StaticClass()->GetFName();
-		if (Algo::AnyOf(AssetDataList, [RedirectorClassName](FAssetData* AssetData) { return AssetData->AssetClass == RedirectorClassName; }))
+		if (EnumHasAnyFlags(Options, FPackageReader::EReadOptions::Dependencies))
 		{
-			TBitArray<>& ImportUsedInGame = DependencyData.ImportUsedInGame;
-			for (int32 ImportNum = ImportUsedInGame.Num(), Index = 0; Index < ImportNum; ++Index)
+			// In version VER_UE4_ASSETREGISTRY_DEPENDENCYFLAGS, UObjectRedirectors were incorrectly saved as having
+			// editor-only imports, since UObjectRedirector is an editor-only class. But UObjectRedirectors are
+			// followed during cooking and so their imports should be considered used-in-game. SavePackage was fixed
+			// to save them as in-game imports by adding HasNonEditorOnlyReferences; the next version bump after that
+			// fix was VER_UE4_CORRECT_LICENSEE_FLAG. Mark all dependencies in the affected version as used in game
+			// if the package has a UObjectRedirector object.
+			FName RedirectorClassName = UObjectRedirector::StaticClass()->GetFName();
+			if (Algo::AnyOf(AssetDataList, [RedirectorClassName](FAssetData* AssetData) { return AssetData->AssetClass == RedirectorClassName; }))
 			{
-				ImportUsedInGame[Index] = true;
+				TBitArray<>& ImportUsedInGame = DependencyData.ImportUsedInGame;
+				for (int32 ImportNum = ImportUsedInGame.Num(), Index = 0; Index < ImportNum; ++Index)
+				{
+					ImportUsedInGame[Index] = true;
+				}
 			}
 		}
-		// END DEPRECATION_TODO
 	}
 
 	return true;
