PackageReader and EditorDomain: add all objects of type UScriptStruct to the list of ImportedClasses. UScriptStruct objects might be used for serialization even if they are not the class of an export in the package. One example of this is UDataTable's RowStruct.
#rnx

diff --git a/Engine/Source/Editor/UnrealEd/Private/EditorDomain/EditorDomainUtils.cpp b/Engine/Source/Editor/UnrealEd/Private/EditorDomain/EditorDomainUtils.cpp
--- a/Engine/Source/Editor/UnrealEd/Private/EditorDomain/EditorDomainUtils.cpp
+++ b/Engine/Source/Editor/UnrealEd/Private/EditorDomain/EditorDomainUtils.cpp
@@ -427,7 +427,7 @@ FClassDigestData* FPrecacheClassDigest::GetRecursive(FName ClassName, bool bAllo
 	if (bAllowRedirects)
 	{
 		FCoreRedirectObjectName ClassNameRedirect(NameStringBuffer);
-		FCoreRedirectObjectName RedirectedClassNameRedirect = FCoreRedirects::GetRedirectedName(ECoreRedirectFlags::Type_Class, ClassNameRedirect);
+		FCoreRedirectObjectName RedirectedClassNameRedirect = FCoreRedirects::GetRedirectedName(ECoreRedirectFlags::Type_Class | ECoreRedirectFlags::Type_Struct, ClassNameRedirect);
 		if (ClassNameRedirect != RedirectedClassNameRedirect)
 		{
 			NameStringBuffer = RedirectedClassNameRedirect.ToString();
diff --git a/Engine/Source/Runtime/AssetRegistry/Private/AssetDataGatherer.cpp b/Engine/Source/Runtime/AssetRegistry/Private/AssetDataGatherer.cpp
--- a/Engine/Source/Runtime/AssetRegistry/Private/AssetDataGatherer.cpp
+++ b/Engine/Source/Runtime/AssetRegistry/Private/AssetDataGatherer.cpp
@@ -32,7 +32,7 @@ namespace AssetDataGathererConstants
 	constexpr int32 SingleThreadFilesPerBatch = 3;
 	constexpr int32 ExpectedMaxBatchSize = 100;
 	constexpr int32 MinSecondsToElapseBeforeCacheWrite = 60;
-	static constexpr uint32 CacheSerializationMagic = 0x087DF133; // Versioning and integrity checking
+	static constexpr uint32 CacheSerializationMagic = 0xC8F81894; // Versioning and integrity checking
 }
 
 namespace UE
diff --git a/Engine/Source/Runtime/AssetRegistry/Private/PackageReader.cpp b/Engine/Source/Runtime/AssetRegistry/Private/PackageReader.cpp
--- a/Engine/Source/Runtime/AssetRegistry/Private/PackageReader.cpp
+++ b/Engine/Source/Runtime/AssetRegistry/Private/PackageReader.cpp
@@ -559,11 +559,15 @@ bool FPackageReader::SerializeImportMap(TArray<FObjectImport>& OutImportMap)
 	return true;
 }
 
+static FName CoreUObjectPackageName(TEXT("/Script/CoreUObject"));
+static FName ScriptStructName(TEXT("ScriptStruct"));
+
 bool FPackageReader::SerializeImportedClasses(const TArray<FObjectImport>& ImportMap, TArray<FName>& OutClassNames)
 {
 	OutClassNames.Reset();
 
 	TSet<int32> ClassImportIndices;
+	// Any import that is specified as the class of an export is an imported class
 	if (PackageFileSummary.ExportCount > 0)
 	{
 		if (!StartSerializeSection(PackageFileSummary.ExportOffset))
@@ -593,6 +597,21 @@ bool FPackageReader::SerializeImportedClasses(const TArray<FObjectImport>& Impor
 			}
 		}
 	}
+	// Any imports of types UScriptStruct are an imported struct and need to be added to ImportedClasses
+	// This covers e.g. DataTable, which has a RowStruct pointer that it uses in its native serialization to
+	// serialize data into its rows
+	// TODO: Projects may create their own ScriptStruct subclass, and if they use one of these subclasses
+	// as a serialized-external-struct-pointer then we will miss it. In a future implementation we will 
+	// change the PackageReader to report all imports, and allow the AssetRegistry to decide which ones
+	// are classes based on its class database.
+	for (int32 ImportIndex = 0; ImportIndex < ImportMap.Num(); ++ImportIndex)
+	{
+		const FObjectImport& ObjectImport = ImportMap[ImportIndex];
+		if (ObjectImport.ClassPackage == CoreUObjectPackageName && ObjectImport.ClassName == ScriptStructName)
+		{
+			ClassImportIndices.Add(ImportIndex);
+		}
+	}
 
 	TArray<FName, TInlineAllocator<5>>  ParentChain;
 	FNameBuilder ClassObjectPath;
