AssetRegistry: Add EnumerateAllPackages.
#rnx

diff --git a/Engine/Source/Runtime/AssetRegistry/Private/AssetRegistry.cpp b/Engine/Source/Runtime/AssetRegistry/Private/AssetRegistry.cpp
--- a/Engine/Source/Runtime/AssetRegistry/Private/AssetRegistry.cpp
+++ b/Engine/Source/Runtime/AssetRegistry/Private/AssetRegistry.cpp
@@ -2267,6 +2267,16 @@ TOptional<FAssetPackageData> UAssetRegistryImpl::GetAssetPackageDataCopy(FName P
 	return AssetPackageData ? *AssetPackageData : TOptional<FAssetPackageData>();
 }
 
+void UAssetRegistryImpl::EnumerateAllPackages(TFunctionRef<void(FName PackageName, const FAssetPackageData& PackageData)> Callback) const
+{
+	FReadScopeLock InterfaceScopeLock(InterfaceLock);
+	for (const TPair<FName, const FAssetPackageData*>& Pair : GuardedData.GetState().GetAssetPackageDataMap())
+	{
+		Callback(Pair.Key, *Pair.Value);
+	}
+}
+
+
 bool UAssetRegistryImpl::DoesPackageExistOnDisk(FName PackageName, FString* OutCorrectCasePackageName, FString* OutExtension) const
 {
 	auto CalculateExtension = [](const FString& PackageNameStr, TConstArrayView<FAssetData> Assets) -> FString
diff --git a/Engine/Source/Runtime/AssetRegistry/Private/AssetRegistry.h b/Engine/Source/Runtime/AssetRegistry/Private/AssetRegistry.h
--- a/Engine/Source/Runtime/AssetRegistry/Private/AssetRegistry.h
+++ b/Engine/Source/Runtime/AssetRegistry/Private/AssetRegistry.h
@@ -60,6 +60,7 @@ public:
 	virtual bool GetReferencers(FName PackageName, TArray<FName>& OutReferencers, UE::AssetRegistry::EDependencyCategory Category = UE::AssetRegistry::EDependencyCategory::Package, const UE::AssetRegistry::FDependencyQuery& Flags = UE::AssetRegistry::FDependencyQuery()) const override;
 	virtual const FAssetPackageData* GetAssetPackageData(FName PackageName) const override;
 	virtual TOptional<FAssetPackageData> GetAssetPackageDataCopy(FName PackageName) const override;
+	virtual void EnumerateAllPackages(TFunctionRef<void(FName PackageName, const FAssetPackageData& PackageData)> Callback) const override;
 	virtual bool DoesPackageExistOnDisk(FName PackageName, FString* OutCorrectCasePackageName = nullptr, FString* OutExtension = nullptr) const override;
 	virtual FName GetRedirectedObjectPath(const FName ObjectPath) const override;
 	virtual bool GetAncestorClassNames(FName ClassName, TArray<FName>& OutAncestorClassNames) const override;
diff --git a/Engine/Source/Runtime/AssetRegistry/Public/AssetRegistry/IAssetRegistry.h b/Engine/Source/Runtime/AssetRegistry/Public/AssetRegistry/IAssetRegistry.h
--- a/Engine/Source/Runtime/AssetRegistry/Public/AssetRegistry/IAssetRegistry.h
+++ b/Engine/Source/Runtime/AssetRegistry/Public/AssetRegistry/IAssetRegistry.h
@@ -320,6 +320,12 @@ public:
 	UE_DEPRECATED(5.0, "Receiving a pointer is not threadsafe. Use GetAssetPackageDataCopy instead.")
 	virtual const FAssetPackageData* GetAssetPackageData(FName PackageName) const = 0;
 
+	/**
+	 * Enumerate all PackageDatas in the AssetRegistry. The callback is called from within the AssetRegistry's lock, so it must not call
+	 * arbitrary code that could call back into the AssetRegistry; doing so would deadlock.
+	 */
+	virtual void EnumerateAllPackages(TFunctionRef<void(FName PackageName, const FAssetPackageData& PackageData)> Callback) const = 0;
+
 	virtual bool DoesPackageExistOnDisk(FName PackageName, FString* OutCorrectCasePackageName = nullptr, FString* OutExtension = nullptr) const = 0;
 
 	/** Uses the asset registry to look for ObjectRedirectors. This will follow the chain of redirectors. It will return the original path if no redirectors are found */
