AssetDataGatherer: Change serialization format of dependencies in the editor's AssetDataGatherer cache to improve load performance.
#rnx

diff --git a/Engine/Source/Runtime/AssetRegistry/Private/AssetDataGatherer.cpp b/Engine/Source/Runtime/AssetRegistry/Private/AssetDataGatherer.cpp
--- a/Engine/Source/Runtime/AssetRegistry/Private/AssetDataGatherer.cpp
+++ b/Engine/Source/Runtime/AssetRegistry/Private/AssetDataGatherer.cpp
@@ -32,7 +32,7 @@ namespace AssetDataGathererConstants
 	constexpr int32 SingleThreadFilesPerBatch = 3;
 	constexpr int32 ExpectedMaxBatchSize = 100;
 	constexpr int32 MinSecondsToElapseBeforeCacheWrite = 60;
-	static constexpr uint32 CacheSerializationMagic = 0xC8F81894; // Versioning and integrity checking
+	static constexpr uint32 CacheSerializationMagic = 0x9D39D87B; // Versioning and integrity checking
 }
 
 namespace UE
@@ -3054,10 +3054,9 @@ bool FAssetDataGatherer::ReadAssetFile(FPackageReader& PackageReader, TArray<FAs
 			FName RedirectorClassName = UObjectRedirector::StaticClass()->GetFName();
 			if (Algo::AnyOf(AssetDataList, [RedirectorClassName](FAssetData* AssetData) { return AssetData->AssetClass == RedirectorClassName; }))
 			{
-				TBitArray<>& ImportUsedInGame = DependencyData.ImportUsedInGame;
-				for (int32 ImportNum = ImportUsedInGame.Num(), Index = 0; Index < ImportNum; ++Index)
+				for (FPackageDependencyData::FPackageDependency& Dependency : DependencyData.PackageDependencies)
 				{
-					ImportUsedInGame[Index] = true;
+					Dependency.Property |= UE::AssetRegistry::EDependencyProperty::Game;
 				}
 			}
 		}
diff --git a/Engine/Source/Runtime/AssetRegistry/Private/AssetRegistry.cpp b/Engine/Source/Runtime/AssetRegistry/Private/AssetRegistry.cpp
--- a/Engine/Source/Runtime/AssetRegistry/Private/AssetRegistry.cpp
+++ b/Engine/Source/Runtime/AssetRegistry/Private/AssetRegistry.cpp
@@ -3760,10 +3760,10 @@ bool ReadAssetFile(FPackageReader& PackageReader, IAssetRegistry::FLoadPackageRe
 
 		if (InOutData.bGetDependencies)
 		{
-			InOutData.DataDependencies.Reset(DependencyData.ImportMap.Num());
-			for (int32 ImportIdx = 0; ImportIdx < DependencyData.ImportMap.Num(); ++ImportIdx)
+			InOutData.DataDependencies.Reset(DependencyData.PackageDependencies.Num());
+			for (FPackageDependencyData::FPackageDependency& Dependency : DependencyData.PackageDependencies)
 			{
-				InOutData.DataDependencies.Emplace(DependencyData.GetImportPackageName(ImportIdx));
+				InOutData.DataDependencies.Emplace(Dependency.PackageName);
 			}
 		}
 	}
@@ -4171,65 +4171,50 @@ void FAssetRegistryImpl::DependencyDataGathered(const double TickStartTime, TRin
 			// Don't bother registering dependencies on these packages, every package in the game will depend on them
 			static TArray<FName> ScriptPackagesToSkip = TArray<FName>{ TEXT("/Script/CoreUObject"), TEXT("/Script/Engine"), TEXT("/Script/BlueprintGraph"), TEXT("/Script/UnrealEd") };
 
-			// Determine the new package dependencies
+			// Conditionally add package dependencies
 			TMap<FName, FDependsNode::FPackageFlagSet> PackageDependencies;
-			check(Result.ImportUsedInGame.Num() == Result.ImportMap.Num());
-			for (int32 ImportIdx = 0; ImportIdx < Result.ImportMap.Num(); ++ImportIdx)
+			for (FPackageDependencyData::FPackageDependency& DependencyData : Result.PackageDependencies)
 			{
-				FName AssetReference = Result.GetImportPackageName(ImportIdx);
-
-				// Should we skip this because it's too common?
-				if (ScriptPackagesToSkip.Contains(AssetReference))
+				// Skip hard dependencies to the common script packages
+				FName DependencyPackageName = DependencyData.PackageName;
+				if (EnumHasAnyFlags(DependencyData.Property, EDependencyProperty::Hard) && ScriptPackagesToSkip.Contains(DependencyPackageName))
 				{
 					continue;
 				}
 				FName Redirected = FCoreRedirects::GetRedirectedName(ECoreRedirectFlags::Type_Package,
-					FCoreRedirectObjectName(NAME_None, NAME_None, AssetReference)).PackageName;
-				AssetReference = Redirected;
+					FCoreRedirectObjectName(NAME_None, NAME_None, DependencyPackageName)).PackageName;
+				DependencyPackageName = Redirected;
 
-				EDependencyProperty DependencyProperty = EDependencyProperty::Build | EDependencyProperty::Hard;
-				DependencyProperty |= Result.ImportUsedInGame[ImportIdx] ? EDependencyProperty::Game : EDependencyProperty::None;
-				PackageDependencies.FindOrAdd(AssetReference).Add(FDependsNode::PackagePropertiesToByte(DependencyProperty));
-			}
-
-			check(Result.SoftPackageUsedInGame.Num() == Result.SoftPackageReferenceList.Num());
-			for (int32 SoftPackageIdx = 0; SoftPackageIdx < Result.SoftPackageReferenceList.Num(); ++SoftPackageIdx)
-			{
-				FName AssetReference = Result.SoftPackageReferenceList[SoftPackageIdx];
-				FName Redirected = FCoreRedirects::GetRedirectedName(ECoreRedirectFlags::Type_Package,
-					FCoreRedirectObjectName(NAME_None, NAME_None, AssetReference)).PackageName;
-				AssetReference = Redirected;
-
-				EDependencyProperty DependencyProperty = UE::AssetRegistry::EDependencyProperty::Build;
-				DependencyProperty |= (Result.SoftPackageUsedInGame[SoftPackageIdx] ? EDependencyProperty::Game : EDependencyProperty::None);
-				PackageDependencies.FindOrAdd(AssetReference).Add(FDependsNode::PackagePropertiesToByte(DependencyProperty));
+				FDependsNode::FPackageFlagSet& PackageFlagSet = PackageDependencies.FindOrAdd(DependencyPackageName);
+				PackageFlagSet.Add(FDependsNode::PackagePropertiesToByte(DependencyData.Property));
 			}
 
 			// Doubly-link all of the PackageDependencies
 			for (TPair<FName, FDependsNode::FPackageFlagSet>& NewDependsIt : PackageDependencies)
 			{
-				FDependsNode* DependsNode = State.CreateOrFindDependsNode(NewDependsIt.Key);
+				FName DependencyPackageName = NewDependsIt.Key;
+				FAssetIdentifier Identifier(DependencyPackageName);
+				FDependsNode* DependsNode = State.CreateOrFindDependsNode(Identifier);
 
 				// Handle failure of CreateOrFindDependsNode
 				// And Skip dependencies to self 
 				if (DependsNode != nullptr && DependsNode != Node)
 				{
-					const FAssetIdentifier& Identifier = DependsNode->GetIdentifier();
-					if (DependsNode->GetConnectionCount() == 0 && Identifier.IsPackage())
+					if (DependsNode->GetConnectionCount() == 0)
 					{
 						// This was newly created, see if we need to read the script package Guid
-						const FNameBuilder PackageName(Identifier.PackageName);
+						const FNameBuilder DependencyPackageNameStr(DependencyPackageName);
 
-						if (FPackageName::IsScriptPackage(PackageName))
+						if (FPackageName::IsScriptPackage(DependencyPackageNameStr))
 						{
 							// Get the guid off the script package, it is updated when script is changed so we need to refresh it every run
-							UPackage* Package = FindPackage(nullptr, *PackageName);
+							UPackage* Package = FindPackage(nullptr, *DependencyPackageNameStr);
 
 							if (Package)
 							{
-								FAssetPackageData* ScriptPackageData = State.CreateOrGetAssetPackageData(Identifier.PackageName);
+								FAssetPackageData* ScriptPackageData = State.CreateOrGetAssetPackageData(DependencyPackageName);
 								PRAGMA_DISABLE_DEPRECATION_WARNINGS
-									ScriptPackageData->PackageGuid = Package->GetGuid();
+								ScriptPackageData->PackageGuid = Package->GetGuid();
 								PRAGMA_ENABLE_DEPRECATION_WARNINGS
 							}
 						}
@@ -4240,44 +4225,13 @@ void FAssetRegistryImpl::DependencyDataGathered(const double TickStartTime, TRin
 				}
 			}
 
-			for (const TPair<FPackageIndex, TArray<FName>>& SearchableNameList : Result.SearchableNamesMap)
+			// Add node for all name references
+			for (FPackageDependencyData::FSearchableNamesDependency& NamesDependency : Result.SearchableNameDependencies)
 			{
-				FName ObjectName;
-				FName PackageName;
-
-				// Find object and package name from linker
-				FPackageIndex LinkerIndex = SearchableNameList.Key;
-				if (LinkerIndex.IsExport())
-				{
-					// Package name has to be this package, take a guess at object name
-					PackageName = Result.PackageName;
-					ObjectName = FName(*FPackageName::GetLongPackageAssetName(Result.PackageName.ToString()));
-				}
-				else if (LinkerIndex.IsImport())
+				for (FName& ValueName : NamesDependency.ValueNames)
 				{
-					FObjectResource* Resource = &Result.ImpExp(LinkerIndex);
-					FPackageIndex OuterLinkerIndex = Resource->OuterIndex;
-					check(OuterLinkerIndex.IsNull() || OuterLinkerIndex.IsImport());
-					if (!OuterLinkerIndex.IsNull())
-					{
-						ObjectName = Resource->ObjectName;
-						while (!OuterLinkerIndex.IsNull())
-						{
-							Resource = &Result.ImpExp(OuterLinkerIndex);
-							OuterLinkerIndex = Resource->OuterIndex;
-							check(OuterLinkerIndex.IsNull() || OuterLinkerIndex.IsImport());
-						}
-					}
-					PackageName = Resource->ObjectName;
-				}
-
-				for (FName NameReference : SearchableNameList.Value)
-				{
-					FAssetIdentifier AssetId = FAssetIdentifier(PackageName, ObjectName, NameReference);
-
-					// Add node for all name references
+					FAssetIdentifier AssetId(NamesDependency.PackageName, NamesDependency.ObjectName, ValueName);
 					FDependsNode* DependsNode = State.CreateOrFindDependsNode(AssetId);
-
 					if (DependsNode != nullptr)
 					{
 						Node->AddDependency(DependsNode, EDependencyCategory::SearchableName, EDependencyProperty::None);
diff --git a/Engine/Source/Runtime/AssetRegistry/Private/AssetRegistryState.cpp b/Engine/Source/Runtime/AssetRegistry/Private/AssetRegistryState.cpp
--- a/Engine/Source/Runtime/AssetRegistry/Private/AssetRegistryState.cpp
+++ b/Engine/Source/Runtime/AssetRegistry/Private/AssetRegistryState.cpp
@@ -1426,7 +1426,11 @@ SIZE_T FAssetRegistryState::GetAllocatedSize(bool bLogDetailed) const
 		UE_LOG(LogAssetRegistry, Log, TEXT("Dependency Arrays Size: %" SIZE_T_FMT "k"), DependenciesSize / 1024);
 	}
 
-	SIZE_T PackageDataSize = CachedPackageData.Num() * sizeof(FAssetPackageData);
+	SIZE_T PackageDataSize = CachedPackageData.Num() * (sizeof(FAssetPackageData) + sizeof(FAssetPackageData*));
+	for (const TPair<FName, FAssetPackageData*>& PackageDataPair : CachedPackageData)
+	{
+		PackageDataSize += PackageDataPair.Value->GetAllocatedSize();
+	}
 
 	SIZE_T TotalBytes = MapMemory + AssetDataSize + TagMemoryUsage.GetFixedSize()  + TagMemoryUsage.GetLooseSize() + DependNodesSize + DependenciesSize + PackageDataSize + MapArrayMemory;
 
diff --git a/Engine/Source/Runtime/AssetRegistry/Private/PackageDependencyData.cpp b/Engine/Source/Runtime/AssetRegistry/Private/PackageDependencyData.cpp
--- a/Engine/Source/Runtime/AssetRegistry/Private/PackageDependencyData.cpp
+++ b/Engine/Source/Runtime/AssetRegistry/Private/PackageDependencyData.cpp
@@ -2,12 +2,15 @@
 
 #include "PackageDependencyData.h"
 
-FName FPackageDependencyData::GetImportPackageName(int32 ImportIndex)
+FName FPackageDependencyData::GetImportPackageName(const TArray<FObjectImport>& ImportMap, int32 ImportIndex)
 {
-	FPackageIndex LinkerIndex = FPackageIndex::FromImport(ImportIndex);
-	while (LinkerIndex.IsImport())
+	for (int32 NumCycles = 0; NumCycles < ImportMap.Num(); ++NumCycles)
 	{
-		FObjectImport& Resource = Imp(LinkerIndex);
+		if (!ImportMap.IsValidIndex(ImportIndex))
+		{
+			return NAME_None;
+		}
+		const FObjectImport& Resource = ImportMap[ImportIndex];
 		// If the import has a package name set, then that's the import package name,
 		if (Resource.HasPackageName())
 		{
@@ -18,7 +21,82 @@ FName FPackageDependencyData::GetImportPackageName(int32 ImportIndex)
 		{
 			return Resource.ObjectName;
 		}
-		LinkerIndex = Resource.OuterIndex;
+		if (!Resource.OuterIndex.IsImport())
+		{
+				return NAME_None;
+		}
+		ImportIndex = Resource.OuterIndex.ToImport();
 	}
 	return NAME_None;
 }
+
+void FPackageDependencyData::LoadDependenciesFromPackageHeader(FName SourcePackageName, TArray<FObjectImport>& ImportMap,
+	TArray<FName>& SoftPackageReferenceList, TMap<FPackageIndex, TArray<FName>>& SearchableNames,
+	TBitArray<>& ImportUsedInGame, TBitArray<>& SoftPackageUsedInGame)
+{
+	using namespace UE::AssetRegistry;
+
+	PackageDependencies.Reset(ImportMap.Num() + SoftPackageReferenceList.Num());
+	check(ImportMap.Num() == ImportUsedInGame.Num());
+	for (int32 ImportIdx = 0; ImportIdx < ImportMap.Num(); ++ImportIdx)
+	{
+		FName DependencyPackageName = GetImportPackageName(ImportMap, ImportIdx);
+		EDependencyProperty DependencyProperty = EDependencyProperty::Build | EDependencyProperty::Hard;
+		DependencyProperty |= ImportUsedInGame[ImportIdx] ? EDependencyProperty::Game : EDependencyProperty::None;
+		PackageDependencies.Add({ DependencyPackageName, DependencyProperty });
+	}
+
+	check(SoftPackageReferenceList.Num() == SoftPackageUsedInGame.Num());
+	for (int32 SoftPackageIdx = 0; SoftPackageIdx < SoftPackageReferenceList.Num(); ++SoftPackageIdx)
+	{
+		FName DependencyPackageName = SoftPackageReferenceList[SoftPackageIdx];
+		FAssetIdentifier AssetId(DependencyPackageName);
+		EDependencyProperty DependencyProperty = UE::AssetRegistry::EDependencyProperty::Build; // !EDependencyProperty::Hard
+		DependencyProperty |= (SoftPackageUsedInGame[SoftPackageIdx] ? EDependencyProperty::Game : EDependencyProperty::None);
+		PackageDependencies.Add({ DependencyPackageName, DependencyProperty });
+	}
+
+	SearchableNameDependencies.Reset(SearchableNames.Num());
+	for (const TPair<FPackageIndex, TArray<FName>>& SearchableNameList : SearchableNames)
+	{
+		FName ObjectName;
+		FName DependencyPackageName;
+
+		// Find object and package name from linker
+		FPackageIndex LinkerIndex = SearchableNameList.Key;
+		if (LinkerIndex.IsExport())
+		{
+			// Package name has to be this package, take a guess at object name
+			DependencyPackageName = SourcePackageName;
+			ObjectName = FName(*FPackageName::GetLongPackageAssetName(DependencyPackageName.ToString()));
+		}
+		else if (LinkerIndex.IsImport())
+		{
+			int32 ImportIndex = LinkerIndex.ToImport();
+			if (!ImportMap.IsValidIndex(ImportIndex))
+			{
+				continue;
+			}
+			FObjectImport& Resource = ImportMap[ImportIndex];
+			FPackageIndex OuterLinkerIndex = Resource.OuterIndex;
+			if (!OuterLinkerIndex.IsNull())
+			{
+				ObjectName = Resource.ObjectName;
+			}
+			DependencyPackageName = GetImportPackageName(ImportMap, ImportIndex);
+			if (DependencyPackageName.IsNone())
+			{
+				continue;
+			}
+		}
+		else
+		{
+			continue;
+		}
+
+		FSearchableNamesDependency& Dependency = SearchableNameDependencies.Emplace_GetRef();
+		Dependency.PackageName = DependencyPackageName;
+		Dependency.ObjectName = ObjectName;
+		Dependency.ValueNames = SearchableNameList.Value;
+	}
+}
diff --git a/Engine/Source/Runtime/AssetRegistry/Private/PackageDependencyData.h b/Engine/Source/Runtime/AssetRegistry/Private/PackageDependencyData.h
--- a/Engine/Source/Runtime/AssetRegistry/Private/PackageDependencyData.h
+++ b/Engine/Source/Runtime/AssetRegistry/Private/PackageDependencyData.h
@@ -5,68 +5,79 @@
 #include "CoreMinimal.h"
 
 #include "AssetRegistry/AssetData.h"
+#include "AssetRegistry/IAssetRegistry.h"
 #include "Containers/BitArray.h"
-#include "Misc/AssetRegistryInterface.h"
 #include "UObject/Linker.h"
 
-class FPackageDependencyData : public FLinkerTables
+class FPackageDependencyData
 {
 public:
+	struct FPackageDependency
+	{
+		FName PackageName;
+		UE::AssetRegistry::EDependencyProperty Property;
+		friend FArchive& operator<<(FArchive& Ar, FPackageDependency& Dependency)
+		{
+			Ar << Dependency.PackageName;
+			uint8 PropertyAsInteger = static_cast<uint8>(Dependency.Property);
+			Ar << PropertyAsInteger;
+			Dependency.Property = static_cast<UE::AssetRegistry::EDependencyProperty>(PropertyAsInteger);
+			return Ar;
+		}
+	};
+	struct FSearchableNamesDependency
+	{
+		FName PackageName;
+		FName ObjectName;
+		TArray<FName> ValueNames;
+		friend FArchive& operator<<(FArchive& Ar, FSearchableNamesDependency& Dependency)
+		{
+			Ar << Dependency.PackageName << Dependency.ObjectName << Dependency.ValueNames;
+			return Ar;
+		}
+	};
+
 	/** The name of the package that dependency data is gathered from */
 	FName PackageName;
 
 	/** Asset Package data, gathered at the same time as dependency data */
 	FAssetPackageData PackageData;
 
-	TBitArray<> ImportUsedInGame;
-	TBitArray<> SoftPackageUsedInGame;
+	// Dependency Data
+	TArray<FPackageDependency> PackageDependencies;
+	TArray<FSearchableNamesDependency> SearchableNameDependencies;
 
-	bool bHasPackageData = false; // Transient
-	bool bHasDependencyData = false; // Transient
+	// Transient Flags indicating which types of data have been gathered
+	bool bHasPackageData = false;
+	bool bHasDependencyData = false;
 
 	/**
-	 * Return the package name of the UObject represented by the specified import. 
-	 * 
-	 * @param	PackageIndex	package index for the resource to get the name for
-	 *
-	 * @return	the path name of the UObject represented by the resource at PackageIndex, or the empty string if this isn't an import
-	 */
-	FName GetImportPackageName(int32 ImportIndex);
-
-	/**
-	 * Serialize as part of the registry cache. This is not meant to be serialized as part of a package so  it does not handle versions normally
-	 * To version this data change FAssetRegistryVersion or CacheSerializationVersion
+	 * Serialize as part of the registry cache. This is not meant to be serialized as part of a package so it does not handle versions normally
+	 * To version this data change FAssetRegistryVersion or AssetDataGathererConstants::CacheSerializationMagic
 	 */
 	void SerializeForCache(FArchive& Ar)
 	{
 		Ar << PackageName;
-		Ar << ImportMap;
-		Ar << SoftPackageReferenceList;
-		Ar << SearchableNamesMap;
 		PackageData.SerializeForCache(Ar);
-		Ar << ImportUsedInGame;
-		Ar << SoftPackageUsedInGame;
-		if (Ar.IsLoading())
-		{
-			if (!IsValid())
-			{
-				Ar.SetError();
-			}
-		}
+		Ar << PackageDependencies;
+		Ar << SearchableNameDependencies;
 	}
 
-	bool IsValid() const
-	{
-		return ImportUsedInGame.Num() == ImportMap.Num() &&
-			SoftPackageUsedInGame.Num() == SoftPackageReferenceList.Num();
-	}
+	void LoadDependenciesFromPackageHeader(FName PackageName, TArray<FObjectImport>& ImportMap,
+		TArray<FName>& SoftPackageReferenceList, TMap<FPackageIndex, TArray<FName>>& SearchableNames,
+		TBitArray<>& ImportUsedInGame, TBitArray<>& SoftPackageUsedInGame);
 
 	/** Returns the amount of memory allocated by this container, not including sizeof(*this). */
 	SIZE_T GetAllocatedSize() const
 	{
-		SIZE_T Result = FLinkerTables::GetAllocatedSize();
-		Result += ImportUsedInGame.GetAllocatedSize();
-		Result += SoftPackageUsedInGame.GetAllocatedSize();
+		SIZE_T Result = PackageDependencies.GetAllocatedSize();
+		Result += SearchableNameDependencies.GetAllocatedSize();
+		Result += PackageData.GetAllocatedSize();
 		return Result;
 	}
+
+private:
+	FName GetImportPackageName(const TArray<FObjectImport>& ImportMap, int32 ImportIndex);
+
+
 };
diff --git a/Engine/Source/Runtime/AssetRegistry/Private/PackageReader.cpp b/Engine/Source/Runtime/AssetRegistry/Private/PackageReader.cpp
--- a/Engine/Source/Runtime/AssetRegistry/Private/PackageReader.cpp
+++ b/Engine/Source/Runtime/AssetRegistry/Private/PackageReader.cpp
@@ -11,6 +11,7 @@
 #include "Misc/Guid.h"
 #include "Misc/PackageName.h"
 #include "UObject/Class.h"
+#include "UObject/Linker.h"
 #include "UObject/PackageTrailer.h"
 
 FPackageReader::FPackageReader()
@@ -253,13 +254,14 @@ bool FPackageReader::ReadAssetRegistryData(TArray<FAssetData*>& AssetDataList, b
 	return true;
 }
 
-bool FPackageReader::SerializeAssetRegistryDependencyData(FPackageDependencyData& DependencyData)
+bool FPackageReader::SerializeAssetRegistryDependencyData(TBitArray<>& OutImportUsedInGame, TBitArray<>& OutSoftPackageUsedInGame,
+	const TArray<FObjectImport>& ImportMap, const TArray<FName>& SoftPackageReferenceList)
 {
 	if (AssetRegistryDependencyDataOffset == INDEX_NONE)
 	{
 		// For old package versions that did not write out the dependency flags, set default values of the flags
-		DependencyData.ImportUsedInGame.Init(true, DependencyData.ImportMap.Num());
-		DependencyData.SoftPackageUsedInGame.Init(true, DependencyData.SoftPackageReferenceList.Num());
+		OutImportUsedInGame.Init(true, ImportMap.Num());
+		OutSoftPackageUsedInGame.Init(true, SoftPackageReferenceList.Num());
 		return true;
 	}
 
@@ -268,8 +270,9 @@ bool FPackageReader::SerializeAssetRegistryDependencyData(FPackageDependencyData
 		return false;
 	}
 
-	if (!UE::AssetRegistry::ReadPackageDataDependencies(*this, DependencyData.ImportUsedInGame, DependencyData.SoftPackageUsedInGame)
-		|| !DependencyData.IsValid())
+	if (!UE::AssetRegistry::ReadPackageDataDependencies(*this, OutImportUsedInGame, OutSoftPackageUsedInGame) ||
+		OutImportUsedInGame.Num() != ImportMap.Num() ||
+		OutSoftPackageUsedInGame.Num() != SoftPackageReferenceList.Num())
 	{
 		UE_PACKAGEREADER_CORRUPTPACKAGE_WARNING("SerializeAssetRegistryDependencyData", PackageFilename);
 		return false;
@@ -439,7 +442,9 @@ bool FPackageReader::ReadDependencyData(FPackageDependencyData& OutDependencyDat
 	{
 		return false;
 	}
-	if (!SerializeImportMap(OutDependencyData.ImportMap))
+
+	TArray<FObjectImport> ImportMap;
+	if (!SerializeImportMap(ImportMap))
 	{
 		return false;
 	}
@@ -457,7 +462,7 @@ bool FPackageReader::ReadDependencyData(FPackageDependencyData& OutDependencyDat
 		PackageData.FileVersionLicenseeUE = PackageFileSummary.GetFileVersionLicenseeUE();
 		PackageData.SetIsLicenseeVersion(PackageFileSummary.SavedByEngineVersion.IsLicenseeVersion());
 
-		if (!SerializeImportedClasses(OutDependencyData.ImportMap, PackageData.ImportedClasses))
+		if (!SerializeImportedClasses(ImportMap, PackageData.ImportedClasses))
 		{
 			return false;
 		}
@@ -470,27 +475,29 @@ bool FPackageReader::ReadDependencyData(FPackageDependencyData& OutDependencyDat
 	if (EnumHasAnyFlags(Options, EReadOptions::Dependencies))
 	{
 		OutDependencyData.bHasDependencyData = true;
-		if (!SerializeSoftPackageReferenceList(OutDependencyData.SoftPackageReferenceList))
+		TArray<FName> SoftPackageReferenceList;
+		if (!SerializeSoftPackageReferenceList(SoftPackageReferenceList))
 		{
 			return false;
 		}
-		if (!SerializeSearchableNamesMap(OutDependencyData))
+		FLinkerTables SearchableNames;
+		if (!SerializeSearchableNamesMap(SearchableNames))
 		{
 			return false;
 		}
-		if (!SerializeAssetRegistryDependencyData(OutDependencyData))
+
+		TBitArray<> ImportUsedInGame;
+		TBitArray<> SoftPackageUsedInGame;
+		if (!SerializeAssetRegistryDependencyData(ImportUsedInGame, SoftPackageUsedInGame, ImportMap,
+			SoftPackageReferenceList))
 		{
 			return false;
 		}
-	}
-	else
-	{
-		// We used OutDependencyData.ImportMap to store data for SerializeImportedClasses to use, but 
-		// our contract says we should leave it empty if !EReadOptions::Dependencies
-		OutDependencyData.ImportMap.Reset();
+
+		OutDependencyData.LoadDependenciesFromPackageHeader(OutDependencyData.PackageName, ImportMap, SoftPackageReferenceList,
+			SearchableNames.SearchableNamesMap, ImportUsedInGame, SoftPackageUsedInGame);
 	}
 
-	checkf(OutDependencyData.IsValid(), TEXT("We should have early exited above rather than creating invalid dependency data"));
 	return true;
 }
 
@@ -767,7 +774,7 @@ bool FPackageReader::SerializeSoftPackageReferenceList(TArray<FName>& OutSoftPac
 	return true;
 }
 
-bool FPackageReader::SerializeSearchableNamesMap(FPackageDependencyData& OutDependencyData)
+bool FPackageReader::SerializeSearchableNamesMap(FLinkerTables& OutSearchableNames)
 {
 	if (UEVer() >= VER_UE4_ADDED_SEARCHABLE_NAMES && PackageFileSummary.SearchableNamesOffset > 0)
 	{
@@ -777,7 +784,7 @@ bool FPackageReader::SerializeSearchableNamesMap(FPackageDependencyData& OutDepe
 			return false;
 		}
 
-		OutDependencyData.SerializeSearchableNamesMap(*this);
+		OutSearchableNames.SerializeSearchableNamesMap(*this);
 		if (IsError())
 		{
 			UE_PACKAGEREADER_CORRUPTPACKAGE_WARNING("SerializeSearchableNamesMapInvalidSearchableNamesMap", PackageFilename);
diff --git a/Engine/Source/Runtime/AssetRegistry/Private/PackageReader.h b/Engine/Source/Runtime/AssetRegistry/Private/PackageReader.h
--- a/Engine/Source/Runtime/AssetRegistry/Private/PackageReader.h
+++ b/Engine/Source/Runtime/AssetRegistry/Private/PackageReader.h
@@ -10,6 +10,7 @@
 
 struct FAssetData;
 class FAssetPackageData;
+class FLinkerTables;
 class FPackageDependencyData;
 
 class FPackageReader : public FArchiveUObject
@@ -88,8 +89,9 @@ private:
 	bool SerializeExportMap(TArray<FObjectExport>& OutExportMap);
 	bool SerializeImportedClasses(const TArray<FObjectImport>& ImportMap, TArray<FName>& OutClassNames);
 	bool SerializeSoftPackageReferenceList(TArray<FName>& OutSoftPackageReferenceList);
-	bool SerializeSearchableNamesMap(FPackageDependencyData& OutDependencyData);
-	bool SerializeAssetRegistryDependencyData(FPackageDependencyData& DependencyData);
+	bool SerializeSearchableNamesMap(FLinkerTables& OutSearchableNames);
+	bool SerializeAssetRegistryDependencyData(TBitArray<>& OutImportUsedInGame, TBitArray<>& OutSoftPackageUsedInGame,
+		const TArray<FObjectImport>& ImportMap, const TArray<FName>& SoftPackageReferenceList);
 	bool SerializePackageTrailer(FAssetPackageData& PackageData);
 
 	/** Returns flags the asset package was saved with */
diff --git a/Engine/Source/Runtime/CoreUObject/Public/AssetRegistry/AssetData.h b/Engine/Source/Runtime/CoreUObject/Public/AssetRegistry/AssetData.h
--- a/Engine/Source/Runtime/CoreUObject/Public/AssetRegistry/AssetData.h
+++ b/Engine/Source/Runtime/CoreUObject/Public/AssetRegistry/AssetData.h
@@ -748,6 +748,12 @@ public:
 	COREUOBJECT_API void SerializeForCache(FArchive& Ar);
 	COREUOBJECT_API void SerializeForCacheOldVersion(FArchive& Ar, FAssetRegistryVersion::Type Version);
 
+	/** Returns the amount of memory allocated by this container, not including sizeof(*this). */
+	SIZE_T GetAllocatedSize() const
+	{
+		return ImportedClasses.GetAllocatedSize();
+	}
+
 private:
 	FORCEINLINE void SerializeForCacheInternal(FArchive& Ar, FAssetPackageData& PackageData, FAssetRegistryVersion::Type Version);
 
diff --git a/Engine/Source/Runtime/CoreUObject/Public/Misc/AssetRegistryInterface.h b/Engine/Source/Runtime/CoreUObject/Public/Misc/AssetRegistryInterface.h
--- a/Engine/Source/Runtime/CoreUObject/Public/Misc/AssetRegistryInterface.h
+++ b/Engine/Source/Runtime/CoreUObject/Public/Misc/AssetRegistryInterface.h
@@ -49,6 +49,9 @@ namespace UE::AssetRegistry
 	// 
 	enum class EDependencyCategory : uint8
 	{
+		// This enum is serialized by value into the runtime AssetRegistry and the AssetDataGatherer.
+		// If any values change or are removed, bump FAssetRegistryVersion and add backwards compatibility in FAssetRegistryState::Serialize.
+
 		Package = 0x01,			// The target asset of any package dependency is expected to be loadable whenever the source asset is available; see EDependencyProperty for different requirements of the loadability.
 		Manage = 0x02,			// The target asset of any manage dependency is managed (e.g. given a disk layout location) either directly or indirectly by the source asset. Used by UAssetManager.
 		SearchableName = 0x04,  // Targets of SearchableName dependencies are FNames Keys inside of an Asset. The Source Asset of the dependency has a value for that Key. Used to search for Assets with a given Key,Value for the custom Key.
@@ -64,6 +67,8 @@ namespace UE::AssetRegistry
 	enum class EDependencyProperty : uint8
 	{
 		None = 0,
+		// This enum is serialized by value into the runtime AssetRegistry and the AssetDataGatherer.
+		// If any values change or are removed, bump FAssetRegistryVersion and add backwards compatibility in FAssetRegistryState::Serialize.
 
 		// Package Dependencies
 		PackageMask = 0x7,
