AssetDataGatherer: Fix typo in ReadAssetFile found by static analysis.
#rnx

diff --git a/Engine/Source/Runtime/AssetRegistry/Private/AssetDataGatherer.cpp b/Engine/Source/Runtime/AssetRegistry/Private/AssetDataGatherer.cpp
--- a/Engine/Source/Runtime/AssetRegistry/Private/AssetDataGatherer.cpp
+++ b/Engine/Source/Runtime/AssetRegistry/Private/AssetDataGatherer.cpp
@@ -3018,8 +3018,8 @@ bool FAssetDataGatherer::ReadAssetFile(const FString& AssetLongPackageName, cons
 	else
 	{
 		return ReadAssetFile(PackageReader, AssetDataList, DependencyData, CookedPackageNamesWithoutAssetData,
-			(bGatherAssetPackageData ? FPackageReader::EReadOptions::PackageData : FPackageReader::EReadOptions::PackageData) |
-			(bGatherDependsData ? FPackageReader::EReadOptions::Dependencies : FPackageReader::EReadOptions::PackageData));
+			(bGatherAssetPackageData ? FPackageReader::EReadOptions::PackageData : FPackageReader::EReadOptions::None) |
+			(bGatherDependsData ? FPackageReader::EReadOptions::Dependencies : FPackageReader::EReadOptions::None));
 	}
 }
 
