Preload the AssetDataGatherer's cache to remove a stall when waiting on the gatherer to load the data for packages requested early in engine startup.
#rnx

diff --git a/Engine/Source/Runtime/AssetRegistry/Private/AssetDataGatherer.cpp b/Engine/Source/Runtime/AssetRegistry/Private/AssetDataGatherer.cpp
--- a/Engine/Source/Runtime/AssetRegistry/Private/AssetDataGatherer.cpp
+++ b/Engine/Source/Runtime/AssetRegistry/Private/AssetDataGatherer.cpp
@@ -18,6 +18,8 @@
 #include "Misc/AsciiSet.h"
 #include "Misc/Char.h"
 #include "Misc/CommandLine.h"
+#include "Misc/DelayedAutoRegister.h"
+#include "Misc/Parse.h"
 #include "Misc/Paths.h"
 #include "Misc/PathViews.h"
 #include "Misc/ScopeLock.h"
@@ -35,11 +37,31 @@ namespace AssetDataGathererConstants
 	static constexpr uint32 CacheSerializationMagic = 0x9D39D87B; // Versioning and integrity checking
 }
 
-namespace UE
-{
-namespace AssetDataGather
+namespace UE::AssetDataGather::Private
 {
-namespace Private
+	/** A structure to hold serialized cache data from async loads before adding it to the Gatherer's main cache. */
+	struct FCachePayload
+	{
+		TUniquePtr<FName[]> PackageNames;
+		TUniquePtr<FDiskCachedAssetData[]> AssetDatas;
+		int32 NumAssets = 0;
+		bool bSucceeded = false;
+		void Reset()
+		{
+			PackageNames.Reset();
+			AssetDatas.Reset();
+			NumAssets = 0;
+			bSucceeded = false;
+		}
+	};
+
+	void SerializeCacheSave(FAssetRegistryWriter& Ar, const TArray<TPair<FName, FDiskCachedAssetData*>>& AssetsToSave);
+	FCachePayload SerializeCacheLoad(FAssetRegistryReader& Ar);
+	FCachePayload LoadCacheFile(FStringView CacheFilename);
+}
+
+
+namespace UE::AssetDataGather::Private
 {
 
 /** InOutResult = Value, but without shrinking the string to fit. */
@@ -2448,6 +2470,10 @@ protected:
 			SaveBlock();
 		}
 	}
+	const FBlock& GetCurrentBlock() const
+	{
+		return Block;
+	}
 
 	void Save(FMemoryView Data)
 	{
@@ -2512,6 +2538,7 @@ class FChecksumArchiveWriter : public FChecksumArchiveBase
 public:
 	using FChecksumArchiveBase::FChecksumArchiveBase;
 	virtual void Serialize(void* V, int64 Len) override { Save(FMemoryView(V, Len)); }
+	virtual int64 Tell() override { return InnerArchive.Tell() + GetCurrentBlock().GetUsedSize(); }
 };
 
 class FChecksumArchiveReader : public FChecksumArchiveBase
@@ -2519,11 +2546,139 @@ class FChecksumArchiveReader : public FChecksumArchiveBase
 public:
 	using FChecksumArchiveBase::FChecksumArchiveBase;
 	virtual void Serialize(void* V, int64 Len) override { Load(FMutableMemoryView(V, Len)); }
+	virtual int64 Tell() override { return InnerArchive.Tell() - GetCurrentBlock().GetRemainingSize(); }
+};
+
+
+/**
+ * Settings about whether to use cache data for the AssetDataGatherer; these settings are shared by
+ * FPreloader and the FAssetDataGatherer.
+ */
+struct FPreloadSettings
+{
+	void Initialize()
+	{
+		if (bInitialized)
+		{
+			return;
+		}
+		bInitialized = true;
+
+		FString ProjectIntermediateDir = FPaths::ProjectIntermediateDir();
+		bForceDependsGathering = FParse::Param(FCommandLine::Get(), TEXT("ForceDependsGathering"));
+		bGatherDependsData = (GIsEditor && !FParse::Param(FCommandLine::Get(), TEXT("NoDependsGathering"))) || bForceDependsGathering;
+		bCacheEnabled = !FParse::Param(FCommandLine::Get(), TEXT("NoAssetRegistryCache")) && !FParse::Param(FCommandLine::Get(), TEXT("multiprocess"));
+		bool bAsyncEnabled = FPlatformProcess::SupportsMultithreading() && FTaskGraphInterface::IsRunning();
+
+		MonolithicCacheFilename = FPaths::ProjectIntermediateDir() / (bGatherDependsData ? TEXT("CachedAssetRegistry.bin") : TEXT("CachedAssetRegistryNoDeps.bin"));
+#if UE_EDITOR // See note on FPreloader for why we only allow preloading if UE_EDITOR
+		if (bCacheEnabled && bAsyncEnabled)
+		{
+			bPreloadMonolithicCache = GIsEditor && (!IsRunningCommandlet() || IsRunningCookCommandlet());
+		}
+		else
+#endif
+		{
+			bPreloadMonolithicCache = false;
+		}
+	}
+	bool IsCacheEnabled() const
+	{
+		return bCacheEnabled;
+	}
+	bool IsPreloadMonolithicCache() const
+	{
+		return bPreloadMonolithicCache;
+	}
+	bool IsGatherDependsData() const
+	{
+		return bGatherDependsData;
+	}
+	bool IsForceDependsGathering() const
+	{
+		return bForceDependsGathering;
+	}
+	const FString& GetMonolithicCacheFilename() const
+	{
+		return MonolithicCacheFilename;
+	}
+private:
+	FString MonolithicCacheFilename;
+	bool bForceDependsGathering = false;
+	bool bGatherDependsData = false;
+	bool bCacheEnabled = false;
+	bool bPreloadMonolithicCache = false;
+	bool bInitialized = false;
+};
+FPreloadSettings GPreloadSettings;
+
+#if UE_EDITOR
+/** A class to preload the monolithic cache used by FAssetDataGatherer. Preloading the cache allows us to
+ * start very early in editor startup, so that we have time to finish the cache load before the engine starts
+ * making package load requests that need to use the gathered data.
+ *
+ * In UE_EDITOR, we know the values we need to decide whether we can preload early enough that it is useful to preload
+ * In other configurations we do not know those parameters for sure until EDelayedRegisterRunPhase::ShaderTypesReady,
+ *  which occurs around the same time as UAssetRegistryImpl is created, so it is not useful to preload.
+ */
+class FPreloader : public FDelayedAutoRegisterHelper
+{
+public:
+	FPreloader()
+		// The callback needs to occur after GIsEditor, ProjectIntermediateDir, IsRunningCommandlet, and
+		// IsRunningCookCommandlet have been set
+		:FDelayedAutoRegisterHelper(EDelayedRegisterRunPhase::IniSystemReady, [this]() { DelayedInitialize(); })
+	{
+	}
+
+	~FPreloader()
+	{
+		PreloadReady.Wait();
+	}
+
+	FCachePayload Consume()
+	{
+		check(bInitialized);
+		if (bConsumed)
+		{
+			return FCachePayload();
+		}
+		bConsumed = true;
+		PreloadReady.Wait();
+		PreloadReady.Reset();
+		return FCachePayload(MoveTemp(Payload));
+	}
+
+private:
+
+	void DelayedInitialize()
+	{
+		GPreloadSettings.Initialize();
+
+		if (GPreloadSettings.IsPreloadMonolithicCache())
+		{
+			if (IFileManager::Get().FileExists(*GPreloadSettings.GetMonolithicCacheFilename()))
+			{
+				PreloadReady = Async(EAsyncExecution::TaskGraph, [this]() { LoadAsync(); });
+			}
+		}
+		bInitialized = true;
+	}
+
+	void LoadAsync()
+	{
+		Payload = LoadCacheFile(GPreloadSettings.GetMonolithicCacheFilename());
+	}
+
+	TFuture<void> PreloadReady;
+	FCachePayload Payload;
+	bool bInitialized = false;
+	bool bConsumed = false;
 };
+FPreloader GPreloader;
+#endif
 
-} // namespace Private
-} // namespace AssetDataGather
-} // namespace UE
+} // namespace UE::AssetDataGather::Private
 
 FAssetDataGatherer::FAssetDataGatherer(const TArray<FString>& InLongPackageNamesDenyList, const TArray<FString>& InMountRelativePathsDenyList, bool bInIsSynchronous)
 	: Thread(nullptr)
@@ -2534,9 +2689,7 @@ FAssetDataGatherer::FAssetDataGatherer(const TArray<FString>& InLongPackageNames
 	, bSaveAsyncCacheTriggered(false)
 	, SearchStartTime(0)
 	, LastCacheWriteTime(0.0)
-	, bUseMonolithicCache(false)
-	, bUseTickManagedCache(false)
-	, bHasLoadedTickManagedCache(false)
+	, bHasLoadedMonolithicCache(false)
 	, bDiscoveryIsComplete(false)
 	, bIsComplete(false)
 	, bIsIdle(false)
@@ -2549,11 +2702,13 @@ FAssetDataGatherer::FAssetDataGatherer(const TArray<FString>& InLongPackageNames
 {
 	using namespace UE::AssetDataGather::Private;
 
-	bool bForceDependsGathering = FParse::Param(FCommandLine::Get(), TEXT("ForceDependsGathering"));
-	bGatherAssetPackageData = GIsEditor || bForceDependsGathering;
-	bGatherDependsData = (GIsEditor && !FParse::Param(FCommandLine::Get(), TEXT("NoDependsGathering"))) || bForceDependsGathering;
-
-	bCacheEnabled = !FParse::Param(FCommandLine::Get(), TEXT("NoAssetRegistryCache")) && !FParse::Param(FCommandLine::Get(), TEXT("multiprocess"));
+	GPreloadSettings.Initialize();
+	bGatherAssetPackageData = GIsEditor || GPreloadSettings.IsForceDependsGathering();
+	bGatherDependsData = GPreloadSettings.IsGatherDependsData();
+	bCacheEnabled = GPreloadSettings.IsCacheEnabled();
+	// If PreloadMonolithicCache is true, we are using MonolithicCache. Otherwise it may be set to true later if game/commandlet calls SearchAllAssets.
+	bUseMonolithicCache = GPreloadSettings.IsPreloadMonolithicCache();
+	LastCacheWriteTime = FPlatformTime::Seconds();
 
 #if !UE_BUILD_SHIPPING
 	bool bCommandlineSynchronous;
@@ -2591,27 +2746,16 @@ FAssetDataGatherer::~FAssetDataGatherer()
 	DiskCachedAssetBlocks.Empty();
 }
 
-void FAssetDataGatherer::SetUseMonolithicCache(bool bInUseMonolithicCache)
+void FAssetDataGatherer::ActivateMonolithicCache()
 {
 	FGathererScopeLock ResultsScopeLock(&ResultsLock);
-	if (!bCacheEnabled || bUseMonolithicCache == bInUseMonolithicCache)
+	if (!bCacheEnabled || bUseMonolithicCache )
 	{
 		return;
 	}
 
-	bUseMonolithicCache = bInUseMonolithicCache;
-	bHasLoadedTickManagedCache = false;
+	bUseMonolithicCache = true;
 	LastCacheWriteTime = FPlatformTime::Seconds();
-	if (bUseMonolithicCache)
-	{
-		bUseTickManagedCache = bUseMonolithicCache;
-		TickManagedCacheFilename = FPaths::ProjectIntermediateDir() / (bGatherDependsData ? TEXT("CachedAssetRegistry.bin") : TEXT("CachedAssetRegistryNoDeps.bin"));
-	}
-	else
-	{
-		bUseTickManagedCache = false;
-		TickManagedCacheFilename.Empty();
-	}
 }
 
 void FAssetDataGatherer::StartAsync()
@@ -2655,11 +2799,13 @@ uint32 FAssetDataGatherer::Run()
 
 void FAssetDataGatherer::InnerTickLoop(bool bInIsSynchronousTick, bool bContributeToCacheSave)
 {
+	using namespace UE::AssetDataGather::Private;
+
 	// Synchronous ticks during Wait contribute to saving of the async cache only if there is no dedicated async thread to do it (bIsSynchronous is true)
 	// The dedicated async thread always contributes
 	bContributeToCacheSave = !bInIsSynchronousTick || (bIsSynchronous && bContributeToCacheSave);
 
-	FString CacheFilename;
+	bool bShouldSaveMonolithicCache = false;
 	TArray<TPair<FName, FDiskCachedAssetData*>> AssetsToSave;
 	{
 		CHECK_IS_NOT_LOCKED_CURRENT_THREAD(ResultsLock);
@@ -2674,10 +2820,13 @@ void FAssetDataGatherer::InnerTickLoop(bool bInIsSynchronousTick, bool bContribu
 
 		if (bContributeToCacheSave)
 		{
-			TryReserveSaveAsyncCache(CacheFilename, AssetsToSave);
+			TryReserveSaveMonolithicCache(bShouldSaveMonolithicCache, AssetsToSave);
 		}
 	}
-	SaveCacheFileInternal(CacheFilename, AssetsToSave, true /* bIsAsyncCache */);
+	if (bShouldSaveMonolithicCache)
+	{
+		SaveCacheFileInternal(GPreloadSettings.GetMonolithicCacheFilename(), AssetsToSave, true /* bIsAsyncCache */);
+	}
 }
 
 FAssetDataGatherer::FScopedPause::FScopedPause(const FAssetDataGatherer& InOwner)
@@ -2740,7 +2889,7 @@ void FAssetDataGatherer::TickInternal(bool& bOutIsTickInterrupt)
 	TArray<FAssetData*, FBatchInlineAllocator> LocalAssetResults;
 	TArray<FPackageDependencyData, FBatchInlineAllocator> LocalDependencyResults;
 	TArray<FString, FBatchInlineAllocator> LocalCookedPackageNamesWithoutAssetDataResults;
-	FString LocalCacheFilename;
+	bool bLoadMonolithicCache = false;
 	double LocalLastCacheWriteTime = 0.0;
 	bool bWaitBatchCountDecremented = false;
 	bOutIsTickInterrupt = false;
@@ -2785,19 +2934,31 @@ void FAssetDataGatherer::TickInternal(bool& bOutIsTickInterrupt)
 			SetIsIdle(true);
 			return;
 		}
-		if (bUseTickManagedCache && !bHasLoadedTickManagedCache)
+		if (bUseMonolithicCache && !bHasLoadedMonolithicCache)
 		{
-			LocalCacheFilename = TickManagedCacheFilename;
+			bLoadMonolithicCache = true;
 		}
 		LocalLastCacheWriteTime = LastCacheWriteTime;
 	}
 
 	// Load the async cache if not yet loaded
-	if (!LocalCacheFilename.IsEmpty())
+	if (bLoadMonolithicCache)
 	{
-		LoadCacheFileInternal(LocalCacheFilename);
+		FCachePayload Payload;
+#if UE_EDITOR
+		if (GPreloadSettings.IsPreloadMonolithicCache())
+		{
+			Payload = GPreloader.Consume();
+		}
+		else
+#endif
+		{
+			Payload = UE::AssetDataGather::Private::LoadCacheFile(GPreloadSettings.GetMonolithicCacheFilename());
+		}
+		ConsumeCacheFile(MoveTemp(Payload));
+
 		FGathererScopeLock ResultsScopeLock(&ResultsLock);
-		bHasLoadedTickManagedCache = true;
+		bHasLoadedMonolithicCache = true;
 	}
 
 	struct FReadContext
@@ -2970,7 +3131,7 @@ void FAssetDataGatherer::TickInternal(bool& bOutIsTickInterrupt)
 			}
 		}
 
-		if (bUseTickManagedCache && !bIsSavingAsyncCache 
+		if (bUseMonolithicCache && !bIsSavingAsyncCache 
 			&& FPlatformTime::Seconds() - LocalLastCacheWriteTime >= AssetDataGathererConstants::MinSecondsToElapseBeforeCacheWrite)
 		{
 			bSaveAsyncCacheTriggered = true;
@@ -3312,50 +3473,45 @@ FString FAssetDataGatherer::GetCacheFilename(TConstArrayView<FString> CacheFileP
 
 void FAssetDataGatherer::LoadCacheFile(FStringView CacheFilename)
 {
+	using namespace UE::AssetDataGather::Private;
+	if (!bCacheEnabled)
+	{
+		return;
+	}
+
+	FCachePayload Payload = UE::AssetDataGather::Private::LoadCacheFile(CacheFilename);
 	FScopedPause ScopedPause(*this);
 	CHECK_IS_NOT_LOCKED_CURRENT_THREAD(ResultsLock);
 	FGathererScopeLock TickScopeLock(&TickLock);
-	LoadCacheFileInternal(CacheFilename);
+	ConsumeCacheFile(MoveTemp(Payload));
 }
 
-void FAssetDataGatherer::LoadCacheFileInternal(FStringView CacheFilename)
+void FAssetDataGatherer::ConsumeCacheFile(UE::AssetDataGather::Private::FCachePayload&& Payload)
 {
 	CHECK_IS_LOCKED_CURRENT_THREAD(TickLock);
-	if (!bCacheEnabled)
+	CHECK_IS_NOT_LOCKED_CURRENT_THREAD(ResultsLock);
+	if (!Payload.bSucceeded || Payload.NumAssets == 0)
 	{
 		return;
 	}
 
-	TRACE_CPUPROFILER_EVENT_SCOPE(LoadCacheFile);
-	// load the cached data
-	TUniquePtr<FArchive> FileAr(IFileManager::Get().CreateFileReader(*FString(CacheFilename), FILEREAD_Silent));
-	if (FileAr && !FileAr->IsError() && FileAr->TotalSize() > 2 * sizeof(uint32))
+	DiskCachedAssetDataMap.Reserve(DiskCachedAssetDataMap.Num() + Payload.NumAssets);
+	for (int32 AssetIndex = 0; AssetIndex < Payload.NumAssets; ++AssetIndex)
 	{
-		uint32 MagicNumber = 0;
-		*FileAr << MagicNumber;
-		if (!FileAr->IsError() && MagicNumber == AssetDataGathererConstants::CacheSerializationMagic)
-		{
-			FAssetRegistryVersion::Type RegistryVersion;
-			if (FAssetRegistryVersion::SerializeVersion(*FileAr, RegistryVersion) && RegistryVersion == FAssetRegistryVersion::LatestVersion)
-			{
-				UE::AssetDataGather::Private::FChecksumArchiveReader ChecksummingReader(*FileAr);
-				FAssetRegistryReader RegistryReader(ChecksummingReader);
-				if (!RegistryReader.IsError())
-				{
-					SerializeCacheLoad(RegistryReader);
-
-					FGathererScopeLock ResultsScopeLock(&ResultsLock);
-					DependencyResults.Reserve(DiskCachedAssetDataMap.Num());
-					AssetResults.Reserve(DiskCachedAssetDataMap.Num());
-				}
-			}
-		}
+		DiskCachedAssetDataMap.Add(*(Payload.PackageNames.Get() + AssetIndex),
+			(Payload.AssetDatas.Get() + AssetIndex)); // -C6385
 	}
+	DiskCachedAssetBlocks.Emplace(Payload.NumAssets, Payload.AssetDatas.Release());
+	Payload.Reset();
+
+	FGathererScopeLock ResultsScopeLock(&ResultsLock);
+	DependencyResults.Reserve(DiskCachedAssetDataMap.Num());
+	AssetResults.Reserve(DiskCachedAssetDataMap.Num());
 }
 
-void FAssetDataGatherer::TryReserveSaveAsyncCache(FString& OutCacheFilename, TArray<TPair<FName,FDiskCachedAssetData*>>& AssetsToSave)
+void FAssetDataGatherer::TryReserveSaveMonolithicCache(bool& bOutShouldSave, TArray<TPair<FName,FDiskCachedAssetData*>>& AssetsToSave)
 {
-	OutCacheFilename.Reset();
+	bOutShouldSave = false;
 	if (IsStopped)
 	{
 		return;
@@ -3367,12 +3523,9 @@ void FAssetDataGatherer::TryReserveSaveAsyncCache(FString& OutCacheFilename, TAr
 	CHECK_IS_LOCKED_CURRENT_THREAD(TickLock);
 	{
 		FGathererScopeLock ResultsScopeLock(&ResultsLock);
-		if (bUseTickManagedCache && !TickManagedCacheFilename.IsEmpty())
-		{
-			OutCacheFilename = TickManagedCacheFilename;
-		}
+		bOutShouldSave = bUseMonolithicCache;
 	}
-	if (!OutCacheFilename.IsEmpty())
+	if (bOutShouldSave)
 	{
 		GetAssetsToSave(TArrayView<const FString>(), AssetsToSave);
 		bIsSavingAsyncCache = true;
@@ -3436,7 +3589,7 @@ void FAssetDataGatherer::SaveCacheFileInternal(const FString& CacheFilename, con
 			// to save certain common tags as FName.
 			UE::AssetDataGather::Private::FChecksumArchiveWriter ChecksummingWriter(*FileAr);
 			FAssetRegistryWriter Ar(FAssetRegistryWriterOptions(), ChecksummingWriter);
-			SerializeCacheSave(Ar, AssetsToSave);
+			UE::AssetDataGather::Private::SerializeCacheSave(Ar, AssetsToSave);
 		}
 #else		
 		checkf(false, TEXT("Cannot save asset registry cache in this configuration"));
@@ -3459,11 +3612,12 @@ void FAssetDataGatherer::SaveCacheFileInternal(const FString& CacheFilename, con
 	}
 }
 
-void FAssetDataGatherer::SerializeCacheSave(FAssetRegistryWriter& Ar, const TArray<TPair<FName,FDiskCachedAssetData*>>& AssetsToSave)
+namespace UE::AssetDataGather::Private
 {
-#if ALLOW_NAME_BATCH_SAVING
-	using namespace UE::AssetDataGather::Private;
 
+void SerializeCacheSave(FAssetRegistryWriter& Ar, const TArray<TPair<FName,FDiskCachedAssetData*>>& AssetsToSave)
+{
+#if ALLOW_NAME_BATCH_SAVING
 	double SerializeStartTime = FPlatformTime::Seconds();
 
 	// serialize number of objects
@@ -3481,12 +3635,14 @@ void FAssetDataGatherer::SerializeCacheSave(FAssetRegistryWriter& Ar, const TArr
 #endif
 }
 
-void FAssetDataGatherer::SerializeCacheLoad(FAssetRegistryReader& Ar)
+FCachePayload SerializeCacheLoad(FAssetRegistryReader& Ar)
 {
-	CHECK_IS_LOCKED_CURRENT_THREAD(TickLock);
-	using namespace UE::AssetDataGather::Private;
-
 	double SerializeStartTime = FPlatformTime::Seconds();
+	ON_SCOPE_EXIT
+	{
+		UE_LOG(LogAssetRegistry, Verbose, TEXT("Asset data gatherer serialized in %0.6f seconds"), FPlatformTime::Seconds() - SerializeStartTime);
+	};
+
 	// serialize number of objects
 	int32 LocalNumAssets = 0;
 	Ar << LocalNumAssets;
@@ -3496,45 +3652,85 @@ void FAssetDataGatherer::SerializeCacheLoad(FAssetRegistryReader& Ar)
 	if (Ar.IsError() || LocalNumAssets < 0 || MaxPossibleNumAssets < LocalNumAssets)
 	{
 		Ar.SetError();
+		return FCachePayload();
 	}
-	else
+
+	if (LocalNumAssets == 0)
 	{
-		FSoftObjectPathSerializationScope SerializationScope(NAME_None, NAME_None, ESoftObjectPathCollectType::NonPackage, ESoftObjectPathSerializeType::AlwaysSerialize);
+		FCachePayload Payload = FCachePayload();
+		Payload.bSucceeded = true;
+		return Payload;
+	}
 
-		// allocate one single block for all asset data structs (to reduce tens of thousands of heap allocations)
-		FName* PackageNameBlock = new FName[LocalNumAssets];
-		FDiskCachedAssetData* AssetDataBlock = new FDiskCachedAssetData[LocalNumAssets];
-		for (int32 AssetIndex = 0; AssetIndex < LocalNumAssets; ++AssetIndex)
-		{
-			// Load the name first to add the entry to the tmap below
-			// Visual Studio Static Analyzer issues C6385 if we call Ar << PackageNameBlock[AssetIndex] or AssetDataBlock[AssetIndex].SerializeForCache
-			Ar << *(PackageNameBlock + AssetIndex); // -C6385
-			(AssetDataBlock + AssetIndex)->SerializeForCache(Ar); // -C6385
-			if (Ar.IsError())
-			{
-				// There was an error reading the cache. Bail out.
-				break;
-			}
-		}
+	FSoftObjectPathSerializationScope SerializationScope(NAME_None, NAME_None, ESoftObjectPathCollectType::NonPackage, ESoftObjectPathSerializeType::AlwaysSerialize);
 
+	// allocate one single block for all asset data structs (to reduce tens of thousands of heap allocations)
+	TUniquePtr<FName[]> PackageNameBlock(new FName[LocalNumAssets]);
+	TUniquePtr<FDiskCachedAssetData[]> AssetDataBlock(new FDiskCachedAssetData[LocalNumAssets]);
+	for (int32 AssetIndex = 0; AssetIndex < LocalNumAssets; ++AssetIndex)
+	{
+		// Load the name first to add the entry to the tmap below
+		// Visual Studio Static Analyzer issues C6385 if we call Ar << PackageNameBlock[AssetIndex] or AssetDataBlock[AssetIndex].SerializeForCache
+		Ar << *(PackageNameBlock.Get() + AssetIndex); // -C6385
+		(AssetDataBlock.Get() + AssetIndex)->SerializeForCache(Ar); // -C6385
 		if (Ar.IsError())
 		{
-			delete[] AssetDataBlock;
-			UE_LOG(LogAssetRegistry, Error, TEXT("There was an error loading the asset registry cache."));
-		}
-		else
-		{
-			DiskCachedAssetDataMap.Reserve(DiskCachedAssetDataMap.Num() + LocalNumAssets);
-			for (int32 AssetIndex = 0; AssetIndex < LocalNumAssets; ++AssetIndex)
-			{
-				DiskCachedAssetDataMap.Add(*(PackageNameBlock + AssetIndex), (AssetDataBlock + AssetIndex)); // -C6385
-			}
-			DiskCachedAssetBlocks.Emplace(LocalNumAssets, AssetDataBlock);
+			// There was an error reading the cache. Bail out.
+			break;
 		}
-		delete[] PackageNameBlock;
 	}
 
-	UE_LOG(LogAssetRegistry, Verbose, TEXT("Asset data gatherer serialized in %0.6f seconds"), FPlatformTime::Seconds() - SerializeStartTime);
+	if (Ar.IsError())
+	{
+		return FCachePayload();
+	}
+	FCachePayload Result;
+	Result.PackageNames = MoveTemp(PackageNameBlock);
+	Result.AssetDatas = MoveTemp(AssetDataBlock);
+	Result.NumAssets = LocalNumAssets;
+	Result.bSucceeded = true;
+	return Result;
+}
+
+FCachePayload LoadCacheFile(FStringView CacheFilename)
+{
+	TRACE_CPUPROFILER_EVENT_SCOPE(LoadCacheFile);
+
+	TUniquePtr<FArchive> FileAr(IFileManager::Get().CreateFileReader(*WriteToString<256>(CacheFilename), FILEREAD_Silent));
+	if (!FileAr || FileAr->IsError() || FileAr->TotalSize() < 2 * sizeof(uint32))
+	{
+		return FCachePayload();
+	}
+
+	uint32 MagicNumber = 0;
+	*FileAr << MagicNumber;
+	if (FileAr->IsError() || MagicNumber != AssetDataGathererConstants::CacheSerializationMagic)
+	{
+		return FCachePayload();
+	}
+
+	FAssetRegistryVersion::Type RegistryVersion;
+	if (!FAssetRegistryVersion::SerializeVersion(*FileAr, RegistryVersion) || RegistryVersion != FAssetRegistryVersion::LatestVersion)
+	{
+		return FCachePayload();
+	}
+
+	UE::AssetDataGather::Private::FChecksumArchiveReader ChecksummingReader(*FileAr);
+	FAssetRegistryReader RegistryReader(ChecksummingReader);
+	if (RegistryReader.IsError())
+	{
+		return FCachePayload();
+	}
+
+	FCachePayload Payload;
+	Payload = UE::AssetDataGather::Private::SerializeCacheLoad(RegistryReader);
+	if (!Payload.bSucceeded)
+	{
+		UE_LOG(LogAssetRegistry, Error, TEXT("There was an error loading the asset registry cache."));
+	}
+	return Payload;
+}
+
 }
 
 SIZE_T FAssetDataGatherer::GetAllocatedSize() const
@@ -3580,7 +3776,7 @@ SIZE_T FAssetDataGatherer::GetAllocatedSize() const
 	Result += GetArrayRecursiveAllocatedSize(CookedPackageNamesWithoutAssetDataResults);
 	Result += SearchTimes.GetAllocatedSize();
 	Result += GetArrayRecursiveAllocatedSize(DiscoveredPaths);
-	Result += TickManagedCacheFilename.GetAllocatedSize();
+	Result += GPreloadSettings.GetMonolithicCacheFilename().GetAllocatedSize();
 
 	Result += NewCachedAssetData.GetAllocatedSize();
 	for (const FDiskCachedAssetData* Value : NewCachedAssetData)
diff --git a/Engine/Source/Runtime/AssetRegistry/Private/AssetDataGatherer.h b/Engine/Source/Runtime/AssetRegistry/Private/AssetDataGatherer.h
--- a/Engine/Source/Runtime/AssetRegistry/Private/AssetDataGatherer.h
+++ b/Engine/Source/Runtime/AssetRegistry/Private/AssetDataGatherer.h
@@ -32,6 +32,7 @@ namespace UE::AssetDataGather::Private
 {
 class FAssetDataDiscovery;
 class FFilesToSearch;
+struct FCachePayload;
 struct FGatheredPathData;
 struct FPathExistence;
 struct FSetPathProperties;
@@ -63,7 +64,7 @@ public:
 	// Extra at-construction configuration 
 
 	/** Configure the gatherer to use a single monolithic cache for all files, and read/write this cache during ticks. */
-	void SetUseMonolithicCache(bool bInUseMonolithicCache);
+	void ActivateMonolithicCache();
 
 
 	// Controlling Async behavior
@@ -226,13 +227,13 @@ private:
 		TArray<FAssetData*>& AssetDataList, FPackageDependencyData& DependencyData,
 		TArray<FString>& CookedPackagesToLoadUponDiscovery, bool& OutCanRetry) const;
 
-	/** Helper for LoadCacheFile and internal callsites; occurs inside the proper critical section. */
-	void LoadCacheFileInternal(FStringView CacheFilename);
+	/** Add the given AssetDatas into DiskCachedAssetDataMap and DiskCachedAssetBlocks. */
+	void ConsumeCacheFile(UE::AssetDataGather::Private::FCachePayload&& Payload);
 	/**
-	 * If a save of the async cache has been triggered, get the cache filename and pointers to all elements that should be saved,
+	 * If a save of the monolithic cache has been triggered, get the cache filename and pointers to all elements that should be saved,
 	 * for later saving outside of the critical section.
 	 */
-	void TryReserveSaveAsyncCache(FString& OutCacheFilename, TArray<TPair<FName,FDiskCachedAssetData*>>& AssetsToSave);
+	void TryReserveSaveMonolithicCache(bool& bOutShouldSave, TArray<TPair<FName,FDiskCachedAssetData*>>& AssetsToSave);
 	/**
 	 * If the CacheFilename/AssetsToSave are non empty, save the cache file. 
 	 * This function reads the read-only-after-creation data from each FDiskCachedAssetData*, but otherwise does not use
@@ -244,10 +245,6 @@ private:
 	 * Filters the list of assets by child paths of the elements in SaveCacheLongPackageNameDirs, if it is non-empty.
 	 */
 	void GetAssetsToSave(TArrayView<const FString> SaveCacheLongPackageNameDirs, TArray<TPair<FName,FDiskCachedAssetData*>>& OutAssetsToSave);
-	/** Writes to file the timestamped cache of discovered assets. */
-	void SerializeCacheSave(FAssetRegistryWriter& Ar, const TArray<TPair<FName,FDiskCachedAssetData*>>& AssetsToSave);
-	/** Reads from file the timestamped cache of discovered assets, for quick loading of data for assets that have not changed on disk. */
-	void SerializeCacheLoad(FAssetRegistryReader& Ar);
 
 	/* Adds the given PackageName,DiskCachedAssetData pair into NewCachedAssetDataMap, and detects collisions for multiple files with the same PackageName */
 	void AddToCache(FName PackageName, FDiskCachedAssetData* DiskCachedAssetData);
@@ -346,12 +343,6 @@ private:
 	 * Read/writable only within ResultsLock.
 	 */
 	TArray<FString> DiscoveredPaths;
-	/**
-	 * If using a TickManaged cache, this is the name of the cache file.
-	 * The cachefile contains timestamp,guid,and asset data for each package file in the the previous process's gather.
-	 * Read/writable only within ResultsLock.
-	 */
-	FString TickManagedCacheFilename;
 
 	/** The current search start time, set in the first tick after idle and used for performance metrics when reporting results. Read/writable only within ResultsLock. */
 	double SearchStartTime;
@@ -363,18 +354,13 @@ private:
 	 */
 	int32 NumPathsToSearchAtLastSyncPoint;
 	/**
-	 * Track whether we are using a monolithic cache. The monolithic cache is only supported as a TickManaged cache.
-	 * Setting it to true sets up the TickManaged cache. Read/writable only within ResultsLock.
-	 */
-	bool bUseMonolithicCache;
-	/**
-	 * Track whether a cache file should be loaded and saved by the Tick function
-	 * (as opposed to being externally managed by calls to LoadCacheFile/SaveCacheFile).
+	 * Track whether we are using a monolithic cache that should be loaded/saved during tick.
+	 * Wether we are or not, the AssetRegistry can also call LoadCacheFile/ScanPathsSynchronous to load/save smaller files.
 	 * Read/writable only within ResultsLock.
 	 */
-	bool bUseTickManagedCache;
-	/** If bUseTickManagedCache is true, track whether the cache has been loaded. Read/writable only within ResultsLock. */
-	bool bHasLoadedTickManagedCache;
+	bool bUseMonolithicCache;
+	/** If bHasLoadedMonolithicCache is true, track whether the cache has been loaded. Read/writable only within ResultsLock. */
+	bool bHasLoadedMonolithicCache;
 	/** Track whether the Discovery subsystem has gone idle and we have read all filenames from it. Read/writable only within ResultsLock. */
 	bool bDiscoveryIsComplete;
 	/** Track whether this Gather has gone idle and a caller has read all search data from it. Read/writable only within ResultsLock. */
@@ -397,7 +383,7 @@ private:
 	TArray<TPair<int32, FDiskCachedAssetData*>> DiskCachedAssetBlocks;
 	/**
 	 * Map of PackageName to cached discovered assets that were loaded from disk.
-	 * This should only be modified in the loading section of SerializeCache calls.
+	 * This should only be modified by ConsumeCacheFile.
 	 * Read/Writable only within TickLock.
 	 */
 	TMap<FName, FDiskCachedAssetData*> DiskCachedAssetDataMap;
diff --git a/Engine/Source/Runtime/AssetRegistry/Private/AssetRegistry.cpp b/Engine/Source/Runtime/AssetRegistry/Private/AssetRegistry.cpp
--- a/Engine/Source/Runtime/AssetRegistry/Private/AssetRegistry.cpp
+++ b/Engine/Source/Runtime/AssetRegistry/Private/AssetRegistry.cpp
@@ -1440,7 +1440,7 @@ void FAssetRegistryImpl::SearchAllAssets(Impl::FEventContext& EventContext, bool
 		bSynchronousSearch = true;
 	}
 
-	Gatherer.SetUseMonolithicCache(true);
+	Gatherer.ActivateMonolithicCache();
 
 	// Add all existing mountpoints to the GlobalGatherer
 	// This will include Engine content, Game content, but also may include mounted content directories for one or more plugins.
