Expose loading an asset registry from disk.

diff --git a/Engine/Source/Runtime/AssetRegistry/Private/AssetRegistry.cpp b/Engine/Source/Runtime/AssetRegistry/Private/AssetRegistry.cpp
--- a/Engine/Source/Runtime/AssetRegistry/Private/AssetRegistry.cpp
+++ b/Engine/Source/Runtime/AssetRegistry/Private/AssetRegistry.cpp
@@ -26,7 +26,6 @@
 #include "Misc/ScopeRWLock.h"
 #include "PackageReader.h"
 #include "Serialization/ArrayReader.h"
-#include "Serialization/LargeMemoryReader.h"
 #include "Templates/UnrealTemplate.h"
 #include "UObject/ConstructorHelpers.h"
 #include "UObject/CoreRedirects.h"
@@ -127,25 +126,6 @@ FAssetRegistryInterface GAssetRegistryInterface;
 
 DEFINE_LOG_CATEGORY(LogAssetRegistry);
 
-static bool LoadAssetRegistry(const TCHAR* Path, const FAssetRegistryLoadOptions& Options, FAssetRegistryState& Out)
-{
-	check(Path);
-
-	TUniquePtr<FArchive> FileReader(IFileManager::Get().CreateFileReader(Path));
-	if (FileReader)
-	{
-		// It's faster to load the whole file into memory on a Gen5 console
-		TArray64<uint8> Data;
-		Data.SetNumUninitialized(FileReader->TotalSize());
-		FileReader->Serialize(Data.GetData(), Data.Num());
-		check(!FileReader->IsError());
-		
-		FLargeMemoryReader MemoryReader(Data.GetData(), Data.Num());
-		return Out.Load(MemoryReader, Options);
-	}
-
-	return false;
-}
 
 namespace UE::AssetRegistry::Premade
 {
@@ -288,7 +268,7 @@ private:
 		const int32 ThreadReduction = 2; // This thread + main thread already has work to do 
 		int32 MaxWorkers = CanLoadAsync() ? FPlatformMisc::NumberOfCoresIncludingHyperthreads() - ThreadReduction : 0;
 		Options.ParallelWorkers = FMath::Clamp(MaxWorkers, 0, 16);
-		bLoadSucceeded = LoadAssetRegistry(*ARPath, Options, Payload);
+		bLoadSucceeded = FAssetRegistryState::LoadFromDisk(*ARPath, Options, Payload);
 		UE_CLOG(!bLoadSucceeded, LogAssetRegistry, Warning, TEXT("Premade AssetRegistry path %s existed but failed to load."), *ARPath);
 		UE_CLOG(bLoadSucceeded, LogAssetRegistry, Log, TEXT("Premade AssetRegistry loaded from '%s'"), *ARPath);
 		return bLoadSucceeded;
diff --git a/Engine/Source/Runtime/AssetRegistry/Private/AssetRegistryImpl.h b/Engine/Source/Runtime/AssetRegistry/Private/AssetRegistryImpl.h
--- a/Engine/Source/Runtime/AssetRegistry/Private/AssetRegistryImpl.h
+++ b/Engine/Source/Runtime/AssetRegistry/Private/AssetRegistryImpl.h
@@ -500,9 +500,6 @@ void RunAssetsThroughFilter(TArray<FAssetData>& AssetDataList, const FARCompiled
 void InitializeSerializationOptionsFromIni(FAssetRegistrySerializationOptions& Options, const FString& PlatformIniName,
 	ESerializationTarget Target = ESerializationTarget::ForGame);
 
-/** Load via -Serialize the Out state from file at the given path */
-bool LoadAssetRegistry(const TCHAR* Path, const FAssetRegistryLoadOptions& Options, FAssetRegistryState& Out);
-
 /** Gets the current availability of an asset, primarily for streaming install purposes. */
 EAssetAvailability::Type GetAssetAvailability(const FAssetData& AssetData);
 /** Gets an ETA or percentage complete for an asset that is still in the process of being installed. */
diff --git a/Engine/Source/Runtime/AssetRegistry/Private/AssetRegistryState.cpp b/Engine/Source/Runtime/AssetRegistry/Private/AssetRegistryState.cpp
--- a/Engine/Source/Runtime/AssetRegistry/Private/AssetRegistryState.cpp
+++ b/Engine/Source/Runtime/AssetRegistry/Private/AssetRegistryState.cpp
@@ -8,6 +8,7 @@
 #include "Misc/Paths.h"
 #include "Misc/PathViews.h"
 #include "Serialization/ArrayReader.h"
+#include "Serialization/LargeMemoryReader.h"
 #include "Serialization/MemoryWriter.h"
 #include "Misc/ConfigCacheIni.h"
 #include "UObject/UObjectHash.h"
@@ -1157,12 +1158,16 @@ bool FAssetRegistryState::Save(FArchive& OriginalAr, const FAssetRegistrySeriali
 	return !OriginalAr.IsError();
 }
 
-bool FAssetRegistryState::Load(FArchive& OriginalAr, const FAssetRegistryLoadOptions& Options)
+bool FAssetRegistryState::Load(FArchive& OriginalAr, const FAssetRegistryLoadOptions& Options, FAssetRegistryVersion::Type* OutVersion)
 {
 	LLM_SCOPE(ELLMTag::AssetRegistry);
 
 	FAssetRegistryVersion::Type Version = FAssetRegistryVersion::LatestVersion;
 	FAssetRegistryVersion::SerializeVersion(OriginalAr, Version);
+	if (OutVersion != nullptr)
+	{
+		*OutVersion = Version;
+	}
 
 	FSoftObjectPathSerializationScope SerializationScope(NAME_None, NAME_None, ESoftObjectPathCollectType::NonPackage, ESoftObjectPathSerializeType::AlwaysSerialize);
 
@@ -1195,6 +1200,26 @@ bool FAssetRegistryState::Load(FArchive& OriginalAr, const FAssetRegistryLoadOpt
 	return !OriginalAr.IsError();
 }
 
+/* static */ bool FAssetRegistryState::LoadFromDisk(const TCHAR* InPath, const FAssetRegistryLoadOptions& InOptions, FAssetRegistryState& OutState, FAssetRegistryVersion::Type* OutVersion)
+{
+	check(InPath);
+
+	TUniquePtr<FArchive> FileReader(IFileManager::Get().CreateFileReader(InPath));
+	if (FileReader)
+	{
+		// It's faster to load the whole file into memory on a Gen5 console
+		TArray64<uint8> Data;
+		Data.SetNumUninitialized(FileReader->TotalSize());
+		FileReader->Serialize(Data.GetData(), Data.Num());
+		check(!FileReader->IsError());
+
+		FLargeMemoryReader MemoryReader(Data.GetData(), Data.Num());
+		return OutState.Load(MemoryReader, InOptions, OutVersion);
+	}
+
+	return false;
+}
+
 template<class Archive>
 void FAssetRegistryState::Load(Archive&& Ar, FAssetRegistryVersion::Type Version, const FAssetRegistryLoadOptions& Options)
 {
diff --git a/Engine/Source/Runtime/AssetRegistry/Public/AssetRegistry/AssetRegistryState.h b/Engine/Source/Runtime/AssetRegistry/Public/AssetRegistry/AssetRegistryState.h
--- a/Engine/Source/Runtime/AssetRegistry/Public/AssetRegistry/AssetRegistryState.h
+++ b/Engine/Source/Runtime/AssetRegistry/Public/AssetRegistry/AssetRegistryState.h
@@ -385,7 +385,14 @@ public:
 
 	/** Save without editor-only data */
 	bool Save(FArchive& Ar, const FAssetRegistrySerializationOptions& Options);
-	bool Load(FArchive& Ar, const FAssetRegistryLoadOptions& Options = FAssetRegistryLoadOptions());
+	bool Load(FArchive& Ar, const FAssetRegistryLoadOptions& Options = FAssetRegistryLoadOptions(), FAssetRegistryVersion::Type* OutVersion = nullptr);
+
+	/** 
+	* Example Usage:
+	*	FAssetRegistryState AssetRegistry;
+	*	bool bSucceeded = FAssetRegistryState::LoadFromDisk(TEXT("Path/To/AR"), FAssetRegistryLoadOptions(), AssetRegistry);
+	*/
+	static bool LoadFromDisk(const TCHAR* InPath, const FAssetRegistryLoadOptions& InOptions, FAssetRegistryState& OutState, FAssetRegistryVersion::Type* OutVersion = nullptr);
 
 	/** Returns memory size of entire registry, optionally logging sizes */
 	SIZE_T GetAllocatedSize(bool bLogDetailed = false) const;
