Provide a common way for determining a single "representative" asset for a package in the case where there are multiple assets in a package.

diff --git a/Engine/Source/Runtime/AssetRegistry/Private/AssetRegistry.cpp b/Engine/Source/Runtime/AssetRegistry/Private/AssetRegistry.cpp
--- a/Engine/Source/Runtime/AssetRegistry/Private/AssetRegistry.cpp
+++ b/Engine/Source/Runtime/AssetRegistry/Private/AssetRegistry.cpp
@@ -5875,6 +5875,80 @@ UAssetRegistryImpl::FFileLoadProgressUpdatedEvent& UAssetRegistryImpl::OnFileLoa
 
 namespace UE::AssetRegistry
 {
+const FAssetData* GetMostImportantAsset(TConstArrayView<const FAssetData*> PackageAssetDatas, bool bRequireOneTopLevelAsset)
+{
+	if (PackageAssetDatas.Num() == 1) // common case
+	{
+		return PackageAssetDatas[0];
+	}
+
+	// Find a candidate asset.
+	// If there's a "UAsset", then we use that as the asset.
+	// If not, then we look for a "TopLevelAsset", i.e. one that shows up in the content browser.
+	int32 TopLevelAssetCount = 0;
+
+	// If we have multiple TLAs, then we pick the "least" TLA.
+	// If we have NO TLAs, then we pick the "least" asset,
+	// both determined by class then name:
+	auto AssetDataLessThan = [](const FAssetData* LHS, const FAssetData* RHS)
+	{
+		int32 ClassCompare = LHS->AssetClass.Compare(RHS->AssetClass);
+		if (ClassCompare == 0)
+		{
+			return LHS->AssetName.LexicalLess(RHS->AssetName);
+		}
+		return ClassCompare < 0;
+	};
+
+	const FAssetData* LeastTopLevelAsset = nullptr;
+	const FAssetData* LeastAsset = nullptr;
+	for (const FAssetData* Asset : PackageAssetDatas)
+	{
+		if (Asset->AssetName.IsNone())
+		{
+			continue;
+		}
+		if (Asset->IsUAsset())
+		{
+			return Asset;
+		}
+		// This is after IsUAsset because Blueprints can be the UAsset but also be considered skipable.
+		if (FFiltering::ShouldSkipAsset(Asset->AssetClass, Asset->PackageFlags))
+		{
+			continue;
+		}
+		if (Asset->IsTopLevelAsset())
+		{
+			TopLevelAssetCount++;
+			if (LeastTopLevelAsset == nullptr ||
+				AssetDataLessThan(Asset, LeastTopLevelAsset))
+			{
+				LeastTopLevelAsset = Asset;
+			}
+		}
+		if (LeastAsset == nullptr ||
+			AssetDataLessThan(Asset, LeastAsset))
+		{
+			LeastAsset = Asset;
+		}
+	}
+
+	if (bRequireOneTopLevelAsset)
+	{
+		if (TopLevelAssetCount == 1)
+		{
+			return LeastTopLevelAsset;
+		}
+		return nullptr;
+	}
+
+	if (TopLevelAssetCount)
+	{
+		return LeastTopLevelAsset;
+	}
+	return LeastAsset;
+}
+
 
 void GetAssetForPackages(TConstArrayView<FName> PackageNames, TMap<FName, FAssetData>& OutPackageToAssetData)
 {
@@ -5883,7 +5957,7 @@ void GetAssetForPackages(TConstArrayView<FName> PackageNames, TMap<FName, FAsset
 	{
 		Filter.PackageNames.Add(PackageName);
 	}
-
+	
 	TArray<FAssetData> AssetDataList;
 	IAssetRegistry* AssetRegistry = IAssetRegistry::Get();
 	if (!AssetRegistry)
@@ -5891,37 +5965,30 @@ void GetAssetForPackages(TConstArrayView<FName> PackageNames, TMap<FName, FAsset
 		return;
 	}
 	AssetRegistry->GetAssets(Filter, AssetDataList);
-	auto IsHigherDisplayPriority = [](const FAssetData& LHS, const FAssetData& RHS)
+
+	if (AssetDataList.Num() == 0)
 	{
-		bool bLHSIsEmpty = LHS.AssetName.IsNone();
-		bool bRHSIsEmpty = RHS.AssetName.IsNone();
-		if (bLHSIsEmpty || bRHSIsEmpty)
-		{
-			return !bLHSIsEmpty;
-		}
-		bool bLHSIsUAsset = LHS.IsUAsset();
-		bool bRHSIsUAsset = RHS.IsUAsset();
-		if (bLHSIsUAsset != bRHSIsUAsset)
-		{
-			return bLHSIsUAsset;
-		}
-		bool bLHSShouldSkip = FFiltering::ShouldSkipAsset(LHS.AssetClass, LHS.PackageFlags);
-		bool bRHSShouldSkip = FFiltering::ShouldSkipAsset(LHS.AssetClass, LHS.PackageFlags);
-		if (bLHSShouldSkip || bRHSShouldSkip)
-		{
-			return !bLHSShouldSkip;
-		}
-		return LHS.AssetName.FastLess(RHS.AssetName);
-	};
-	for (FAssetData& AssetData : AssetDataList)
+		return;
+	}
+
+	Algo::SortBy(AssetDataList, &FAssetData::PackageName, FNameFastLess());
+
+	TArray<const FAssetData*, TInlineAllocator<1>> PackageAssetDatas;
+	FName CurrentPackageName = AssetDataList[0].PackageName;
+	for (const FAssetData& AssetData : AssetDataList)
 	{
-		// If there are multiple assets in the same package, use the highest priority asset for the display
-		FAssetData& MappedAssetData = OutPackageToAssetData.FindOrAdd(AssetData.PackageName);
-		if (IsHigherDisplayPriority(AssetData, MappedAssetData))
+		if (CurrentPackageName != AssetData.PackageName)
 		{
-			MappedAssetData = MoveTemp(AssetData);
+			OutPackageToAssetData.FindOrAdd(CurrentPackageName) = *GetMostImportantAsset(PackageAssetDatas, false);
+			PackageAssetDatas.Empty();
+			CurrentPackageName = AssetData.PackageName;
 		}
+
+		PackageAssetDatas.Push(&AssetData);
 	}
+
+	OutPackageToAssetData.FindOrAdd(CurrentPackageName) = *GetMostImportantAsset(PackageAssetDatas, false);
+
 }
 
-}
\ No newline at end of file
+} // namespace AssetRegistry
\ No newline at end of file
diff --git a/Engine/Source/Runtime/AssetRegistry/Public/AssetRegistry/IAssetRegistry.h b/Engine/Source/Runtime/AssetRegistry/Public/AssetRegistry/IAssetRegistry.h
--- a/Engine/Source/Runtime/AssetRegistry/Public/AssetRegistry/IAssetRegistry.h
+++ b/Engine/Source/Runtime/AssetRegistry/Public/AssetRegistry/IAssetRegistry.h
@@ -680,13 +680,31 @@ namespace AssetRegistry
 	ASSETREGISTRY_API bool ReadPackageDataDependencies(FArchive& BinaryArchive, TBitArray<>& OutImportUsedInGame, TBitArray<>& OutSoftPackageUsedInGame);
 
 	/**
-	 * Given a list of packages, gather the primary assets for each package.
+	 * Given a list of packages, gather the most important assets for each package.
 	 * If multiple assets are in a package, the most important asset will be added.
 	 * If a package does not exist or does not have any assets, no entry will be added for that package name.
 	 */
 	ASSETREGISTRY_API void GetAssetForPackages(TConstArrayView<FName> PackageNames, TMap<FName, FAssetData>& OutPackageToAssetData);
-}
-}
+
+	/**
+	 * Given a list of asset datas for a specific package, find an asset considered "most important" or "representative".
+	 * This is distinct from a Primary asset, and is used for user facing representation of a package or other cases
+	 * where you need to relate information about a package to an asset.
+	 * 
+	 * Usually there is only 1 asset per package so this is straightforward, however in the multiple asset case it:
+	 *	Tries to find the "UAsset" via the FAssetData::IsUAsset() function. (i.e. asset name matches package name)
+	 *	If none exist, tries to find a "Top Level Asset" using FAssetData::IsToplevelAsset(). (i.e. outer == package)
+	 *		If only one exists, use that.
+	 *		Otherwise, if bRequireOneTopLevelAsset is false, gather the set of possibles and return the first sorted on asset class then name.
+	 *			If no top level assets, all package assets
+	 *			If multiple top level assets, all top level assets
+	 * 
+	 * A good source for PackageAssetDatas is FAssetRegistryState::GetAssetsByPackageName.
+	 */
+	ASSETREGISTRY_API const FAssetData* GetMostImportantAsset(TConstArrayView<const FAssetData*> PackageAssetDatas, bool bRequireOneTopLevelAsset);
+
+} // namespace AssetRegistry
+} // namespace UE
 
 /** Returns the filename without filepath for the DevelopmentAssetRegistry written by the cooker. */
 ASSETREGISTRY_API const TCHAR* GetDevelopmentAssetRegistryFilename();
diff --git a/Engine/Source/Runtime/CoreUObject/Public/AssetRegistry/AssetData.h b/Engine/Source/Runtime/CoreUObject/Public/AssetRegistry/AssetData.h
--- a/Engine/Source/Runtime/CoreUObject/Public/AssetRegistry/AssetData.h
+++ b/Engine/Source/Runtime/CoreUObject/Public/AssetRegistry/AssetData.h
@@ -196,6 +196,7 @@ public:
 	 * Returns true iff the Asset is a TopLevelAsset (not a subobject, its outer is a UPackage).
 	 * Only TopLevelAssets can be PrimaryAssets in the AssetManager.
 	 * A TopLevelAsset is not necessarily the main asset in a package; see IsUAsset.
+	 * Note that this is distinct from UObject::IsAsset because IsAsset can be overloaded (see e.g. AActor::IsAsset)
 	 */
 	COREUOBJECT_API bool IsTopLevelAsset() const;
 	
