Reset intead of Empty() - edit that should have been in last change but got lost in the shelf.

diff --git a/Engine/Source/Runtime/AssetRegistry/Private/AssetRegistry.cpp b/Engine/Source/Runtime/AssetRegistry/Private/AssetRegistry.cpp
--- a/Engine/Source/Runtime/AssetRegistry/Private/AssetRegistry.cpp
+++ b/Engine/Source/Runtime/AssetRegistry/Private/AssetRegistry.cpp
@@ -5980,7 +5980,7 @@ void GetAssetForPackages(TConstArrayView<FName> PackageNames, TMap<FName, FAsset
 		if (CurrentPackageName != AssetData.PackageName)
 		{
 			OutPackageToAssetData.FindOrAdd(CurrentPackageName) = *GetMostImportantAsset(PackageAssetDatas, false);
-			PackageAssetDatas.Empty();
+			PackageAssetDatas.Reset();
 			CurrentPackageName = AssetData.PackageName;
 		}
 
