Optimize AssetDataGatherer cache load

* Blake3 -> xxhash3
* Enabled parallel name batch loading
* Add memory mapped file loading path to reduce memcpy
* Preload file in parallel to some CPU work

diff --git a/Engine/Source/Runtime/AssetRegistry/Private/AssetDataGatherer.cpp b/Engine/Source/Runtime/AssetRegistry/Private/AssetDataGatherer.cpp
--- a/Engine/Source/Runtime/AssetRegistry/Private/AssetDataGatherer.cpp
+++ b/Engine/Source/Runtime/AssetRegistry/Private/AssetDataGatherer.cpp
@@ -8,11 +8,13 @@
 #include "AssetRegistry/AssetData.h"
 #include "AssetRegistryArchive.h"
 #include "AssetRegistryPrivate.h"
+#include "Async/MappedFileHandle.h"
 #include "Async/ParallelFor.h"
 #include "Containers/BinaryHeap.h"
+#include "HAL/PlatformFileManager.h"
 #include "HAL/PlatformProcess.h"
 #include "HAL/RunnableThread.h"
-#include "Hash/Blake3.h"
+#include "Hash/xxhash.h"
 #include "Math/NumericLimits.h"
 #include "Memory/MemoryView.h"
 #include "Misc/AsciiSet.h"
@@ -26,6 +28,7 @@
 #include "Misc/ScopeExit.h"
 #include "PackageReader.h"
 #include "Serialization/Archive.h"
+#include "Tasks/Task.h"
 #include "Templates/UniquePtr.h"
 #include "Templates/UnrealTemplate.h"
 
@@ -34,36 +37,34 @@ namespace AssetDataGathererConstants
 	constexpr int32 SingleThreadFilesPerBatch = 3;
 	constexpr int32 ExpectedMaxBatchSize = 100;
 	constexpr int32 MinSecondsToElapseBeforeCacheWrite = 60;
-	static constexpr uint32 CacheSerializationMagic = 0x9D39D87B; // Versioning and integrity checking
+	static constexpr uint32 CacheSerializationMagic = 0x3339D87B; // Versioning and integrity checking
+	static constexpr uint64 CurrentVersion = FAssetRegistryVersion::LatestVersion | (uint64(CacheSerializationMagic) << 32);
 }
 
 namespace UE::AssetDataGather::Private
 {
-	/** A structure to hold serialized cache data from async loads before adding it to the Gatherer's main cache. */
-	struct FCachePayload
+
+/** A structure to hold serialized cache data from async loads before adding it to the Gatherer's main cache. */
+struct FCachePayload
+{
+	TUniquePtr<FName[]> PackageNames;
+	TUniquePtr<FDiskCachedAssetData[]> AssetDatas;
+	int32 NumAssets = 0;
+	bool bSucceeded = false;
+	void Reset()
 	{
-		TUniquePtr<FName[]> PackageNames;
-		TUniquePtr<FDiskCachedAssetData[]> AssetDatas;
-		int32 NumAssets = 0;
-		bool bSucceeded = false;
-		void Reset()
-		{
-			PackageNames.Reset();
-			AssetDatas.Reset();
-			NumAssets = 0;
-			bSucceeded = false;
-		}
-	};
+		PackageNames.Reset();
+		AssetDatas.Reset();
+		NumAssets = 0;
+		bSucceeded = false;
+	}
+};
 
-	void SerializeCacheSave(FAssetRegistryWriter& Ar, const TArray<TPair<FName, FDiskCachedAssetData*>>& AssetsToSave);
-	FCachePayload SerializeCacheLoad(FAssetRegistryReader& Ar);
-	FCachePayload LoadCacheFile(FStringView CacheFilename);
-}
+void SerializeCacheSave(FAssetRegistryWriter& Ar, const TArray<TPair<FName, FDiskCachedAssetData*>>& AssetsToSave);
+FCachePayload SerializeCacheLoad(FAssetRegistryReader& Ar);
+FCachePayload LoadCacheFile(FStringView CacheFilename);
 
 
-namespace UE::AssetDataGather::Private
-{
-
 /** InOutResult = Value, but without shrinking the string to fit. */
 void AssignStringWithoutShrinking(FString& InOutResult, FStringView Value)
 {
@@ -2343,24 +2344,88 @@ bool FAssetDataDiscovery::ShouldDirBeReported(FStringView LongPackageName) const
 	return !DirLongPackageNamesToNotReport.ContainsByHash(GetTypeHash(LongPackageName), LongPackageName);
 }
 
-
-class FChecksumArchiveBase : public FArchiveProxy
+// Reads an FMemoryView once
+class FMemoryViewReader
 {
-	// Enables both versioning and distinguishing out-of-sync reads from data corruption
-	static constexpr uint32 HeaderMagic = 0xb1a3;
-	static constexpr uint32 SaveBlockSize = 4 << 20;
+public:
+	FMemoryViewReader() = default;
+	FMemoryViewReader(FMemoryView Data) : Remaining(Data), TotalSize(Data.GetSize()) {}
 
-	static uint64 CalculateChecksum(FMemoryView Data)
+	uint64 GetRemainingSize() const { return Remaining.GetSize(); }
+	uint64 GetTotalSize() const { return TotalSize; }
+	uint64 Tell() const { return TotalSize - Remaining.GetSize(); }
+
+	FMemoryView Load(uint64 Size)
 	{
-		return INTEL_ORDER64(*reinterpret_cast<uint64*>(FBlake3::HashBuffer(Data).GetBytes()));
+		check(Size <= Remaining.GetSize());
+		FMemoryView Out(Remaining.GetData(), Size);
+		Remaining += Size;
+		return Out;
 	}
 
-	struct FBlockHeader
+	void Load(FMutableMemoryView Out)
 	{
-		uint32 Magic = 0;
-		uint32 Size = 0;
-		uint64 Checksum = 0;
-	};
+		FMemoryView In = Load(Out.GetSize());
+		FMemory::Memcpy(Out.GetData(), In.GetData(), In.GetSize());
+	}
+
+	template<typename T>
+	T Load()
+	{
+		static_assert(std::is_integral_v<T>, "Only integer loading supported");
+		static_assert(PLATFORM_LITTLE_ENDIAN, "Byte-swapping not implemented");
+		return FPlatformMemory::ReadUnaligned<T>(Load(sizeof(T)).GetData());
+	}
+
+	template<typename T>
+	TOptional<T> TryLoad()
+	{
+		return sizeof(T) <= Remaining.GetSize() ? Load<T>() : TOptional<T>();
+	}
+
+private:
+	FMemoryView Remaining;
+	uint64 TotalSize = 0;
+};
+
+
+// Enables both versioning and distinguishing out-of-sync reads from data corruption
+static constexpr uint32 BlockMagic = 0xb1a3;
+
+struct FBlockHeader
+{
+	uint32 Magic = 0;
+	uint32 Size = 0;
+	uint64 Checksum = 0;
+};
+
+TOptional<FBlockHeader> LoadBlockHeader(FMemoryView Data)
+{
+	check(Data.GetSize() == sizeof(FBlockHeader));
+
+	FMemoryViewReader Reader(Data);
+	FBlockHeader Header;
+	Header.Magic = Reader.Load<uint32>();
+	Header.Size = Reader.Load<uint32>();
+	Header.Checksum = Reader.Load<uint64>();
+
+	if (Header.Magic != BlockMagic)
+	{
+		UE_LOG(LogAssetRegistry, Error, TEXT("Wrong block magic (0x%x)"), Header.Magic);
+		return {};
+	}
+
+	return Header;
+}
+
+static uint64 CalculateBlockChecksum(FMemoryView Data)
+{
+	return INTEL_ORDER64(FXxHash64::HashBuffer(Data).Hash);
+}
+
+class FChecksumArchiveBase : public FArchiveProxy
+{
+	static constexpr uint32 SaveBlockSize = 4 << 20;
 
 	struct FBlock
 	{
@@ -2410,9 +2475,9 @@ class FChecksumArchiveBase : public FArchiveProxy
 	void SaveBlock()
 	{
 		FBlockHeader Header;
-		Header.Magic = HeaderMagic;
-		Header.Size = static_cast<uint32>(Block.GetUsedSize());
-		Header.Checksum = CalculateChecksum(Block.GetUsed());
+		Header.Magic = BlockMagic;
+		Header.Size = IntCastChecked<uint32>(Block.GetUsedSize());
+		Header.Checksum = CalculateBlockChecksum(Block.GetUsed());
 		InnerArchive << Header.Magic << Header.Size << Header.Checksum;
 
 		InnerArchive.Serialize(Block.Begin, Header.Size);
@@ -2424,36 +2489,35 @@ class FChecksumArchiveBase : public FArchiveProxy
 	{
 		check(Block.GetRemainingSize() == 0);
 
-		FBlockHeader Header;
-		InnerArchive << Header.Magic << Header.Size << Header.Checksum;
-
+		uint8 HeaderData[sizeof(FBlockHeader)];
+		InnerArchive.Serialize(HeaderData, sizeof(HeaderData));
 		if (InnerArchive.IsError())
 		{
 			UE_LOG(LogAssetRegistry, Error, TEXT("Couldn't read block header"));
 			return false;
 		}
-		else if (Header.Magic != HeaderMagic)
+
+		if (TOptional<FBlockHeader> Header = LoadBlockHeader(MakeMemoryView(HeaderData)))
 		{
-			UE_LOG(LogAssetRegistry, Error, TEXT("Wrong block magic (0x%x)"), Header.Magic);
-			return false;
-		}
+			Block.Reset(Header->Size);
 
-		Block.Reset(Header.Size);
-			
-		InnerArchive.Serialize(Block.Begin, Header.Size);
+			InnerArchive.Serialize(Block.Begin, Header->Size);
 
-		if (InnerArchive.IsError())
-		{
-			UE_LOG(LogAssetRegistry, Error, TEXT("Couldn't read block data"));
-			return false;
-		}
-		else if (CalculateChecksum(Block.GetRemaining()) != Header.Checksum)
-		{
-			UE_LOG(LogAssetRegistry, Error, TEXT("Wrong block checksum"));
-			return false;
+			if (InnerArchive.IsError())
+			{
+				UE_LOG(LogAssetRegistry, Error, TEXT("Couldn't read block data"));
+				return false;
+			}
+			else if (CalculateBlockChecksum(Block.GetRemaining()) != Header->Checksum)
+			{
+				UE_LOG(LogAssetRegistry, Error, TEXT("Wrong block checksum"));
+				return false;
+			}
+
+			return true;
 		}
 
-		return true;
+		return false;
 	}
 
 public:
@@ -2549,6 +2613,111 @@ public:
 	virtual int64 Tell() override { return InnerArchive.Tell() - GetCurrentBlock().GetRemainingSize(); }
 };
 
+// Memory-mapped equivalent of FChecksumArchiveReader
+class FChecksumViewReader : public FArchive
+{
+public:
+	explicit FChecksumViewReader(FMemoryViewReader&& Reader, FStringView InFileName)
+	: RemainingBlocks(MoveTemp(Reader))
+	, FileName(InFileName)
+	{
+		SetIsLoading(true);
+	}
+
+private:
+	FMemoryViewReader RemainingBlocks;
+	FMemoryViewReader CurrentBlock;
+	FString FileName;
+
+	virtual void Seek(int64) override { unimplemented(); }
+	virtual int64 Tell() override { return RemainingBlocks.Tell() - CurrentBlock.GetRemainingSize(); }
+	virtual int64 TotalSize() override { return RemainingBlocks.GetTotalSize(); }
+
+	virtual void Serialize(void* V, int64 Len) override
+	{
+		FMutableMemoryView Out(V, Len);
+
+		while (CurrentBlock.GetRemainingSize() < Out.GetSize())
+		{
+			if (IsError())
+			{
+				return;
+			}
+
+			FMutableMemoryView OutSlice(Out.GetData(), CurrentBlock.GetRemainingSize());
+			Out += OutSlice.GetSize();
+			CurrentBlock.Load(OutSlice);
+			check(CurrentBlock.GetRemainingSize() == 0);
+
+			TOptional<FMemoryView> NextBlock = LoadNextBlock(RemainingBlocks);
+			if (!NextBlock)
+			{
+				UE_LOG(LogAssetRegistry, Error, TEXT("Integrity check failed, '%s' cache will be discarded"), *FileName);
+				SetError();
+				return;
+			}
+			CurrentBlock = *NextBlock;
+		}
+
+		CurrentBlock.Load(Out);
+	}
+
+	FORCENOINLINE static TOptional<FMemoryView> LoadNextBlock(FMemoryViewReader& In)
+	{
+		if (In.GetRemainingSize() < sizeof(FBlockHeader))
+		{
+			UE_LOG(LogAssetRegistry, Error, TEXT("Couldn't read block header"));
+			return {};
+		}
+
+		if (TOptional<FBlockHeader> Header = LoadBlockHeader(In.Load(sizeof(FBlockHeader))))
+		{
+			if (Header->Size > In.GetRemainingSize())
+			{
+				UE_LOG(LogAssetRegistry, Error, TEXT("Incomplete block"));
+				return {};
+			}
+
+			FMemoryView Block = In.Load(Header->Size);
+			if (CalculateBlockChecksum(Block) != Header->Checksum)
+			{
+				UE_LOG(LogAssetRegistry, Error, TEXT("Wrong block checksum"));
+				return {};
+			}
+
+			return Block;
+		}
+
+		return {};
+	}
+};
+
+// Util that maps an entire file
+class FMemoryMappedFile
+{
+public:
+	FMemoryMappedFile(const TCHAR* Path)
+		: Handle(FPlatformFileManager::Get().GetPlatformFile().OpenMapped(Path))
+		, Region(Handle ? Handle->MapRegion() : nullptr)
+	{}
+
+	void Preload(int64 Size = MAX_int64) const
+	{
+		if (Region)
+		{
+			Region->PreloadHint(0, Size);
+		}
+	}
+	
+	FMemoryView View() const
+	{
+		return Region ? FMemoryView(Region->GetMappedPtr(), Region->GetMappedSize()) : FMemoryView();
+	}
+
+private:
+	TUniquePtr<IMappedFileHandle> Handle;
+	TUniquePtr<IMappedFileRegion> Region;
+};
 
 /**
  * Settings about whether to use cache data for the AssetDataGatherer; these settings are shared by
@@ -3578,11 +3747,9 @@ void FAssetDataGatherer::SaveCacheFileInternal(const FString& CacheFilename, con
 	TUniquePtr<FArchive> FileAr(IFileManager::Get().CreateFileWriter(*TempFilename, 0));
 	if (FileAr)
 	{
-		int32 MagicNumber = AssetDataGathererConstants::CacheSerializationMagic;
-		*FileAr << MagicNumber;
+		uint64 CurrentVersion = AssetDataGathererConstants::CurrentVersion;
+		*FileAr << CurrentVersion;
 
-		FAssetRegistryVersion::Type RegistryVersion = FAssetRegistryVersion::LatestVersion;
-		FAssetRegistryVersion::SerializeVersion(*FileAr, RegistryVersion);
 #if ALLOW_NAME_BATCH_SAVING
 		{
 			// We might be able to reduce load time by using AssetRegistry::SerializationOptions
@@ -3692,42 +3859,52 @@ FCachePayload SerializeCacheLoad(FAssetRegistryReader& Ar)
 	return Result;
 }
 
-FCachePayload LoadCacheFile(FStringView CacheFilename)
+FCachePayload LoadCacheFile(FStringView InCacheFilename)
 {
 	TRACE_CPUPROFILER_EVENT_SCOPE(LoadCacheFile);
+	FString CacheFilename(InCacheFilename);
 
-	TUniquePtr<FArchive> FileAr(IFileManager::Get().CreateFileReader(*WriteToString<256>(CacheFilename), FILEREAD_Silent));
-	if (!FileAr || FileAr->IsError() || FileAr->TotalSize() < 2 * sizeof(uint32))
+	auto DoLoad = [&](FArchive& ChecksummingReader)
 	{
-		return FCachePayload();
-	}
+		int32 WorkerReduction = 2; // Current worker + preload task
+		int32 Parallelism = FMath::Max(FTaskGraphInterface::Get().GetNumWorkerThreads() - WorkerReduction, 0);
+		FAssetRegistryReader RegistryReader(ChecksummingReader, Parallelism);
+		return RegistryReader.IsError() ? FCachePayload() : SerializeCacheLoad(RegistryReader);
+	};
 
-	uint32 MagicNumber = 0;
-	*FileAr << MagicNumber;
-	if (FileAr->IsError() || MagicNumber != AssetDataGathererConstants::CacheSerializationMagic)
+	FCachePayload Payload;
+	if (FPlatformProperties::SupportsMemoryMappedFiles())
 	{
-		return FCachePayload();
-	}
+		FMemoryMappedFile File(*CacheFilename);
+		UE::Tasks::FTask Preload = UE::Tasks::Launch(UE_SOURCE_LOCATION, [&]() { File.Preload(); });
 
-	FAssetRegistryVersion::Type RegistryVersion;
-	if (!FAssetRegistryVersion::SerializeVersion(*FileAr, RegistryVersion) || RegistryVersion != FAssetRegistryVersion::LatestVersion)
-	{
-		return FCachePayload();
-	}
+		FMemoryViewReader FileReader(File.View());
+		TOptional<uint64> Version = FileReader.TryLoad<uint64>();
+		if (Version == AssetDataGathererConstants::CurrentVersion)
+		{
+			FChecksumViewReader ChecksummingReader(MoveTemp(FileReader), CacheFilename);
+			Payload = DoLoad(ChecksummingReader);
+		}
 
-	UE::AssetDataGather::Private::FChecksumArchiveReader ChecksummingReader(*FileAr);
-	FAssetRegistryReader RegistryReader(ChecksummingReader);
-	if (RegistryReader.IsError())
-	{
-		return FCachePayload();
+		Preload.Wait();
+		UE::Tasks::Launch(UE_SOURCE_LOCATION, [KillAsync = MoveTemp(File)](){});
 	}
-
-	FCachePayload Payload;
-	Payload = UE::AssetDataGather::Private::SerializeCacheLoad(RegistryReader);
-	if (!Payload.bSucceeded)
+	else
 	{
-		UE_LOG(LogAssetRegistry, Error, TEXT("There was an error loading the asset registry cache."));
+		TUniquePtr<FArchive> FileAr(IFileManager::Get().CreateFileReader(*CacheFilename, FILEREAD_Silent));
+		if (FileAr && !FileAr->IsError() && FileAr->TotalSize() > sizeof(uint64))
+		{
+			uint64 Version = 0;
+			*FileAr << Version;
+			if (Version == AssetDataGathererConstants::CurrentVersion)
+			{
+				FChecksumArchiveReader ChecksummingReader(*FileAr);
+				Payload = DoLoad(ChecksummingReader);
+			}
+		}
 	}
+	
+	UE_CLOG(!Payload.bSucceeded, LogAssetRegistry, Error, TEXT("There was an error loading the asset registry cache."));
 	return Payload;
 }
 
