Fix false error when asset registry loading detects a version change

#rnx

diff --git a/Engine/Source/Runtime/AssetRegistry/Private/AssetDataGatherer.cpp b/Engine/Source/Runtime/AssetRegistry/Private/AssetDataGatherer.cpp
--- a/Engine/Source/Runtime/AssetRegistry/Private/AssetDataGatherer.cpp
+++ b/Engine/Source/Runtime/AssetRegistry/Private/AssetDataGatherer.cpp
@@ -3884,6 +3884,7 @@ FCachePayload LoadCacheFile(FStringView InCacheFilename)
 		{
 			FChecksumViewReader ChecksummingReader(MoveTemp(FileReader), CacheFilename);
 			Payload = DoLoad(ChecksummingReader);
+			UE_CLOG(!Payload.bSucceeded, LogAssetRegistry, Error, TEXT("There was an error loading the asset registry cache using memory mapping"));
 		}
 
 		Preload.Wait();
@@ -3900,11 +3901,11 @@ FCachePayload LoadCacheFile(FStringView InCacheFilename)
 			{
 				FChecksumArchiveReader ChecksummingReader(*FileAr);
 				Payload = DoLoad(ChecksummingReader);
+				UE_CLOG(!Payload.bSucceeded, LogAssetRegistry, Error, TEXT("There was an error loading the asset registry cache"));
 			}
 		}
 	}
 	
-	UE_CLOG(!Payload.bSucceeded, LogAssetRegistry, Error, TEXT("There was an error loading the asset registry cache."));
 	return Payload;
 }
 
