AssetRegistry fix for CrashReporterClient: disable the premade preloaded assetregistry in programs that require cookedcontent. It is only available in cooked game or in cooked editor.
#rnx

diff --git a/Engine/Source/Runtime/AssetRegistry/Private/AssetRegistry.cpp b/Engine/Source/Runtime/AssetRegistry/Private/AssetRegistry.cpp
--- a/Engine/Source/Runtime/AssetRegistry/Private/AssetRegistry.cpp
+++ b/Engine/Source/Runtime/AssetRegistry/Private/AssetRegistry.cpp
@@ -133,7 +133,8 @@ namespace UE::AssetRegistry::Premade
 /** Returns whether the given executable configuration supports AssetRegistry Preloading. Called before Main. */
 static bool IsEnabled()
 {
-	return FPlatformProperties::RequiresCookedData() || ASSETREGISTRY_ENABLE_PREMADE_REGISTRY_IN_EDITOR;
+	return (FPlatformProperties::RequiresCookedData() && IsRunningGame())
+		|| ASSETREGISTRY_ENABLE_PREMADE_REGISTRY_IN_EDITOR;
 }
 
 static bool CanLoadAsync()
@@ -666,7 +667,8 @@ void FAssetRegistryImpl::LoadPremadeAssetRegistry(Impl::FEventContext& EventCont
 		}
 		else
 		{
-			UE_CLOG(FPlatformProperties::RequiresCookedData(), LogAssetRegistry, Fatal, TEXT("Failed to load premade asset registry."));
+			UE_CLOG(FPlatformProperties::RequiresCookedData() && IsRunningGame(),
+				LogAssetRegistry, Fatal, TEXT("Failed to load premade asset registry."));
 		}
 	}
 
