Fix for Cooked Dedicated Server AssetRegistry Enabled Issue.

diff --git a/Engine/Source/Runtime/AssetRegistry/Private/AssetRegistry.cpp b/Engine/Source/Runtime/AssetRegistry/Private/AssetRegistry.cpp
--- a/Engine/Source/Runtime/AssetRegistry/Private/AssetRegistry.cpp
+++ b/Engine/Source/Runtime/AssetRegistry/Private/AssetRegistry.cpp
@@ -133,7 +133,7 @@ namespace UE::AssetRegistry::Premade
 /** Returns whether the given executable configuration supports AssetRegistry Preloading. Called before Main. */
 static bool IsEnabled()
 {
-	return (FPlatformProperties::RequiresCookedData() && IsRunningGame())
+	return (FPlatformProperties::RequiresCookedData() && (IsRunningGame() || IsRunningDedicatedServer()))
 		|| ASSETREGISTRY_ENABLE_PREMADE_REGISTRY_IN_EDITOR;
 }
 
@@ -667,7 +667,7 @@ void FAssetRegistryImpl::LoadPremadeAssetRegistry(Impl::FEventContext& EventCont
 		}
 		else
 		{
-			UE_CLOG(FPlatformProperties::RequiresCookedData() && IsRunningGame(),
+			UE_CLOG(FPlatformProperties::RequiresCookedData() && (IsRunningGame() || IsRunningDedicatedServer()),
 				LogAssetRegistry, Fatal, TEXT("Failed to load premade asset registry."));
 		}
 	}
