AssetRegistry: Deprecate the include redirectors for "AssetData.h" -> "AssetRegistry/AssetData.h" and other AssetRegistry public includes.
#rnx

diff --git a/Engine/Source/Runtime/AssetRegistry/Public/ARFilter.h b/Engine/Source/Runtime/AssetRegistry/Public/ARFilter.h
--- a/Engine/Source/Runtime/AssetRegistry/Public/ARFilter.h
+++ b/Engine/Source/Runtime/AssetRegistry/Public/ARFilter.h
@@ -2,4 +2,7 @@
 
 #pragma once
 
+// Deprecated in UE 5.1.
+#pragma message (__FILE__"(6): warning: #include AssetRegistry/ARFilter.h instead of ARFilter.h")
+
 #include "AssetRegistry/ARFilter.h"
diff --git a/Engine/Source/Runtime/AssetRegistry/Public/AssetBundleData.h b/Engine/Source/Runtime/AssetRegistry/Public/AssetBundleData.h
--- a/Engine/Source/Runtime/AssetRegistry/Public/AssetBundleData.h
+++ b/Engine/Source/Runtime/AssetRegistry/Public/AssetBundleData.h
@@ -2,4 +2,7 @@
 
 #pragma once
 
+// Deprecated in UE 5.1.
+#pragma message (__FILE__"(6): warning: #include AssetRegistry/AssetBundleData.h instead of AssetBundleData.h")
+
 #include "AssetRegistry/AssetBundleData.h"
diff --git a/Engine/Source/Runtime/AssetRegistry/Public/AssetData.h b/Engine/Source/Runtime/AssetRegistry/Public/AssetData.h
--- a/Engine/Source/Runtime/AssetRegistry/Public/AssetData.h
+++ b/Engine/Source/Runtime/AssetRegistry/Public/AssetData.h
@@ -2,4 +2,7 @@
 
 #pragma once
 
+// Deprecated in UE 5.1.
+#pragma message (__FILE__"(6): warning: #include AssetRegistry/AssetData.h instead of AssetData.h")
+
 #include "AssetRegistry/AssetData.h"
diff --git a/Engine/Source/Runtime/AssetRegistry/Public/AssetDataTagMap.h b/Engine/Source/Runtime/AssetRegistry/Public/AssetDataTagMap.h
--- a/Engine/Source/Runtime/AssetRegistry/Public/AssetDataTagMap.h
+++ b/Engine/Source/Runtime/AssetRegistry/Public/AssetDataTagMap.h
@@ -2,4 +2,7 @@
 
 #pragma once
 
+// Deprecated in UE 5.1.
+#pragma message (__FILE__"(6): warning: #include AssetRegistry/AssetDataTagMap.h instead of AssetDataTagMap.h")
+
 #include "AssetRegistry/AssetDataTagMap.h"
diff --git a/Engine/Source/Runtime/AssetRegistry/Public/AssetRegistryHelpers.h b/Engine/Source/Runtime/AssetRegistry/Public/AssetRegistryHelpers.h
--- a/Engine/Source/Runtime/AssetRegistry/Public/AssetRegistryHelpers.h
+++ b/Engine/Source/Runtime/AssetRegistry/Public/AssetRegistryHelpers.h
@@ -2,4 +2,7 @@
 
 #pragma once
 
+// Deprecated in UE 5.1.
+#pragma message (__FILE__"(6): warning: #include AssetRegistry/AssetRegistryHelpers.h instead of AssetRegistryHelpers.h")
+
 #include "AssetRegistry/AssetRegistryHelpers.h"
diff --git a/Engine/Source/Runtime/AssetRegistry/Public/AssetRegistryModule.h b/Engine/Source/Runtime/AssetRegistry/Public/AssetRegistryModule.h
--- a/Engine/Source/Runtime/AssetRegistry/Public/AssetRegistryModule.h
+++ b/Engine/Source/Runtime/AssetRegistry/Public/AssetRegistryModule.h
@@ -2,4 +2,7 @@
 
 #pragma once
 
+// Deprecated in UE 5.1.
+#pragma message (__FILE__"(6): warning: #include AssetRegistry/AssetRegistryModule.h instead of AssetRegistryModule.h")
+
 #include "AssetRegistry/AssetRegistryModule.h"
diff --git a/Engine/Source/Runtime/AssetRegistry/Public/AssetRegistryState.h b/Engine/Source/Runtime/AssetRegistry/Public/AssetRegistryState.h
--- a/Engine/Source/Runtime/AssetRegistry/Public/AssetRegistryState.h
+++ b/Engine/Source/Runtime/AssetRegistry/Public/AssetRegistryState.h
@@ -2,4 +2,7 @@
 
 #pragma once
 
+// Deprecated in UE 5.1.
+#pragma message (__FILE__"(6): warning: #include AssetRegistry/AssetRegistryState.h instead of AssetRegistryState.h")
+
 #include "AssetRegistry/AssetRegistryState.h"
diff --git a/Engine/Source/Runtime/AssetRegistry/Public/IAssetRegistry.h b/Engine/Source/Runtime/AssetRegistry/Public/IAssetRegistry.h
--- a/Engine/Source/Runtime/AssetRegistry/Public/IAssetRegistry.h
+++ b/Engine/Source/Runtime/AssetRegistry/Public/IAssetRegistry.h
@@ -2,5 +2,8 @@
 
 #pragma once
 
+// Deprecated in UE 5.1.
+#pragma message (__FILE__"(6): warning: #include AssetRegistry/IAssetRegistry.h instead of IAssetRegistry.h")
+
 #include "AssetRegistry/IAssetRegistry.h"
 
