- AssetRegistry support for an added LoadCalculatedDependencies phase to the AssetDataGather to allow registered IAssetDependencyGatherers to add dependencies
- Implement a World Partition IAssetDependencyGatherer to add dependencies from the World to all its ExternalActor packages.

diff --git a/Engine/Source/Runtime/AssetRegistry/Private/AssetDependencyGatherer.cpp b/Engine/Source/Runtime/AssetRegistry/Private/AssetDependencyGatherer.cpp
new file mode 100644
--- /dev/null
+++ b/Engine/Source/Runtime/AssetRegistry/Private/AssetDependencyGatherer.cpp
@@ -0,0 +1,38 @@
+// Copyright Epic Games, Inc. All Rights Reserved.
+
+#include "AssetRegistry/AssetDependencyGatherer.h"
+
+#if WITH_EDITOR
+
+namespace UE::AssetDependencyGatherer::Private
+{
+
+static TLinkedList<FRegisteredAssetDependencyGatherer*>* GRegisteredAssetDependencyGathererList = nullptr;
+
+FRegisteredAssetDependencyGatherer::FRegisteredAssetDependencyGatherer()
+	: GlobalListLink(this)
+{
+	GlobalListLink.LinkHead(GetRegisteredList());
+}
+
+FRegisteredAssetDependencyGatherer::~FRegisteredAssetDependencyGatherer()
+{
+	GlobalListLink.Unlink();
+}
+
+TLinkedList<FRegisteredAssetDependencyGatherer*>*& FRegisteredAssetDependencyGatherer::GetRegisteredList()
+{
+	return GRegisteredAssetDependencyGathererList;
+}
+
+void FRegisteredAssetDependencyGatherer::ForEach(TFunctionRef<void(FRegisteredAssetDependencyGatherer*)> Func)
+{
+	for (TLinkedList<FRegisteredAssetDependencyGatherer*>::TIterator It(GetRegisteredList()); It; It.Next())
+	{
+		Func(*It);
+	}
+}
+	
+}
+
+#endif
\ No newline at end of file
diff --git a/Engine/Source/Runtime/AssetRegistry/Private/AssetRegistry.cpp b/Engine/Source/Runtime/AssetRegistry/Private/AssetRegistry.cpp
--- a/Engine/Source/Runtime/AssetRegistry/Private/AssetRegistry.cpp
+++ b/Engine/Source/Runtime/AssetRegistry/Private/AssetRegistry.cpp
@@ -5,6 +5,7 @@
 #include "Algo/Unique.h"
 #include "AssetDataGatherer.h"
 #include "AssetRegistry/ARFilter.h"
+#include "AssetRegistry/AssetDependencyGatherer.h"
 #include "AssetRegistryConsoleCommands.h"
 #include "AssetRegistryPrivate.h"
 #include "Async/Async.h"
@@ -778,6 +779,27 @@ void FAssetRegistryImpl::Initialize(Impl::FInitializeContext& Context)
 	}
 
 	InitRedirectors(Context.Events, Context.bRedirectorsNeedSubscribe);
+
+#if WITH_EDITOR
+	if (!UE::AssetDependencyGatherer::Private::FRegisteredAssetDependencyGatherer::IsEmpty())
+	{
+		TArray<UObject*> Classes;
+		GetObjectsOfClass(UClass::StaticClass(), Classes);
+
+		/** Per Class dependency gatherers */
+		UE::AssetDependencyGatherer::Private::FRegisteredAssetDependencyGatherer::ForEach([&](UE::AssetDependencyGatherer::Private::FRegisteredAssetDependencyGatherer* RegisteredAssetDependencyGatherer)
+		{
+			UClass* AssetClass = RegisteredAssetDependencyGatherer->GetAssetClass();
+			for(UObject* ClassObject : Classes)
+			{
+				if (UClass* Class = Cast<UClass>(ClassObject); Class && Class->IsChildOf(AssetClass) && !Class->HasAnyClassFlags(CLASS_Abstract | CLASS_Deprecated | CLASS_NewerVersionExists))
+				{
+					RegisteredDependencyGathererClasses.Add(Class, RegisteredAssetDependencyGatherer);
+				}
+			}
+		});
+	}
+#endif
 }
 
 }
@@ -3391,11 +3413,20 @@ Impl::EGatherStatus FAssetRegistryImpl::TickGatherer(Impl::FEventContext& EventC
 		UE_LOG(LogAssetRegistry, Verbose, TEXT("### Background search completed in %0.4f seconds"), SearchTimes[SearchTimeIdx]);
 	}
 	const bool bHadAssetsToProcess = BackgroundAssetResults.Num() > 0 || BackgroundDependencyResults.Num() > 0;
+	auto GetNumGatherFromDiskPending = [&NumFilesToSearch, &NumPathsToSearch, this]()
+	{
+		return NumFilesToSearch + NumPathsToSearch + BackgroundPathResults.Num() + BackgroundAssetResults.Num()
+			+ BackgroundDependencyResults.Num() + BackgroundCookedPackageNamesWithoutAssetDataResults.Num();
+	};
 	auto UpdateStatus = [bIsSearching, bAbleToProgress, bHadAssetsToProcess, bIsDiscoveringFiles, &EventContext,
-		&OutStatus, &NumFilesToSearch, &NumPathsToSearch, this](bool bInterrupted)
+		&OutStatus, this](int32 NumGatherPending, bool bInterrupted)
 	{
 		// Compute total pending, plus highest pending for this run so we can show a good progress bar
-		int32 NumPending = NumFilesToSearch + NumPathsToSearch + BackgroundPathResults.Num() + BackgroundAssetResults.Num() + BackgroundDependencyResults.Num() + BackgroundCookedPackageNamesWithoutAssetDataResults.Num();
+		int32 NumPending = NumGatherPending
+#if WITH_EDITOR
+			+ (PackagesNeedingDependencyCalculation.Num() ? 1 : 0)
+#endif
+			;
 		HighestPending = FMath::Max(this->HighestPending, NumPending);
 
 		if (!bInterrupted && !bIsSearching && NumPending == 0)
@@ -3465,13 +3496,28 @@ Impl::EGatherStatus FAssetRegistryImpl::TickGatherer(Impl::FEventContext& EventC
 		CookedPackageNamesWithoutAssetDataGathered(EventContext, TickStartTime, BackgroundCookedPackageNamesWithoutAssetDataResults, bOutInterrupted);
 		if (bOutInterrupted)
 		{
-			UpdateStatus(true /* bInterrupted */);
+			UpdateStatus(GetNumGatherFromDiskPending(), true /* bInterrupted */);
 			return OutStatus;
 		}
 	}
 
+	// Load Calculated Dependencies when the gather from disk is complete; the full gather is not complete until after this is done
+	int32 NumGatherFromDiskPending = GetNumGatherFromDiskPending();
+#if WITH_EDITOR
+	bool bDiskGatherComplete = !bIsSearching && NumGatherFromDiskPending == 0;
+	if (bDiskGatherComplete && PackagesNeedingDependencyCalculation.Num())
+	{
+		LoadCalculatedDependencies(nullptr, TickStartTime, bOutInterrupted);
+		if (bOutInterrupted)
+		{
+			UpdateStatus(NumGatherFromDiskPending, true /* bInterrupted */);
+			return OutStatus;
+		}
+	}
+#endif
+
 	// If completing an initial search, refresh the content browser
-	UpdateStatus(false /* bInterrupted */);
+	UpdateStatus(NumGatherFromDiskPending, false /* bInterrupted */);
 
 	if (OutStatus == EGatherStatus::Complete)
 	{
@@ -3540,6 +3586,139 @@ void FAssetRegistryImpl::TickGatherPackage(Impl::FEventContext& EventContext, co
 	}
 }
 
+#if WITH_EDITOR
+void FAssetRegistryImpl::LoadCalculatedDependencies(TArray<FName>* AssetsToCalculate, double TickStartTime, bool& bOutInterrupted)
+{
+	auto CheckForTimeUp = [&TickStartTime, &bOutInterrupted](bool bHadActivity)
+	{
+		// Only Check TimeUp when we found something to do, otherwise we waste time calling FPlatformTime::Seconds
+		if (!bHadActivity)
+		{
+			return false;
+		}
+		if (TickStartTime >= 0 && (FPlatformTime::Seconds() - TickStartTime) >= UE::AssetRegistry::Impl::MaxSecondsPerFrame)
+		{
+			bOutInterrupted = true;
+			return true;
+		}
+		return false;
+	};
+
+	if (AssetsToCalculate)
+	{
+		for (FName PackageName : *AssetsToCalculate)
+		{
+			// We do not remove the package from PackagesNeedingDependencyCalculation, because
+			// we are only calculating an interim result when AssetsToCalculate is non-null
+			// We will run again on each of these PackageNames when TickGatherer finishes gathering all dependencies
+			if (PackagesNeedingDependencyCalculation.Contains(PackageName))
+			{
+				bool bHadActivity;
+				LoadCalculatedDependencies(PackageName, bHadActivity);
+				if (CheckForTimeUp(bHadActivity))
+				{
+					return;
+				}
+			}
+		}
+	}
+	else
+	{
+		for (FName PackageName : PackagesNeedingDependencyCalculation)
+		{
+			bool bHadActivity;
+			LoadCalculatedDependencies(PackageName, bHadActivity);
+			if (CheckForTimeUp(bHadActivity))
+			{
+				return;
+			}
+		}
+		PackagesNeedingDependencyCalculation.Empty();
+	}
+}
+
+void FAssetRegistryImpl::LoadCalculatedDependencies(FName PackageName, bool& bOutHadActivity)
+{
+	bOutHadActivity = false;
+	
+	auto GetCompiledFilter = [this](const FARFilter& InFilter) -> FARCompiledFilter
+	{
+		check(!InFilter.bRecursiveClasses);
+		UE::AssetRegistry::Impl::FClassInheritanceContext InheritanceContext;
+		FARCompiledFilter CompiledFilter;
+		CompileFilter(InheritanceContext, InFilter, CompiledFilter);
+		return CompiledFilter;
+	};
+
+	for (const FAssetData* AssetData : State.GetAssetsByPackageName(PackageName))
+	{
+		if (UE::AssetDependencyGatherer::Private::FRegisteredAssetDependencyGatherer** AssetDependencyGatherer = RegisteredDependencyGathererClasses.Find(AssetData->GetClass()))
+		{
+			if (!bOutHadActivity)
+			{
+				RemoveDirectoryReferencer(PackageName);
+			}
+
+			bOutHadActivity = true;
+			TArray<IAssetDependencyGatherer::FGathereredDependency> GatheredDependencies;
+			TArray<FString> DirectoryReferences;
+
+			(*AssetDependencyGatherer)->GatherDependencies(*AssetData, State, GetCompiledFilter, GatheredDependencies, DirectoryReferences);
+			
+			if (GatheredDependencies.Num())
+			{
+				FDependsNode* SourceNode = State.CreateOrFindDependsNode(FAssetIdentifier(PackageName));
+
+				for (const IAssetDependencyGatherer::FGathereredDependency& GatheredDepencency : GatheredDependencies)
+				{
+					FDependsNode* TargetNode = State.CreateOrFindDependsNode(FAssetIdentifier(GatheredDepencency.PackageName));
+					EDependencyProperty DependencyProperties = GatheredDepencency.Property;
+					SourceNode->AddDependency(TargetNode, EDependencyCategory::Package, DependencyProperties);
+					TargetNode->AddReferencer(SourceNode);
+				}
+			}
+
+			for (const FString& Directory : DirectoryReferences)
+			{
+				AddDirectoryReferencer(PackageName, Directory);
+			}
+		}
+	}
+}
+
+void FAssetRegistryImpl::AddDirectoryReferencer(FName PackageName, const FString& DirectoryLocalPathOrLongPackageName)
+{
+	FString DirectoryLocalPath;
+	if (!FPackageName::TryConvertLongPackageNameToFilename(DirectoryLocalPathOrLongPackageName, DirectoryLocalPath))
+	{
+		UE_LOG(LogAssetRegistry, Warning,
+			TEXT("AddDirectoryReferencer called with LongPackageName %s that cannot be mapped to a local path. Ignoring it."),
+			*DirectoryLocalPathOrLongPackageName);
+		return;
+	}
+	FPaths::MakeStandardFilename(DirectoryLocalPath);
+	DirectoryReferencers.AddUnique(DirectoryLocalPath, PackageName);
+}
+
+void FAssetRegistryImpl::RemoveDirectoryReferencer(FName PackageName)
+{
+	TArray<FString> FoundKeys;
+	for (const TPair<FString, FName>& Pair : DirectoryReferencers)
+	{
+		if (Pair.Value == PackageName)
+		{
+			FoundKeys.Add(Pair.Key);
+		}
+	}
+	for (FString& Key : FoundKeys)
+	{
+		DirectoryReferencers.Remove(Key, PackageName);
+	}
+}
+
+#endif
+
+
 }
 
 void UAssetRegistryImpl::Serialize(FArchive& Ar)
@@ -3916,18 +4095,19 @@ void FAssetRegistryImpl::ScanPathsSynchronous(Impl::FScanPathContext& Context)
 	}
 
 	Gatherer.ScanPathsSynchronous(Context.LocalPaths, Context.bForceRescan, Context.bIgnoreDenyListScanFilters, CacheFilename, Context.PackageDirs);
+	TArray<FName> BufferFoundAssets;
+	TArray<FName>* FoundAssets = &BufferFoundAssets;
+	if (Context.OutFoundAssets)
+	{
+		FoundAssets = Context.OutFoundAssets;
+	}
 
-	auto AssetsFoundCallback = [&Context](const TRingBuffer<FAssetData*>& InFoundAssets)
+	auto AssetsFoundCallback = [&Context, FoundAssets, this](const TRingBuffer<FAssetData*>& InFoundAssets)
 	{
 		Context.NumFoundAssets = InFoundAssets.Num();
 
-		if (!Context.OutFoundAssets)
-		{
-			return;
-		}
-
-		Context.OutFoundAssets->Reset();
-		Context.OutFoundAssets->Reserve(Context.NumFoundAssets);
+		FoundAssets->Reset();
+		FoundAssets->Reserve(Context.NumFoundAssets);
 
 		// The gatherer may have added other assets that were scanned as part of the ongoing background scan; remove any assets that were not in the requested paths
 		for (FAssetData* AssetData : InFoundAssets)
@@ -3963,13 +4143,16 @@ void FAssetRegistryImpl::ScanPathsSynchronous(Impl::FScanPathContext& Context)
 			{
 				UE_LOG(LogAssetRegistry, VeryVerbose, TEXT("FAssetRegistryImpl::ScanPathsSynchronous: Found Asset: %s"),
 					*AssetData->ObjectPath.ToString());
-				Context.OutFoundAssets->Add(AssetData->ObjectPath);
+				FoundAssets->Add(AssetData->ObjectPath);
 			}
 		}
 	};
 
 	bool bUnusedInterrupted;
 	Context.Status = TickGatherer(Context.EventContext, -1., bUnusedInterrupted, FAssetsFoundCallback(AssetsFoundCallback));
+#if WITH_EDITOR
+	LoadCalculatedDependencies(FoundAssets, -1., bUnusedInterrupted);
+#endif
 }
 
 namespace Utils
@@ -4137,6 +4320,9 @@ void FAssetRegistryImpl::DependencyDataGathered(const double TickStartTime, TRin
 		if (Result.bHasDependencyData)
 		{
 			FDependsNode* Node = State.CreateOrFindDependsNode(Result.PackageName);
+#if WITH_EDITOR
+			PackagesNeedingDependencyCalculation.Add(Result.PackageName);
+#endif
 
 			// We will populate the node dependencies below. Empty the set here in case this file was already read
 			// Also remove references to all existing dependencies, those will be also repopulated below
@@ -4587,6 +4773,29 @@ void FAssetRegistryImpl::OnDirectoryChanged(Impl::FEventContext& EventContext, T
 				break;
 			}
 		}
+
+#if WITH_EDITOR
+		if (bIsValidPackageName)
+		{
+			// If a package changes in a referenced directory, modify the Assets that monitor that directory
+			FString DirectoryPath = FPaths::CreateStandardFilename(FPaths::GetPath(File));
+			// TODO: Change DirectoryReferencers to a FileNameMap so we can do this FindParentDirectory check more quickly
+			TArray<FName, TInlineAllocator<1>> WatcherPackageNames;
+			for (const TPair<FString, FName>& Pair : DirectoryReferencers)
+			{
+				if (FPathViews::IsParentPathOf(Pair.Key, DirectoryPath))
+				{
+					WatcherPackageNames.Add(Pair.Value);
+				}
+			}
+			for (FName WatcherPackageName : WatcherPackageNames)
+			{
+				// ScanModifiedAssetFiles accepts LongPackageNames as well as LocalPaths
+				ModifiedFiles.AddUnique(WatcherPackageName.ToString());
+			}
+		}
+#endif
+			
 	}
 
 	if (NewFiles.Num() || NewDirs.Num())
diff --git a/Engine/Source/Runtime/AssetRegistry/Private/AssetRegistryImpl.h b/Engine/Source/Runtime/AssetRegistry/Private/AssetRegistryImpl.h
--- a/Engine/Source/Runtime/AssetRegistry/Private/AssetRegistryImpl.h
+++ b/Engine/Source/Runtime/AssetRegistry/Private/AssetRegistryImpl.h
@@ -28,6 +28,10 @@ class UAssetRegistryImpl;
 namespace UE::AssetRegistry::Impl { struct FEventContext; }
 namespace UE::AssetRegistry::Premade { struct FAsyncConsumer; }
 
+#if WITH_EDITOR
+namespace UE::AssetDependencyGatherer::Private { class FRegisteredAssetDependencyGatherer; }
+#endif
+
 #if ASSETREGISTRY_ENABLE_PREMADE_REGISTRY_IN_EDITOR
 namespace UE::AssetRegistry::Premade
 {
@@ -182,6 +186,18 @@ public:
 	void GetProcessLoadedAssetsBatch(TArray<const UObject*>& OutLoadedAssets, uint32 BatchSize);
 	void PushProcessLoadedAssetsBatch(Impl::FEventContext& EventContext,
 		TArrayView<FAssetData> LoadedAssetDatas, TArrayView<const UObject*> UnprocessedFromBatch);
+	/** Call LoadCalculatedDependencies on each Package updated after the last LoadCalculatedDependencies. */
+	void LoadCalculatedDependencies(TArray<FName>* AssetsToCalculate, double TickStartTime, bool& bOutInterrupted);
+	/**
+	 * Look for a CalculatedDependencies function registered for the asset(s) in the given package
+	 * and call that function to add calculated dependencies. Calculated dependencies are added only after
+	 * all normal dependencies gathered from the AssetRegistry data stored in the package have been loaded.
+	 */
+	void LoadCalculatedDependencies(FName PackageName, bool& bOutHadActivity);
+	/** Add a watch on a directory that to modify data for a package whenever packages in the directory are modified. */
+	void AddDirectoryReferencer(FName PackageName, const FString& DirectoryLocalPathOrLongPackageName);
+	/** Remove all directory watches for PackageName. */
+	void RemoveDirectoryReferencer(FName PackageName);
 #endif
 
 	/** Adds an asset to the empty package list which contains packages that have no assets left in them */
@@ -382,6 +398,18 @@ private:
 
 	/** The set of object paths that have had their disk cache updated from the in memory version */
 	TSet<FName> AssetDataObjectPathsUpdatedOnLoad;
+
+	/**
+	 * The set of object paths that have had their dependencies gathered since the last idle,
+	 * and that need to check for calculated dependencies at the next idle.
+	 */
+	TSet<FName> PackagesNeedingDependencyCalculation;
+
+	/** A map from directoryname to packagename of Packages that have CalculatedDependencies on packages in the directory. */
+	TMultiMap<FString, FName> DirectoryReferencers;
+
+	/** A map of per asset class dependency gatherer called in LoadCalculatedDependencies */
+	TMap<UClass*, UE::AssetDependencyGatherer::Private::FRegisteredAssetDependencyGatherer*> RegisteredDependencyGathererClasses;
 #endif
 #if WITH_ENGINE && WITH_EDITOR
 	/** Class names that return true for IsAsset but which should not be treated as assets in uncooked packages */
diff --git a/Engine/Source/Runtime/AssetRegistry/Public/AssetRegistry/AssetDependencyGatherer.h b/Engine/Source/Runtime/AssetRegistry/Public/AssetRegistry/AssetDependencyGatherer.h
new file mode 100644
--- /dev/null
+++ b/Engine/Source/Runtime/AssetRegistry/Public/AssetRegistry/AssetDependencyGatherer.h
@@ -0,0 +1,88 @@
+// Copyright Epic Games, Inc. All Rights Reserved.
+
+#pragma once
+
+#if WITH_EDITOR
+
+#include "Containers/List.h"
+#include "UObject/Object.h"
+#include "UObject/ObjectMacros.h"
+
+#include "AssetRegistry/AssetData.h"
+#include "AssetRegistry/AssetRegistryState.h"
+#include "Misc/AssetRegistryInterface.h"
+
+/**
+ * Interface class for functions that return extra dependencies to add to Assets of a given Class, 
+ * to e.g. add dependencies on external Assets in a subdirectory. 
+ */
+class IAssetDependencyGatherer
+{
+public:
+	struct FGathereredDependency
+	{
+		FName PackageName;
+		UE::AssetRegistry::EDependencyProperty Property;
+	};
+
+	/**
+	 * Return the extra dependencies to add to the given AssetData, which is guaranteed to 
+	 * have ClassType equal to the class for which the Gatherer was registered.
+	 * 
+	 * WARNING: For high performance these callbacks are called inside the critical section of the AssetRegistry. 
+	 * Attempting to call public functions on the AssetRegistry will deadlock. 
+	 * To send queries about what assets exist, used the passed-in interface functions instead.
+	 */
+	virtual void GatherDependencies(const FAssetData& AssetData, const FAssetRegistryState& AssetRegistryState, 
+		TFunctionRef<FARCompiledFilter(const FARFilter&)> CompileFilterFunc, TArray<IAssetDependencyGatherer::FGathereredDependency>& OutDependencies, 
+		TArray<FString>& OutDependencyDirectories) const = 0;
+};
+
+
+namespace UE::AssetDependencyGatherer::Private
+{
+	class ASSETREGISTRY_API FRegisteredAssetDependencyGatherer
+	{
+	public:
+		FRegisteredAssetDependencyGatherer();
+		virtual ~FRegisteredAssetDependencyGatherer();
+
+		virtual UClass* GetAssetClass() const = 0;
+		virtual void GatherDependencies(const FAssetData& AssetData, const FAssetRegistryState& AssetRegistryState, 
+			TFunctionRef<FARCompiledFilter(const FARFilter&)> CompileFilterFunc, TArray<IAssetDependencyGatherer::FGathereredDependency>& OutDependencies, 
+			TArray<FString>& OutDependencyDirectories) const = 0;
+
+		static void ForEach(TFunctionRef<void(FRegisteredAssetDependencyGatherer*)> Func);
+		static bool IsEmpty() { return !GetRegisteredList(); }
+	private:
+
+		static TLinkedList<FRegisteredAssetDependencyGatherer*>*& GetRegisteredList();
+		TLinkedList<FRegisteredAssetDependencyGatherer*> GlobalListLink;
+	};
+}
+
+
+/**
+ * Used to Register an IAssetDependencyGatherer for a class
+ *
+ * Example usage:
+ *
+ * class FMyAssetDependencyGatherer : public IAssetDependencyGatherer { ... }
+ * REGISTER_ASSETDEPENDENCY_GATHERER(FMyAssetDependencyGatherer, UMyAssetClass);
+ */
+#define REGISTER_ASSETDEPENDENCY_GATHERER(GathererClass, AssetClass) \
+class PREPROCESSOR_JOIN(PREPROCESSOR_JOIN(GathererClass, AssetClass), _Register) : public UE::AssetDependencyGatherer::Private::FRegisteredAssetDependencyGatherer \
+{ \
+	virtual UClass* GetAssetClass() const override { return AssetClass::StaticClass(); } \
+	virtual void GatherDependencies(const FAssetData& AssetData, const FAssetRegistryState& AssetRegistryState, TFunctionRef<FARCompiledFilter(const FARFilter&)> CompileFilterFunc, TArray<IAssetDependencyGatherer::FGathereredDependency>& OutDependencies, TArray<FString>& OutDependencyDirectories) const override \
+	{ \
+		GathererClass Gatherer; \
+		Gatherer.GatherDependencies(AssetData, AssetRegistryState, CompileFilterFunc, OutDependencies, OutDependencyDirectories); \
+	} \
+}; \
+namespace PREPROCESSOR_JOIN(GathererClass, AssetClass) \
+{ \
+	static PREPROCESSOR_JOIN(PREPROCESSOR_JOIN(GathererClass, AssetClass), _Register) DefaultObject; \
+}
+
+#endif // WITH_EDITOR
\ No newline at end of file
diff --git a/Engine/Source/Runtime/Engine/Private/WorldPartition/WorldPartitionAssetDependencyGatherer.cpp b/Engine/Source/Runtime/Engine/Private/WorldPartition/WorldPartitionAssetDependencyGatherer.cpp
new file mode 100644
--- /dev/null
+++ b/Engine/Source/Runtime/Engine/Private/WorldPartition/WorldPartitionAssetDependencyGatherer.cpp
@@ -0,0 +1,38 @@
+// Copyright Epic Games, Inc. All Rights Reserved.
+#include "WorldPartitionAssetDependencyGatherer.h"
+
+#if WITH_EDITOR
+
+#include "Engine/World.h"
+#include "Engine/Level.h"
+#include "HAL/FileManager.h"
+
+// Register FWorldPartitionCookPackageSplitter for UWorld class
+REGISTER_ASSETDEPENDENCY_GATHERER(FWorldPartitionAssetDependencyGatherer, UWorld);
+
+void FWorldPartitionAssetDependencyGatherer::GatherDependencies(const FAssetData& AssetData, const FAssetRegistryState& AssetRegistryState, TFunctionRef<FARCompiledFilter(const FARFilter&)> CompileFilterFunc, TArray<IAssetDependencyGatherer::FGathereredDependency>& OutDependencies, TArray<FString>& OutDependencyDirectories) const
+{
+	if (ULevel::GetIsLevelPartitionedFromAsset(AssetData))
+	{
+		const FString ExternalActorsPath = ULevel::GetExternalActorsPath(AssetData.PackageName.ToString());
+		OutDependencyDirectories.Add(ExternalActorsPath);
+
+		FARFilter Filter;
+		Filter.bRecursivePaths = true;
+		Filter.bIncludeOnlyOnDiskAssets = true;
+		Filter.PackagePaths.Add(*ExternalActorsPath);
+
+		TArray<FAssetData> FilteredAssets;
+		AssetRegistryState.GetAssets(CompileFilterFunc(Filter), {}, FilteredAssets, true);
+
+		for (const FAssetData& FilteredAsset : FilteredAssets)
+		{
+			if (const UClass* FilteredAssetClass = FilteredAsset.GetClass(); FilteredAssetClass && FilteredAssetClass->IsChildOf(AActor::StaticClass()))
+			{
+				OutDependencies.Emplace(IAssetDependencyGatherer::FGathereredDependency{ FilteredAsset.PackageName, UE::AssetRegistry::EDependencyProperty::Game });
+			}
+		}
+	}
+}
+
+#endif
\ No newline at end of file
diff --git a/Engine/Source/Runtime/Engine/Private/WorldPartition/WorldPartitionAssetDependencyGatherer.h b/Engine/Source/Runtime/Engine/Private/WorldPartition/WorldPartitionAssetDependencyGatherer.h
new file mode 100644
--- /dev/null
+++ b/Engine/Source/Runtime/Engine/Private/WorldPartition/WorldPartitionAssetDependencyGatherer.h
@@ -0,0 +1,21 @@
+// Copyright Epic Games, Inc. All Rights Reserved.
+
+#pragma once
+
+#if WITH_EDITOR
+
+#include "AssetRegistry/AssetDependencyGatherer.h"
+#include "AssetRegistry/AssetData.h"
+#include "Engine/World.h"
+
+class FWorldPartitionAssetDependencyGatherer : public IAssetDependencyGatherer
+{
+public:
+	
+	FWorldPartitionAssetDependencyGatherer() = default;
+	virtual ~FWorldPartitionAssetDependencyGatherer() = default;
+
+	virtual void GatherDependencies(const FAssetData& AssetData, const FAssetRegistryState& AssetRegistryState, TFunctionRef<FARCompiledFilter(const FARFilter&)> CompileFilterFunc, TArray<IAssetDependencyGatherer::FGathereredDependency>& OutDependencies, TArray<FString>& OutDependencyDirectories) const override;
+};
+
+#endif
\ No newline at end of file
