Fix regression: infinite loop when timing out (from CL 20102122)

#rnx

diff --git a/Engine/Source/Runtime/AssetRegistry/Private/AssetRegistry.cpp b/Engine/Source/Runtime/AssetRegistry/Private/AssetRegistry.cpp
--- a/Engine/Source/Runtime/AssetRegistry/Private/AssetRegistry.cpp
+++ b/Engine/Source/Runtime/AssetRegistry/Private/AssetRegistry.cpp
@@ -3624,16 +3624,17 @@ void FAssetRegistryImpl::LoadCalculatedDependencies(TArray<FName>* AssetsToCalcu
 	}
 	else
 	{
-		for (FName PackageName : PackagesNeedingDependencyCalculation)
+		for (TSet<FName>::TIterator It = PackagesNeedingDependencyCalculation.CreateIterator(); It; ++It)
 		{
 			bool bHadActivity;
-			LoadCalculatedDependencies(PackageName, bHadActivity);
+			LoadCalculatedDependencies(*It, bHadActivity);
+			It.RemoveCurrent();
 			if (CheckForTimeUp(bHadActivity))
 			{
 				return;
 			}
 		}
-		PackagesNeedingDependencyCalculation.Empty();
+		check(PackagesNeedingDependencyCalculation.IsEmpty());
 	}
 }
 
