AssetRegistryState: fix EnumerateAssets to handle duplicates in Filter.ObjectPaths, and optimize the handling of duplicates using a bit for each filters rather than counting the number of passed filter hits and requiring uniqueness in the filterresults.
Reported by PR #8466 submitted by [REDACTED_PERSON].
#rnx

diff --git a/Engine/Source/Runtime/AssetRegistry/Private/AssetRegistryState.cpp b/Engine/Source/Runtime/AssetRegistry/Private/AssetRegistryState.cpp
--- a/Engine/Source/Runtime/AssetRegistry/Private/AssetRegistryState.cpp
+++ b/Engine/Source/Runtime/AssetRegistry/Private/AssetRegistryState.cpp
@@ -705,7 +705,7 @@ bool FAssetRegistryState::EnumerateAssets(const FARCompiledFilter& Filter, const
 	// On disk tags and values
 	if (Filter.TagsAndValues.Num() > 0)
 	{
-		TSet<FAssetData*> TagAndValuesFilter;
+		TArray<FAssetData*>& TagAndValuesFilter = FilterResults.Emplace_GetRef();
 		// Sometimes number of assets matching this filter is correlated to number of assets matching previous filters 
 		if (FilterResults.Num())
 		{
@@ -740,8 +740,6 @@ bool FAssetRegistryState::EnumerateAssets(const FARCompiledFilter& Filter, const
 				}
 			}
 		}
-
-		FilterResults.Emplace(TagAndValuesFilter.Array());
 	}
 
 	// Perform callback for assets that match all filters
@@ -759,26 +757,28 @@ bool FAssetRegistryState::EnumerateAssets(const FARCompiledFilter& Filter, const
 			return bARFiltering && UE::AssetRegistry::FFiltering::ShouldSkipAsset(AssetData->AssetClass, AssetData->PackageFlags);
 		};
 
-		if (FilterResults.Num() > 1)
+		int32 NumFilterResults = FilterResults.Num();
+		if (NumFilterResults > 1)
 		{
-			// Count how many filters each asset pass
-			TMap<FAssetData*, uint32> PassCounts;
-			for (const TArray<FAssetData*>& FilterEvaluation : FilterResults)
+			// Mark which filters each asset passes
+			uint32 PassAllFiltersValue = (1 << NumFilterResults) - 1; // 1 in every bit for the lowest n bits
+			TMap<FAssetData*, uint32> PassBits;
+			for (int32 FilterIndex = 0; FilterIndex < NumFilterResults; ++FilterIndex)
 			{
-				PassCounts.Reserve(FilterEvaluation.Num());
+				const TArray<FAssetData*>& FilterEvaluation = FilterResults[FilterIndex];
+				PassBits.Reserve(FilterEvaluation.Num());
 				
 				for (FAssetData* AssetData : FilterEvaluation)
 				{
-					PassCounts.FindOrAdd(AssetData)++;	
+					PassBits.FindOrAdd(AssetData) |= (1 << FilterIndex);
 				}
 			}
 
-			// Include assets that match all filters
-			for (TPair<FAssetData*, uint32> PassCount : PassCounts)
+			// Include assets that pass all filters
+			for (TPair<FAssetData*, uint32> PassPair : PassBits)
 			{
-				const FAssetData* AssetData = PassCount.Key;
-				check(PassCount.Value <= (uint32)FilterResults.Num());
-				if (PassCount.Value != FilterResults.Num() || SkipAssetData(AssetData))
+				const FAssetData* AssetData = PassPair.Key;
+				if (PassPair.Value != PassAllFiltersValue || SkipAssetData(AssetData))
 				{
 					continue;
 				}
