Deprecating ANY_PACKAGE.

diff --git a/Engine/Source/Editor/UnrealEd/Private/Tests/AutomationEditorCommon.cpp b/Engine/Source/Editor/UnrealEd/Private/Tests/AutomationEditorCommon.cpp
--- a/Engine/Source/Editor/UnrealEd/Private/Tests/AutomationEditorCommon.cpp
+++ b/Engine/Source/Editor/UnrealEd/Private/Tests/AutomationEditorCommon.cpp
@@ -600,7 +600,7 @@ void FAutomationEditorCommonUtils::CollectTestsByClass(UClass * Class, TArray<FS
 {
 	FAssetRegistryModule& AssetRegistryModule = FModuleManager::LoadModuleChecked<FAssetRegistryModule>(TEXT("AssetRegistry"));
 	TArray<FAssetData> ObjectList;
-	AssetRegistryModule.Get().GetAssetsByClass(Class->GetFName(), ObjectList);
+	AssetRegistryModule.Get().GetAssetsByClass(Class->GetClassPathName(), ObjectList);
 
 	for (TObjectIterator<UClass> AllClassesIt; AllClassesIt; ++AllClassesIt)
 	{
@@ -638,7 +638,7 @@ void FAutomationEditorCommonUtils::CollectGameContentTestsByClass(UClass * Class
 
 	//Generating the list of assets.
 	//This list is being filtered by the game folder and class type.  The results are placed into the ObjectList variable.
-	AssetFilter.ClassNames.Add(Class->GetFName());
+	AssetFilter.ClassPaths.Add(Class->GetClassPathName());
 
 	//removed path as a filter as it causes two large lists to be sorted.  Filtering on "game" directory on iteration
 	//AssetFilter.PackagePaths.Add("/Game");
@@ -728,7 +728,7 @@ void FAutomationEditorCommonUtils::CollectGameContentTests(TArray<FString>& OutB
 				Filename = FPackageName::LongPackageNameToFilename(Filename);
 				if (FAutomationTestFramework::Get().ShouldTestContent(Filename))
 				{
-					FString BeautifiedFilename = FString::Printf(TEXT("%s.%s"), *Asset.AssetClass.ToString(), *Asset.AssetName.ToString());
+					FString BeautifiedFilename = FString::Printf(TEXT("%s.%s"), *Asset.AssetClassPath.ToString(), *Asset.AssetName.ToString());
 					OutBeautifiedNames.Add(BeautifiedFilename);
 					OutTestCommands.Add(Asset.ObjectPath.ToString());
 				}
diff --git a/Engine/Source/Editor/UnrealEd/Private/UnrealEdMisc.cpp b/Engine/Source/Editor/UnrealEd/Private/UnrealEdMisc.cpp
--- a/Engine/Source/Editor/UnrealEd/Private/UnrealEdMisc.cpp
+++ b/Engine/Source/Editor/UnrealEd/Private/UnrealEdMisc.cpp
@@ -824,26 +824,26 @@ void FUnrealEdMisc::TickAssetAnalytics()
 			TArray< FAnalyticsEventAttribute > AssetAttributes;
 			int32 NumMapFiles = 0;
 			TSet< FName > PackageNames;
-			TMap< FName, int32 > ClassInstanceCounts;
+			TMap< FTopLevelAssetPath, int32 > ClassInstanceCounts;
 
 			for( auto AssetIter = AssetData.CreateConstIterator(); AssetIter; ++AssetIter )
 			{
 				PackageNames.Add( AssetIter->PackageName );
-				if( AssetIter->AssetClass == UWorld::StaticClass()->GetFName()  )
+				if( AssetIter->AssetClassPath == UWorld::StaticClass()->GetClassPathName()  )
 				{
 					NumMapFiles++;
 				}
 
-				if( AssetIter->AssetClass != NAME_None )
+				if (!AssetIter->AssetClassPath.IsNull())
 				{
-					int32* ExistingClassCount = ClassInstanceCounts.Find( AssetIter->AssetClass );
+					int32* ExistingClassCount = ClassInstanceCounts.Find( AssetIter->AssetClassPath);
 					if( ExistingClassCount )
 					{
 						++(*ExistingClassCount);
 					}
 					else
 					{
-						ClassInstanceCounts.Add( AssetIter->AssetClass, 1 );
+						ClassInstanceCounts.Add( AssetIter->AssetClassPath, 1 );
 					}
 				}
 			}
diff --git a/Engine/Source/Editor/UnrealEd/Private/UnrealEdSrv.cpp b/Engine/Source/Editor/UnrealEd/Private/UnrealEdSrv.cpp
--- a/Engine/Source/Editor/UnrealEd/Private/UnrealEdSrv.cpp
+++ b/Engine/Source/Editor/UnrealEd/Private/UnrealEdSrv.cpp
@@ -1862,7 +1862,7 @@ bool UUnrealEdEngine::Exec_Actor( UWorld* InWorld, const TCHAR* Str, FOutputDevi
 	if( FParse::Command(&Str,TEXT("ADD")) )
 	{
 		UClass* Class;
-		if( ParseObject<UClass>( Str, TEXT("CLASS="), Class, ANY_PACKAGE ) )
+		if( ParseObject<UClass>( Str, TEXT("CLASS="), Class, nullptr ) )
 		{
 			int32 bSnap = 1;
 			FParse::Value(Str,TEXT("SNAP="),bSnap);
@@ -2258,7 +2258,7 @@ bool UUnrealEdEngine::Exec_Actor( UWorld* InWorld, const TCHAR* Str, FOutputDevi
 			edactReplaceSelectedBrush( InWorld );
 			return true;
 		}
-		else if( ParseObject<UClass>( Str, TEXT("CLASS="), Class, ANY_PACKAGE ) ) // ACTOR REPLACE CLASS=<class>
+		else if( ParseObject<UClass>( Str, TEXT("CLASS="), Class, nullptr ) ) // ACTOR REPLACE CLASS=<class>
 		{
 			const FScopedTransaction Transaction( NSLOCTEXT("UnrealEd", "ReplaceSelectedNonBrushActors", "Replace Selected Non-Brush Actors") );
 			edactReplaceSelectedNonBrushWithClass( Class );
@@ -2359,7 +2359,7 @@ bool UUnrealEdEngine::Exec_Actor( UWorld* InWorld, const TCHAR* Str, FOutputDevi
 		else if( FParse::Command(&Str,TEXT("OFCLASS")) ) // ACTOR SELECT OFCLASS CLASS=<class>
 		{
 			UClass* Class;
-			if( ParseObject<UClass>(Str,TEXT("CLASS="),Class,ANY_PACKAGE) )
+			if( ParseObject<UClass>(Str,TEXT("CLASS="),Class,nullptr) )
 			{
 				const FScopedTransaction Transaction( NSLOCTEXT("UnrealEd", "SelectOfClass", "Select Of Class") );
 				edactSelectOfClass( InWorld, Class );
@@ -2373,7 +2373,7 @@ bool UUnrealEdEngine::Exec_Actor( UWorld* InWorld, const TCHAR* Str, FOutputDevi
 		else if( FParse::Command(&Str,TEXT("OFSUBCLASS")) ) // ACTOR SELECT OFSUBCLASS CLASS=<class>
 		{
 			UClass* Class;
-			if( ParseObject<UClass>(Str,TEXT("CLASS="),Class,ANY_PACKAGE) )
+			if( ParseObject<UClass>(Str,TEXT("CLASS="),Class,nullptr) )
 			{
 				const FScopedTransaction Transaction( NSLOCTEXT("UnrealEd", "SelectSubclassOfClass", "Select Subclass Of Class") );
 				edactSelectSubclassOf( InWorld, Class );
diff --git a/Engine/Source/Editor/UnrealEd/Private/WorkflowOrientedApp/SContentReference.cpp b/Engine/Source/Editor/UnrealEd/Private/WorkflowOrientedApp/SContentReference.cpp
--- a/Engine/Source/Editor/UnrealEd/Private/WorkflowOrientedApp/SContentReference.cpp
+++ b/Engine/Source/Editor/UnrealEd/Private/WorkflowOrientedApp/SContentReference.cpp
@@ -213,7 +213,7 @@ TSharedRef<SWidget> SContentReference::MakeAssetPickerMenu()
 	UClass* FilterClass = AllowedClass.Get();
 	if (FilterClass != NULL)
 	{
-		AssetPickerConfig.Filter.ClassNames.Add(FilterClass->GetFName());
+		AssetPickerConfig.Filter.ClassPaths.Add(FilterClass->GetClassPathName());
 		AssetPickerConfig.Filter.bRecursiveClasses = true;
 	}
 
diff --git a/Engine/Source/Editor/UnrealEd/Private/WorldPartition/WorldPartitionResaveActorsBuilder.cpp b/Engine/Source/Editor/UnrealEd/Private/WorldPartition/WorldPartitionResaveActorsBuilder.cpp
--- a/Engine/Source/Editor/UnrealEd/Private/WorldPartition/WorldPartitionResaveActorsBuilder.cpp
+++ b/Engine/Source/Editor/UnrealEd/Private/WorldPartition/WorldPartitionResaveActorsBuilder.cpp
@@ -109,7 +109,7 @@ bool UWorldPartitionResaveActorsBuilder::RunInternal(UWorld* World, const FCellI
 		}
 
 		// Look for native classes
-		ActorClass = FindObject<UClass>(ANY_PACKAGE, *ActorClassName);
+		ActorClass = UClass::TryFindTypeSlow<UClass>(ActorClassName);
 
 		if (!ActorClass)
 		{
diff --git a/Engine/Source/Editor/UnrealEd/Public/Commandlets/CompileAllBlueprintsCommandlet.h b/Engine/Source/Editor/UnrealEd/Public/Commandlets/CompileAllBlueprintsCommandlet.h
--- a/Engine/Source/Editor/UnrealEd/Public/Commandlets/CompileAllBlueprintsCommandlet.h
+++ b/Engine/Source/Editor/UnrealEd/Public/Commandlets/CompileAllBlueprintsCommandlet.h
@@ -87,7 +87,7 @@ protected:
 	TArray<FString> FileAllowList;
 	TArray<TPair<FString, TArray<FString>>> RequireAssetTags;
 	TArray<TPair<FString, TArray<FString>>> ExcludeAssetTags;
-	FName BlueprintBaseClassName;
+	FTopLevelAssetPath BlueprintBaseClassName;
 
 	//Variables to store overall results
 	int TotalNumFailedLoads;
diff --git a/Engine/Source/Editor/UnrealEd/Public/Editor.h b/Engine/Source/Editor/UnrealEd/Public/Editor.h
--- a/Engine/Source/Editor/UnrealEd/Public/Editor.h
+++ b/Engine/Source/Editor/UnrealEd/Public/Editor.h
@@ -631,19 +631,29 @@ TCHAR* SetFVECTOR( TCHAR* Dest, const FVector* Value );
  * Takes an FName and checks to see that it is unique among all loaded objects.
  *
  * @param	InName		The name to check
- * @param	Outer		The context for validating this object name. Should be a group/package, but could be ANY_PACKAGE if you want to check across the whole system (not recommended)
+ * @param	Outer		The context for validating this object name. Should be a group/package
  * @param	InReason	If the check fails, this string is filled in with the reason why.
  *
  * @return	true if the name is valid
  */
 
-UNREALED_API bool IsUniqueObjectName( const FName& InName, UObject* Outer, FText* InReason=NULL );
+UNREALED_API bool IsUniqueObjectName( const FName& InName, UObject* Outer, FText* InReason = nullptr );
+
+/**
+ * Takes an FName and checks to see that it is unique among all loaded objects in all packages.
+ *
+ * @param	InName		The name to check
+ * @param	InReason	If the check fails, this string is filled in with the reason why.
+ *
+ * @return	true if the name is valid
+ */
+UNREALED_API bool IsGloballyUniqueObjectName(const FName& InName, FText* InReason = nullptr);
 
 /**
  * Takes an FName and checks to see that it is unique among all loaded objects.
  *
  * @param	InName		The name to check
- * @param	Outer		The context for validating this object name. Should be a group/package, but could be ANY_PACKAGE if you want to check across the whole system (not recommended)
+ * @param	Outer		The context for validating this object name. Should be a group/package.
  * @param	InReason	If the check fails, this string is filled in with the reason why.
  *
  * @return	true if the name is valid
diff --git a/Engine/Source/Editor/UnrealEd/Public/EditorClassUtils.h b/Engine/Source/Editor/UnrealEd/Public/EditorClassUtils.h
--- a/Engine/Source/Editor/UnrealEd/Public/EditorClassUtils.h
+++ b/Engine/Source/Editor/UnrealEd/Public/EditorClassUtils.h
@@ -153,8 +153,17 @@ namespace FEditorClassUtils
 	 * @param	InAssetData		Reference to an asset data entry.
 	 * @return					Class path or None if the asset cannot or doesn't have a class associated with it
 	 */
+	UE_DEPRECATED(5.1, "Class names are now represented by path names. Please use GetClassPathNameFromAssetTag.")
 	UNREALED_API FName GetClassPathFromAssetTag(const FAssetData& InAssetData);
 
+	/**
+	 * Gets the class path from the asset tag (i.e. GeneratedClassPath tag on blueprints)
+	 * 
+	 * @param	InAssetData		Reference to an asset data entry.
+	 * @return					Class path or None if the asset cannot or doesn't have a class associated with it
+	 */
+	UNREALED_API FTopLevelAssetPath GetClassPathNameFromAssetTag(const FAssetData& InAssetData);
+
 	/**
 	 * Gets the object path of the class associated with the specified asset 
 	 * (i.e. the BlueprintGeneratedClass of a Blueprint asset or the BlueprintGeneratedClass asset itself)
@@ -163,8 +172,19 @@ namespace FEditorClassUtils
 	 * @param	bGenerateClassPathIfMissing	Whether to generate a class path if the class is missing (and the asset can have a class associated with it)
 	 * @return								Class path or None if the asset cannot or doesn't have a class associated with it
 	 */
+	UE_DEPRECATED(5.1, "Class names are now represented by path names. Please use GetClassPathNameFromAsset.")
 	UNREALED_API FName GetClassPathFromAsset(const FAssetData& InAssetData, bool bGenerateClassPathIfMissing = false);
 
+	/**
+	 * Gets the object path of the class associated with the specified asset 
+	 * (i.e. the BlueprintGeneratedClass of a Blueprint asset or the BlueprintGeneratedClass asset itself)
+	 * 
+	 * @param	InAssetData					Reference to an asset data entry
+	 * @param	bGenerateClassPathIfMissing	Whether to generate a class path if the class is missing (and the asset can have a class associated with it)
+	 * @return								Class path or None if the asset cannot or doesn't have a class associated with it
+	 */
+	UNREALED_API FTopLevelAssetPath GetClassPathNameFromAsset(const FAssetData& InAssetData, bool bGenerateClassPathIfMissing = false);
+
 	/**
 	 * Fetches the set of interface class object paths from an asset data entry containing the appropriate asset tag(s).
 	 * 
diff --git a/Engine/Source/Editor/UnrealEd/Public/Kismet2/BlueprintEditorUtils.h b/Engine/Source/Editor/UnrealEd/Public/Kismet2/BlueprintEditorUtils.h
--- a/Engine/Source/Editor/UnrealEd/Public/Kismet2/BlueprintEditorUtils.h
+++ b/Engine/Source/Editor/UnrealEd/Public/Kismet2/BlueprintEditorUtils.h
@@ -1419,11 +1419,19 @@ public:
 	static FGuid FindInterfaceFunctionGuid(const UFunction* Function, const UClass* InterfaceClass);
 
 	/** Add a new interface, and member function graphs to the blueprint */
+	UE_DEPRECATED(5.1, "Short class names are no longer supported. Use a version of this function that takes FTopLevelAssetPath.")
 	static bool ImplementNewInterface(UBlueprint* Blueprint, const FName& InterfaceClassName);
 
+	/** Add a new interface, and member function graphs to the blueprint */
+	static bool ImplementNewInterface(UBlueprint* Blueprint, FTopLevelAssetPath InterfaceClassPathName);
+
 	/** Remove an implemented interface, and its associated member function graphs.  If bPreserveFunctions is true, then the interface will move its functions to be normal implemented blueprint functions */
+	UE_DEPRECATED(5.1, "Short class names are no longer supported. Use a version of this function that takes FTopLevelAssetPath.")
 	static void RemoveInterface(UBlueprint* Blueprint, const FName& InterfaceClassName, bool bPreserveFunctions = false);
 	
+	/** Remove an implemented interface, and its associated member function graphs.  If bPreserveFunctions is true, then the interface will move its functions to be normal implemented blueprint functions */
+	static void RemoveInterface(UBlueprint* Blueprint, FTopLevelAssetPath InterfaceClassPathName, bool bPreserveFunctions = false);
+
 	/**
 	* Attempt to remove a function from an interfaces list of function graphs.
 	* Note that this will NOT remove interface events (i.e. functions with no outputs)
@@ -1441,8 +1449,12 @@ public:
 	static void PromoteGraphFromInterfaceOverride(UBlueprint* InBlueprint, UEdGraph* InInterfaceGraph);
 
 	/** Gets the graphs currently in the blueprint associated with the specified interface */
+	UE_DEPRECATED(5.1, "Short class names are no longer supported. Use a version of this function that takes FTopLevelAssetPath.")
 	static void GetInterfaceGraphs(UBlueprint* Blueprint, const FName& InterfaceClassName, TArray<UEdGraph*>& ChildGraphs);
 
+	/** Gets the graphs currently in the blueprint associated with the specified interface */
+	static void GetInterfaceGraphs(UBlueprint* Blueprint, FTopLevelAssetPath InterfaceClassPathName, TArray<UEdGraph*>& ChildGraphs);
+
 	/**
 	* Checks if the given function is a part of an interface on this blueprint
 	* 
diff --git a/Engine/Source/Editor/UnrealEd/Public/LocalizedAssetUtil.h b/Engine/Source/Editor/UnrealEd/Public/LocalizedAssetUtil.h
--- a/Engine/Source/Editor/UnrealEd/Public/LocalizedAssetUtil.h
+++ b/Engine/Source/Editor/UnrealEd/Public/LocalizedAssetUtil.h
@@ -26,6 +26,6 @@ struct FLocalizedAssetSCCUtil
 
 struct FLocalizedAssetUtil
 {
-	static bool GetAssetsByPathAndClass(IAssetRegistry& InAssetRegistry, const FName InPackagePath, const FName InClassName, const bool bIncludeLocalizedAssets, TArray<FAssetData>& OutAssets);
-	static bool GetAssetsByPathAndClass(IAssetRegistry& InAssetRegistry, const TArray<FName>& InPackagePaths, const FName InClassName, const bool bIncludeLocalizedAssets, TArray<FAssetData>& OutAssets);
+	static bool GetAssetsByPathAndClass(IAssetRegistry& InAssetRegistry, const FName InPackagePath, const FTopLevelAssetPath& InClassName, const bool bIncludeLocalizedAssets, TArray<FAssetData>& OutAssets);
+	static bool GetAssetsByPathAndClass(IAssetRegistry& InAssetRegistry, const TArray<FName>& InPackagePaths, const FTopLevelAssetPath& InClassName, const bool bIncludeLocalizedAssets, TArray<FAssetData>& OutAssets);
 };
diff --git a/Engine/Source/Editor/WorldBrowser/Private/SLevelsTreeWidget.cpp b/Engine/Source/Editor/WorldBrowser/Private/SLevelsTreeWidget.cpp
--- a/Engine/Source/Editor/WorldBrowser/Private/SLevelsTreeWidget.cpp
+++ b/Engine/Source/Editor/WorldBrowser/Private/SLevelsTreeWidget.cpp
@@ -110,7 +110,7 @@ bool SLevelsTreeWidget::GetWorldAssetsFromDrag(const FDragDropEvent& DragDropEve
 	TArray<FAssetData> AssetList = AssetUtil::ExtractAssetDataFromDrag(DragDropEvent);
 	for (const auto& AssetData : AssetList)
 	{
-		if (AssetData.AssetClass == UWorld::StaticClass()->GetFName())
+		if (AssetData.AssetClassPath == UWorld::StaticClass()->GetClassPathName())
 		{
 			OutWorldAssetList.Add(AssetData);
 		}
diff --git a/Engine/Source/Editor/WorldBrowser/Private/WorldTreeItemTypes.cpp b/Engine/Source/Editor/WorldBrowser/Private/WorldTreeItemTypes.cpp
--- a/Engine/Source/Editor/WorldBrowser/Private/WorldTreeItemTypes.cpp
+++ b/Engine/Source/Editor/WorldBrowser/Private/WorldTreeItemTypes.cpp
@@ -39,7 +39,7 @@ namespace WorldHierarchy
 		
 		for (const FAssetData& Asset : AssetList)
 		{
-			if (Asset.AssetClass == UWorld::StaticClass()->GetFName())
+			if (Asset.AssetClassPath == UWorld::StaticClass()->GetClassPathName())
 			{
 				OutWorlds.Add(Asset);
 			}
diff --git a/Engine/Source/Editor/WorldPartitionEditor/Private/WorldPartitionEditorModule.cpp b/Engine/Source/Editor/WorldPartitionEditor/Private/WorldPartitionEditorModule.cpp
--- a/Engine/Source/Editor/WorldPartitionEditor/Private/WorldPartitionEditorModule.cpp
+++ b/Engine/Source/Editor/WorldPartitionEditor/Private/WorldPartitionEditorModule.cpp
@@ -231,7 +231,7 @@ void FWorldPartitionEditorModule::OnConvertMap()
 	{
 		Config.DefaultPath = OutPathName;
 	}	
-	Config.AssetClassNames.Add(UWorld::StaticClass()->GetFName());
+	Config.AssetClassNames.Add(UWorld::StaticClass()->GetClassPathName());
 
 	TArray<FAssetData> Assets = ContentBrowserSingleton.CreateModalOpenAssetDialog(Config);
 	if (Assets.Num() == 1)
diff --git a/Engine/Source/Programs/UnrealHeaderTool/Private/CodeGenerator.cpp b/Engine/Source/Programs/UnrealHeaderTool/Private/CodeGenerator.cpp
--- a/Engine/Source/Programs/UnrealHeaderTool/Private/CodeGenerator.cpp
+++ b/Engine/Source/Programs/UnrealHeaderTool/Private/CodeGenerator.cpp
@@ -6183,7 +6183,7 @@ void FNativeClassHeaderGenerator::DeleteUnusedGeneratedHeaders(TSet<FString>&& P
 								FString BaseFilename = FPaths::GetBaseFilename(Filename);
 								const int32   GeneratedIndex = BaseFilename.Find(TEXT(".generated"), ESearchCase::IgnoreCase, ESearchDir::FromEnd);
 								const FString ClassName = MoveTemp(BaseFilename).Mid(0, GeneratedIndex);
-								UClass* IntrinsicClass = FEngineAPI::FindObject<UClass>(ANY_PACKAGE, *ClassName);
+								UClass* IntrinsicClass = FEngineAPI::FindFirstObject<UClass>(*ClassName, EFindFirstObjectOptions::None, ELogVerbosity::Fatal, TEXT("looking for intrinsic class"));
 								if (!IntrinsicClass || !IntrinsicClass->HasAnyClassFlags(CLASS_Intrinsic))
 								{
 									IFileManager::Get().Delete(*Fullpath);
@@ -6487,7 +6487,7 @@ void ResolveSuperClasses(FUnrealPackageDefinitionInfo& PackageDef)
 
 UPackage* GetModulePackage(FManifestModule& Module)
 {
-	UPackage* Package = FEngineAPI::FindObjectFast<UPackage>(NULL, FName(*Module.LongPackageName), false, false);
+	UPackage* Package = FEngineAPI::FindObjectFast<UPackage>(NULL, FName(*Module.LongPackageName), false);
 	if (Package == NULL)
 	{
 		Package = CreatePackage(*Module.LongPackageName);
diff --git a/Engine/Source/Programs/UnrealHeaderTool/Private/EngineAPI.h b/Engine/Source/Programs/UnrealHeaderTool/Private/EngineAPI.h
--- a/Engine/Source/Programs/UnrealHeaderTool/Private/EngineAPI.h
+++ b/Engine/Source/Programs/UnrealHeaderTool/Private/EngineAPI.h
@@ -26,8 +26,18 @@ public:
 	 * @see StaticFindObjectFast()
 	 */
 	template< class T >
-	static inline T* FindObjectFast(UObject* Outer, FName Name, bool ExactClass = false, bool AnyPackage = false, EObjectFlags ExclusiveFlags = RF_NoFlags)
+	static inline T* FindObjectFast(UObject* Outer, FName Name, bool ExactClass = false, EObjectFlags ExclusiveFlags = RF_NoFlags)
 	{
-		return (T*)StaticFindObjectFast(T::StaticClass(), Outer, Name, ExactClass, AnyPackage, ExclusiveFlags);
+		return (T*)StaticFindObjectFast(T::StaticClass(), Outer, Name, ExactClass, ExclusiveFlags);
+	}
+
+	/**
+	 * Find an optional object.
+	 * @see StaticFindFirstObject()
+	 */
+	template< class T >
+	static inline T* FindFirstObject(const TCHAR* Name, EFindFirstObjectOptions Options = EFindFirstObjectOptions::None, ELogVerbosity::Type AmbiguousMessageVerbosity = ELogVerbosity::NoLogging, const TCHAR* InCurrentOperation = nullptr)
+	{
+		return (T*)StaticFindFirstObject(T::StaticClass(), Name, Options, AmbiguousMessageVerbosity, InCurrentOperation);
 	}
 };
diff --git a/Engine/Source/Programs/UnrealHeaderTool/Private/UnrealTypeDefinitionInfo.cpp b/Engine/Source/Programs/UnrealHeaderTool/Private/UnrealTypeDefinitionInfo.cpp
--- a/Engine/Source/Programs/UnrealHeaderTool/Private/UnrealTypeDefinitionInfo.cpp
+++ b/Engine/Source/Programs/UnrealHeaderTool/Private/UnrealTypeDefinitionInfo.cpp
@@ -2126,7 +2126,7 @@ void FUnrealClassDefinitionInfo::CreateUObjectEngineTypesInternal(ECreateEngineT
 		FString ClassNameStripped = GetClassNameWithPrefixRemoved(*ClassName);
 
 		// See if we have an existing engine object we should use
-		UClass* Class = FEngineAPI::FindObject<UClass>(ANY_PACKAGE, *ClassNameStripped);
+		UClass* Class = FEngineAPI::FindFirstObject<UClass>(*ClassNameStripped, EFindFirstObjectOptions::EnsureIfAmbiguous);
 
 		if (Class == nullptr)
 		{
@@ -2301,7 +2301,7 @@ void FUnrealClassDefinitionInfo::PostParseFinalizeInternal(EPostParseFinalizePha
 		}
 
 		// See if we have an existing engine object we should use
-		if (UClass* Class = FEngineAPI::FindObject<UClass>(ANY_PACKAGE, *GetFName().ToString()))
+		if (UClass* Class = FEngineAPI::FindFirstObject<UClass>(*GetFName().ToString(), EFindFirstObjectOptions::EnsureIfAmbiguous))
 		{
 
 			// Check for differences in ClassFlags
diff --git a/Engine/Source/Runtime/AIModule/Private/BehaviorTree/Blackboard/BlackboardKeyType_Enum.cpp b/Engine/Source/Runtime/AIModule/Private/BehaviorTree/Blackboard/BlackboardKeyType_Enum.cpp
--- a/Engine/Source/Runtime/AIModule/Private/BehaviorTree/Blackboard/BlackboardKeyType_Enum.cpp
+++ b/Engine/Source/Runtime/AIModule/Private/BehaviorTree/Blackboard/BlackboardKeyType_Enum.cpp
@@ -79,7 +79,7 @@ void UBlackboardKeyType_Enum::PostEditChangeProperty(struct FPropertyChangedEven
 	if (PropertyChangedEvent.Property &&
 		PropertyChangedEvent.Property->GetFName() == GET_MEMBER_NAME_CHECKED(UBlackboardKeyType_Enum, EnumName))
 	{
-		EnumType = FindObject<UEnum>(ANY_PACKAGE, *EnumName, true);
+		EnumType = UClass::TryFindTypeSlow<UEnum>(EnumName, EFindFirstObjectOptions::ExactClass);
 	}
 
 	bIsEnumNameValid = EnumType && !EnumName.IsEmpty();
diff --git a/Engine/Source/Runtime/AssetRegistry/Private/AssetDataGatherer.cpp b/Engine/Source/Runtime/AssetRegistry/Private/AssetDataGatherer.cpp
--- a/Engine/Source/Runtime/AssetRegistry/Private/AssetDataGatherer.cpp
+++ b/Engine/Source/Runtime/AssetRegistry/Private/AssetDataGatherer.cpp
@@ -3381,8 +3381,8 @@ bool FAssetDataGatherer::ReadAssetFile(FPackageReader& PackageReader, TArray<FAs
 			// to save them as in-game imports by adding HasNonEditorOnlyReferences; the next version bump after that
 			// fix was VER_UE4_CORRECT_LICENSEE_FLAG. Mark all dependencies in the affected version as used in game
 			// if the package has a UObjectRedirector object.
-			FName RedirectorClassName = UObjectRedirector::StaticClass()->GetFName();
-			if (Algo::AnyOf(AssetDataList, [RedirectorClassName](FAssetData* AssetData) { return AssetData->AssetClass == RedirectorClassName; }))
+			FTopLevelAssetPath RedirectorClassPathName = UObjectRedirector::StaticClass()->GetClassPathName();
+			if (Algo::AnyOf(AssetDataList, [RedirectorClassPathName](FAssetData* AssetData) { return AssetData->AssetClassPath == RedirectorClassPathName; }))
 			{
 				for (FPackageDependencyData::FPackageDependency& Dependency : DependencyData.PackageDependencies)
 				{
diff --git a/Engine/Source/Runtime/AssetRegistry/Private/AssetRegistry.cpp b/Engine/Source/Runtime/AssetRegistry/Private/AssetRegistry.cpp
--- a/Engine/Source/Runtime/AssetRegistry/Private/AssetRegistry.cpp
+++ b/Engine/Source/Runtime/AssetRegistry/Private/AssetRegistry.cpp
@@ -92,6 +92,12 @@
 
 static FAssetRegistryConsoleCommands ConsoleCommands; // Registers its various console commands in the constructor
 
+namespace UE::AssetRegistry
+{
+	const FName WildcardFName(TEXT("*"));
+	const FTopLevelAssetPath WildcardPathName(TEXT("/*"), TEXT("*"));
+}
+
 namespace UE::AssetRegistry::Impl
 {
 	/** The max time to spend in UAssetRegistryImpl::Tick */
@@ -1182,24 +1188,38 @@ void InitializeSerializationOptionsFromIni(FAssetRegistrySerializationOptions& O
 			FName TagFName = FName(*TagName);
 
 			// Include subclasses if the class is in memory at this time (native classes only)
-			UClass* FilterlistClass = Cast<UClass>(StaticFindObject(UClass::StaticClass(), ANY_PACKAGE, *ClassName));
+			UClass* FilterlistClass = Cast<UClass>(StaticFindObject(UClass::StaticClass(), nullptr, *ClassName));
 			if (FilterlistClass)
 			{
-				Options.CookFilterlistTagsByClass.FindOrAdd(FilterlistClass->GetFName()).Add(TagFName);
+				Options.CookFilterlistTagsByClass.FindOrAdd(FilterlistClass->GetClassPathName()).Add(TagFName);
 
 				TArray<UClass*> DerivedClasses;
 				GetDerivedClasses(FilterlistClass, DerivedClasses);
 				for (UClass* DerivedClass : DerivedClasses)
 				{
-					Options.CookFilterlistTagsByClass.FindOrAdd(DerivedClass->GetFName()).Add(TagFName);
+					Options.CookFilterlistTagsByClass.FindOrAdd(DerivedClass->GetClassPathName()).Add(TagFName);
 				}
 			}
 			else
 			{
+				FTopLevelAssetPath ClassPathName;
+				if (ClassName == TEXTVIEW("*"))
+				{
+					ClassPathName = UE::AssetRegistry::WildcardPathName;
+				}
+				else if (FPackageName::IsShortPackageName(ClassName))
+				{
+					ClassPathName = UClass::TryConvertShortTypeNameToPathName<UClass>(ClassName, ELogVerbosity::Warning, TEXT("Parsing [AssetRegistry] CookedTagsWhitelist or CookedTagsBlacklist"));
+					UE_CLOG(ClassPathName.IsNull(), LogAssetRegistry, Warning, TEXT("Failed to convert short class name \"%s\" when parsing ini [AssetRegistry] CookedTagsWhitelist or CookedTagsBlacklist"), *ClassName);
+				}
+				else
+				{
+					ClassPathName = FTopLevelAssetPath(ClassName);
+				}
 				// Class is not in memory yet. Just add an explicit filter.
 				// Automatically adding subclasses of non-native classes is not supported.
-				// In these cases, using Class=* is usually sufficient
-				Options.CookFilterlistTagsByClass.FindOrAdd(*ClassName).Add(TagFName);
+				// In these cases, using Class=* is usually sufficient				
+				Options.CookFilterlistTagsByClass.FindOrAdd(ClassPathName).Add(TagFName);
 			}
 		}
 	}
@@ -1217,20 +1237,21 @@ void FAssetRegistryImpl::CollectCodeGeneratorClasses()
 	ClassGeneratorNamesRegisteredClassesVersionNumber = GetRegisteredClassesVersionNumber();
 
 	// Work around the fact we don't reference Engine module directly
-	UClass* BlueprintCoreClass = Cast<UClass>(StaticFindObject(UClass::StaticClass(), ANY_PACKAGE, TEXT("BlueprintCore")));
+	FTopLevelAssetPath BlueprintCorePathName(TEXT("/Script/Engine"), TEXT("BlueprintCore"));
+	UClass* BlueprintCoreClass = FindObject<UClass>(BlueprintCorePathName);
 	if (!BlueprintCoreClass)
 	{
 		return;
 	}
 
-	ClassGeneratorNames.Add(BlueprintCoreClass->GetFName());
+	ClassGeneratorNames.Add(BlueprintCoreClass->GetClassPathName());
 
 	TArray<UClass*> BlueprintCoreDerivedClasses;
 	GetDerivedClasses(BlueprintCoreClass, BlueprintCoreDerivedClasses);
 	for (UClass* BPCoreClass : BlueprintCoreDerivedClasses)
 	{
 		bool bAlreadyRecorded;
-		FName BPCoreClassName = BPCoreClass->GetFName();
+		FTopLevelAssetPath BPCoreClassName = BPCoreClass->GetClassPathName();
 		ClassGeneratorNames.Add(BPCoreClassName, &bAlreadyRecorded);
 		if (bAlreadyRecorded)
 		{
@@ -1241,18 +1262,18 @@ void FAssetRegistryImpl::CollectCodeGeneratorClasses()
 		// when AddAssetData is called for those instances, but when we add a new generator class we have to recheck all
 		// instances of the class since they would have failed to detect they were Blueprint classes before.
 		// This can happen if blueprints in plugin B are scanned before their blueprint class from plugin A is scanned.
-		for (const FAssetData* AssetData : State.GetAssetsByClassName(BPCoreClassName))
+		for (const FAssetData* AssetData : State.GetAssetsByClassPathName(BPCoreClassName))
 		{
 			const FString GeneratedClass = AssetData->GetTagValueRef<FString>(FBlueprintTags::GeneratedClassPath);
 			const FString ParentClass = AssetData->GetTagValueRef<FString>(FBlueprintTags::ParentClassPath);
 			if (!GeneratedClass.IsEmpty() && !ParentClass.IsEmpty())
 			{
-				const FName GeneratedClassFName = *Utils::ExportTextPathToObjectName(GeneratedClass);
-				const FName ParentClassFName = *Utils::ExportTextPathToObjectName(ParentClass);
+				const FTopLevelAssetPath GeneratedClassPathName(FPackageName::ExportTextPathToObjectPath(GeneratedClass));
+				const FTopLevelAssetPath ParentClassPathName(FPackageName::ExportTextPathToObjectPath(ParentClass));
 
-				if (!CachedBPInheritanceMap.Contains(GeneratedClassFName))
+				if (!CachedBPInheritanceMap.Contains(GeneratedClassPathName))
 				{
-					CachedBPInheritanceMap.Add(GeneratedClassFName, ParentClassFName);
+					CachedBPInheritanceMap.Add(GeneratedClassPathName, ParentClassPathName);
 
 					// Invalidate caching because CachedBPInheritanceMap got modified
 					TempCachedInheritanceBuffer.bDirty = true;
@@ -1617,10 +1638,28 @@ void FAssetRegistryImpl::EnumerateAssetsByPathNoTags(FName PackagePath,
 
 }
 
+static FTopLevelAssetPath TryConvertShortTypeNameToPathName(FName ClassName)
+{
+	FTopLevelAssetPath ClassPathName;
+	if (ClassName != NAME_None)
+	{
+		FString ShortClassName = ClassName.ToString();
+		ClassPathName = UClass::TryConvertShortTypeNameToPathName<UStruct>(*ShortClassName, ELogVerbosity::Warning, TEXT("AssetRegistry using deprecated function"));
+		UE_CLOG(ClassPathName.IsNull(), LogClass, Error, TEXT("Failed to convert short class name %s to class path name."), *ShortClassName);
+	}
+	return ClassPathName;
+}
+
 bool UAssetRegistryImpl::GetAssetsByClass(FName ClassName, TArray<FAssetData>& OutAssetData, bool bSearchSubClasses) const
+{
+	FTopLevelAssetPath ClassPathName = TryConvertShortTypeNameToPathName(ClassName);
+	return GetAssetsByClass(ClassPathName, OutAssetData, bSearchSubClasses);
+}
+
+bool UAssetRegistryImpl::GetAssetsByClass(FTopLevelAssetPath ClassPathName, TArray<FAssetData>& OutAssetData, bool bSearchSubClasses) const
 {
 	FARFilter Filter;
-	Filter.ClassNames.Add(ClassName);
+	Filter.ClassPaths.Add(ClassPathName);
 	Filter.bRecursiveClasses = bSearchSubClasses;
 	return GetAssets(Filter, OutAssetData);
 }
@@ -1807,7 +1846,7 @@ void EnumerateMemoryAssetsHelper(const FARCompiledFilter& InFilter, TSet<FName>&
 
 			// Could perhaps save some FName -> String conversions by creating this a bit earlier using the UObject constructor
 			// to get package name and path.
-			FAssetData PartialAssetData(PackageNameStr, ObjectPathStr, Obj->GetClass()->GetFName(), FAssetDataTagMap(),
+			FAssetData PartialAssetData(PackageNameStr, ObjectPathStr, Obj->GetClass()->GetClassPathName(), FAssetDataTagMap(),
 				InMemoryPackage->GetChunkIDs(), InMemoryPackage->GetPackageFlags());
 
 			// All filters passed, except for AssetRegistry filter; caller must check that one
@@ -1816,12 +1855,12 @@ void EnumerateMemoryAssetsHelper(const FARCompiledFilter& InFilter, TSet<FName>&
 	};
 
 	// Iterate over all in-memory assets to find the ones that pass the filter components
-	if (InFilter.ClassNames.Num() > 0)
+	if (InFilter.ClassPaths.Num() > 0)
 	{
 		TArray<UObject*> InMemoryObjects;
-		for (FName ClassName : InFilter.ClassNames)
+		for (FTopLevelAssetPath ClassName : InFilter.ClassPaths)
 		{
-			UClass* Class = FindObjectFast<UClass>(nullptr, ClassName, false, true, RF_NoFlags);
+			UClass* Class = FindObject<UClass>(ClassName);
 			if (Class != nullptr)
 			{
 				GetObjectsOfClass(Class, InMemoryObjects, false, RF_NoFlags);
@@ -1928,7 +1967,7 @@ FAssetData UAssetRegistryImpl::GetAssetByObjectPath(const FName ObjectPath, bool
 		FReadScopeLock InterfaceScopeLock(InterfaceLock);
 		const FAssetRegistryState& State = GuardedData.GetState();
 		const FAssetData* FoundData = State.GetAssetByObjectPath(ObjectPath);
-		return (FoundData && !GuardedData.ShouldSkipAsset(FoundData->AssetClass, FoundData->PackageFlags)) ? *FoundData : FAssetData();
+		return (FoundData && !GuardedData.ShouldSkipAsset(FoundData->AssetClassPath, FoundData->PackageFlags)) ? *FoundData : FAssetData();
 	}
 }
 
@@ -2285,13 +2324,13 @@ bool UAssetRegistryImpl::DoesPackageExistOnDisk(FName PackageName, FString* OutC
 {
 	auto CalculateExtension = [](const FString& PackageNameStr, TConstArrayView<FAssetData> Assets) -> FString
 	{
-		FName ClassRedirector = UObjectRedirector::StaticClass()->GetFName();
+		FTopLevelAssetPath ClassRedirector = UObjectRedirector::StaticClass()->GetClassPathName();
 		bool bContainsMap = false;
 		bool bContainsRedirector = false;
 		for (const FAssetData& Asset : Assets)
 		{
 			bContainsMap |= ((Asset.PackageFlags & PKG_ContainsMap) != 0);
-			bContainsRedirector |= (Asset.AssetClass == ClassRedirector);
+			bContainsRedirector |= (Asset.AssetClassPath == ClassRedirector);
 		}
 		if (!bContainsMap && bContainsRedirector)
 		{
@@ -2460,6 +2499,18 @@ FName FAssetRegistryImpl::GetRedirectedObjectPath(const FName ObjectPath) const
 }
 
 bool UAssetRegistryImpl::GetAncestorClassNames(FName ClassName, TArray<FName>& OutAncestorClassNames) const
+{
+	FTopLevelAssetPath ClassPathName = TryConvertShortTypeNameToPathName(ClassName);
+	TArray<FTopLevelAssetPath> OutAncestorClassPathNames;
+	bool bResult = GetAncestorClassNames(ClassPathName, OutAncestorClassPathNames);
+	for (FTopLevelAssetPath AncestorPathName : OutAncestorClassPathNames)
+	{
+		OutAncestorClassNames.Add(AncestorPathName.GetAssetName());
+	}
+	return bResult;
+}
+
+bool UAssetRegistryImpl::GetAncestorClassNames(FTopLevelAssetPath ClassName, TArray<FTopLevelAssetPath>& OutAncestorClassNames) const
 {
 	UE::AssetRegistry::Impl::FClassInheritanceContext InheritanceContext;
 	UE::AssetRegistry::Impl::FClassInheritanceBuffer InheritanceBuffer;
@@ -2472,14 +2523,14 @@ bool UAssetRegistryImpl::GetAncestorClassNames(FName ClassName, TArray<FName>& O
 namespace UE::AssetRegistry
 {
 
-bool FAssetRegistryImpl::GetAncestorClassNames(Impl::FClassInheritanceContext& InheritanceContext, FName ClassName,
-	TArray<FName>& OutAncestorClassNames) const
+bool FAssetRegistryImpl::GetAncestorClassNames(Impl::FClassInheritanceContext& InheritanceContext, FTopLevelAssetPath ClassName,
+	TArray<FTopLevelAssetPath>& OutAncestorClassNames) const
 {
 	// Assume we found the class unless there is an error
 	bool bFoundClass = true;
 
 	InheritanceContext.ConditionalUpdate();
-	const TMap<FName, FName>& InheritanceMap = InheritanceContext.Buffer->InheritanceMap;
+	const TMap<FTopLevelAssetPath, FTopLevelAssetPath>& InheritanceMap = InheritanceContext.Buffer->InheritanceMap;
 
 	// Make sure the requested class is in the inheritance map
 	if (!InheritanceMap.Contains(ClassName))
@@ -2489,7 +2540,7 @@ bool FAssetRegistryImpl::GetAncestorClassNames(Impl::FClassInheritanceContext& I
 	else
 	{
 		// Now follow the map pairs until we cant find any more parents
-		const FName* CurrentClassName = &ClassName;
+		const FTopLevelAssetPath* CurrentClassName = &ClassName;
 		const uint32 MaxInheritanceDepth = 65536;
 		uint32 CurrentInheritanceDepth = 0;
 		while (CurrentInheritanceDepth < MaxInheritanceDepth && CurrentClassName != nullptr)
@@ -2498,7 +2549,7 @@ bool FAssetRegistryImpl::GetAncestorClassNames(Impl::FClassInheritanceContext& I
 
 			if (CurrentClassName)
 			{
-				if (*CurrentClassName == NAME_None)
+				if (CurrentClassName->IsNull())
 				{
 					// No parent, we are at the root
 					CurrentClassName = nullptr;
@@ -2525,6 +2576,27 @@ bool FAssetRegistryImpl::GetAncestorClassNames(Impl::FClassInheritanceContext& I
 
 void UAssetRegistryImpl::GetDerivedClassNames(const TArray<FName>& ClassNames, const TSet<FName>& ExcludedClassNames,
 	TSet<FName>& OutDerivedClassNames) const
+{
+	TArray<FTopLevelAssetPath> ClassPaths;
+	for (FName ClassName : ClassNames)
+	{
+		ClassPaths.Add(TryConvertShortTypeNameToPathName(ClassName));
+	}
+	TSet<FTopLevelAssetPath> ExcludedClassPathNames;
+	for (FName ExcludedClassName : ExcludedClassNames)
+	{
+		ExcludedClassPathNames.Add(TryConvertShortTypeNameToPathName(ExcludedClassName));
+	}
+	TSet<FTopLevelAssetPath> OutDerivedClassPathNames;
+	GetDerivedClassNames(ClassPaths, ExcludedClassPathNames, OutDerivedClassPathNames);
+	for (FTopLevelAssetPath DerivedClassPathName : OutDerivedClassPathNames)
+	{
+		OutDerivedClassNames.Add(DerivedClassPathName.GetAssetName());
+	}
+}
+
+void UAssetRegistryImpl::GetDerivedClassNames(const TArray<FTopLevelAssetPath>& ClassNames, const TSet<FTopLevelAssetPath>& ExcludedClassNames,
+	TSet<FTopLevelAssetPath>& OutDerivedClassNames) const
 {
 	UE::AssetRegistry::Impl::FClassInheritanceContext InheritanceContext;
 	UE::AssetRegistry::Impl::FClassInheritanceBuffer InheritanceBuffer;
@@ -2709,9 +2781,9 @@ bool RunAssetThroughFilter_Unchecked(const FAssetData& AssetData, const FARCompi
 	}
 
 	// Classes
-	if (Filter.ClassNames.Num() > 0)
+	if (Filter.ClassPaths.Num() > 0)
 	{
-		const bool bPassesClasses = Filter.ClassNames.Contains(AssetData.AssetClass);
+		const bool bPassesClasses = Filter.ClassPaths.Contains(AssetData.AssetClassPath);
 		if (bPassesClasses != bPassFilterValue)
 		{
 			return !bPassFilterValue;
@@ -2778,7 +2850,7 @@ void UAssetRegistryImpl::ExpandRecursiveFilter(const FARFilter& InFilter, FARFil
 	ExpandedFilter.PackageNames = CompiledFilter.PackageNames.Array();
 	ExpandedFilter.PackagePaths = CompiledFilter.PackagePaths.Array();
 	ExpandedFilter.ObjectPaths = CompiledFilter.ObjectPaths.Array();
-	ExpandedFilter.ClassNames = CompiledFilter.ClassNames.Array();
+	ExpandedFilter.ClassPaths = CompiledFilter.ClassPaths.Array();
 	ExpandedFilter.TagsAndValues = CompiledFilter.TagsAndValues;
 	ExpandedFilter.bIncludeOnlyOnDiskAssets = CompiledFilter.bIncludeOnlyOnDiskAssets;
 	ExpandedFilter.WithoutPackageFlags = CompiledFilter.WithoutPackageFlags;
@@ -2817,7 +2889,28 @@ void FAssetRegistryImpl::CompileFilter(Impl::FClassInheritanceContext& Inheritan
 		OutCompiledFilter.PackagePaths.Add(FPathTree::NormalizePackagePath(PackagePath));
 	}
 	OutCompiledFilter.ObjectPaths.Append(InFilter.ObjectPaths);
-	OutCompiledFilter.ClassNames.Append(InFilter.ClassNames);
+PRAGMA_DISABLE_DEPRECATION_WARNINGS
+	if (!ensureAlwaysMsgf(InFilter.ClassNames.Num() == 0, TEXT("Asset Registry Filter using ClassNames instead of ClassPaths. First class name: \"%s\""), *InFilter.ClassNames[0].ToString()))
+	{
+		OutCompiledFilter.ClassPaths.Reserve(InFilter.ClassNames.Num());
+		for (FName ClassName : InFilter.ClassNames)
+		{
+			if (!ClassName.IsNone())
+			{
+				FTopLevelAssetPath ClassPathName = UClass::TryConvertShortTypeNameToPathName<UStruct>(ClassName.ToString(), ELogVerbosity::Warning, TEXT("Compiling Asset Registry Filter"));
+				if (!ClassPathName.IsNull())
+				{
+					OutCompiledFilter.ClassPaths.Add(ClassPathName);
+				}
+				else
+				{
+					UE_LOG(LogAssetRegistry, Error, TEXT("Failed to resolve class path for short clas name \"%s\" when compiling asset registry filter"), *ClassName.ToString());
+				}
+			}
+		}
+	}
+PRAGMA_ENABLE_DEPRECATION_WARNINGS
+	OutCompiledFilter.ClassPaths.Append(InFilter.ClassPaths);
 	OutCompiledFilter.TagsAndValues = InFilter.TagsAndValues;
 	OutCompiledFilter.bIncludeOnlyOnDiskAssets = InFilter.bIncludeOnlyOnDiskAssets;
 	OutCompiledFilter.WithoutPackageFlags = InFilter.WithoutPackageFlags;
@@ -2835,16 +2928,16 @@ void FAssetRegistryImpl::CompileFilter(Impl::FClassInheritanceContext& Inheritan
 	if (InFilter.bRecursiveClasses)
 	{
 		// Add the sub-classes of all the input classes to the expanded list, excluding any that were requested
-		if (InFilter.RecursiveClassesExclusionSet.Num() > 0 && InFilter.ClassNames.Num() == 0)
+		if (InFilter.RecursiveClassPathsExclusionSet.Num() > 0 && InFilter.ClassPaths.Num() == 0)
 		{
-			TArray<FName> ClassNamesObject;
-			ClassNamesObject.Add(UObject::StaticClass()->GetFName());
+			TArray<FTopLevelAssetPath> ClassNamesObject;
+			ClassNamesObject.Add(UObject::StaticClass()->GetClassPathName());
 
-			GetSubClasses(InheritanceContext, ClassNamesObject, InFilter.RecursiveClassesExclusionSet, OutCompiledFilter.ClassNames);
+			GetSubClasses(InheritanceContext, ClassNamesObject, InFilter.RecursiveClassPathsExclusionSet, OutCompiledFilter.ClassPaths);
 		}
 		else
 		{
-			GetSubClasses(InheritanceContext, InFilter.ClassNames, InFilter.RecursiveClassesExclusionSet, OutCompiledFilter.ClassNames);
+			GetSubClasses(InheritanceContext, InFilter.ClassPaths, InFilter.RecursiveClassPathsExclusionSet, OutCompiledFilter.ClassPaths);
 		}
 	}
 }
@@ -3802,14 +3895,14 @@ void FAssetRegistryImpl::CachePathsFromState(Impl::FEventContext& EventContext,
 			AddAssetPath(EventContext, AssetData->PackagePath);
 
 			// Populate the class map if adding blueprint
-			if (ClassGeneratorNames.Contains(AssetData->AssetClass))
+			if (ClassGeneratorNames.Contains(AssetData->AssetClassPath))
 			{
 				FAssetRegistryExportPath GeneratedClass = AssetData->GetTagValueRef<FAssetRegistryExportPath>(FBlueprintTags::GeneratedClassPath);
 				FAssetRegistryExportPath ParentClass = AssetData->GetTagValueRef<FAssetRegistryExportPath>(FBlueprintTags::ParentClassPath);
 
 				if (GeneratedClass && ParentClass)
 				{
-					CachedBPInheritanceMap.Add(GeneratedClass.Object, ParentClass.Object);
+					CachedBPInheritanceMap.Add(GeneratedClass.ToTopLevelAssetPath(), ParentClass.ToTopLevelAssetPath());
 
 					// Invalidate caching because CachedBPInheritanceMap got modified
 					TempCachedInheritanceBuffer.bDirty = true;
@@ -3868,7 +3961,7 @@ void FAssetRegistryImpl::GetAllocatedSize(bool bLogDetailed, SIZE_T& StateSize,
 	}
 
 	StaticSize += SerializationOptions.CookFilterlistTagsByClass.GetAllocatedSize();
-	for (const TPair<FName, TSet<FName>>& Pair : SerializationOptions.CookFilterlistTagsByClass)
+	for (const TPair<FTopLevelAssetPath, TSet<FName>>& Pair : SerializationOptions.CookFilterlistTagsByClass)
 	{
 		StaticSize += Pair.Value.GetAllocatedSize();
 	}
@@ -4184,6 +4277,51 @@ bool IsPathMounted(const FString& Path, const TSet<FString>& MountPointsNoTraili
 
 }
 
+#if WITH_EDITOR
+void FAssetRegistryImpl::PostLoadAssetRegistryTags(FAssetData* AssetData)
+{
+	check(AssetData);
+	if (AssetData->TagsAndValues.Num())
+	{
+		FTopLevelAssetPath AssetClassPath = AssetData->AssetClassPath;
+		UClass* AssetClass = FindObject<UClass>(AssetClassPath, true);
+		while (!AssetClass)
+		{
+			// this is probably a blueprint that has not yet been loaded, try to find its native base class
+			const FTopLevelAssetPath* ParentClassPath = CachedBPInheritanceMap.Find(AssetClassPath);
+			if (ParentClassPath && !ParentClassPath->IsNull())
+			{
+				AssetClassPath = *ParentClassPath;
+				AssetClass = FindObject<UClass>(AssetClassPath, true);
+			}
+			else
+			{
+				break;
+			}
+		}
+
+		if (AssetClass)
+		{
+			if (UObject* ClassCDO = AssetClass->GetDefaultObject(false))
+			{
+				TArray<UObject::FAssetRegistryTag> TagsToModify;
+				ClassCDO->PostLoadAssetRegistryTags(*AssetData, TagsToModify);
+				if (TagsToModify.Num())
+				{
+					FAssetDataTagMap TagsAndValues = AssetData->TagsAndValues.CopyMap();
+					for (const UObject::FAssetRegistryTag& Tag : TagsToModify)
+					{
+						TagsAndValues.Remove(Tag.Name);
+						TagsAndValues.Add(Tag.Name, Tag.Value);
+					}
+					AssetData->TagsAndValues = FAssetDataTagMapSharedView(MoveTemp(TagsAndValues));
+				}
+			}
+		}
+	}
+}
+#endif
+
 void FAssetRegistryImpl::AssetSearchDataGathered(Impl::FEventContext& EventContext, const double TickStartTime, TRingBuffer<FAssetData*>& AssetResults)
 {
 	const bool bFlushFullBuffer = TickStartTime < 0;
@@ -4234,6 +4372,9 @@ void FAssetRegistryImpl::AssetSearchDataGathered(Impl::FEventContext& EventConte
 				if (!AssetDataObjectPathsUpdatedOnLoad.Contains(BackgroundResult->ObjectPath))
 #endif
 				{
+#if WITH_EDITOR
+					PostLoadAssetRegistryTags(BackgroundResult.Get());
+#endif
 					// The asset exists in the cache from disk and has not yet been loaded into memory, update it with the new background data
 					UpdateAssetData(EventContext, AssetData, *BackgroundResult);
 				}
@@ -4244,6 +4385,9 @@ void FAssetRegistryImpl::AssetSearchDataGathered(Impl::FEventContext& EventConte
 			// The asset isn't in the cache yet, add it and notify subscribers
 			if (bPathIsMounted)
 			{
+#if WITH_EDITOR
+				PostLoadAssetRegistryTags(BackgroundResult.Get());
+#endif
 				AddAssetData(EventContext, BackgroundResult.Release());
 			}
 		}
@@ -4507,36 +4651,26 @@ bool FAssetRegistryImpl::RemoveAssetPath(Impl::FEventContext& EventContext, FNam
 	return true;
 }
 
-namespace Utils
-{
-
-FString ExportTextPathToObjectName(const FString& InExportTextPath)
-{
-	const FString ObjectPath = FPackageName::ExportTextPathToObjectPath(InExportTextPath);
-	return FPackageName::ObjectPathToObjectName(ObjectPath);
-}
-
-}
-
 void FAssetRegistryImpl::AddAssetData(Impl::FEventContext& EventContext, FAssetData* AssetData)
 {
 	State.AddAssetData(AssetData);
 
-	if (!ShouldSkipAsset(AssetData->AssetClass, AssetData->PackageFlags))
+	if (!ShouldSkipAsset(AssetData->AssetClassPath, AssetData->PackageFlags))
 	{
 		EventContext.AssetEvents.Emplace(*AssetData, Impl::FEventContext::EEvent::Added);
 	}
 
 	// Populate the class map if adding blueprint
-	if (ClassGeneratorNames.Contains(AssetData->AssetClass))
+	if (ClassGeneratorNames.Contains(AssetData->AssetClassPath))
 	{
 		const FString GeneratedClass = AssetData->GetTagValueRef<FString>(FBlueprintTags::GeneratedClassPath);
 		const FString ParentClass = AssetData->GetTagValueRef<FString>(FBlueprintTags::ParentClassPath);
 		if (!GeneratedClass.IsEmpty() && !ParentClass.IsEmpty())
 		{
-			const FName GeneratedClassFName = *Utils::ExportTextPathToObjectName(GeneratedClass);
-			const FName ParentClassFName = *Utils::ExportTextPathToObjectName(ParentClass);
-			CachedBPInheritanceMap.Add(GeneratedClassFName, ParentClassFName);
+			const FTopLevelAssetPath GeneratedClassPathName(GeneratedClass);
+			const FTopLevelAssetPath ParentClassPathName(ParentClass);
+			checkf(!GeneratedClassPathName.IsNull() && !ParentClassPathName.IsNull(), TEXT("Short class names used in AddAssetData: GeneratedClass=%s, ParentClass=%s. Short class names in these tags on the Blueprint class should have been converted to path names."), *GeneratedClass, *ParentClass);
+			CachedBPInheritanceMap.Add(GeneratedClassPathName, ParentClassPathName);
 			
 			// Invalidate caching because CachedBPInheritanceMap got modified
 			TempCachedInheritanceBuffer.bDirty = true;
@@ -4547,13 +4681,14 @@ void FAssetRegistryImpl::AddAssetData(Impl::FEventContext& EventContext, FAssetD
 void FAssetRegistryImpl::UpdateAssetData(Impl::FEventContext& EventContext, FAssetData* AssetData, const FAssetData& NewAssetData)
 {
 	// Update the class map if updating a blueprint
-	if (ClassGeneratorNames.Contains(AssetData->AssetClass))
+	if (ClassGeneratorNames.Contains(AssetData->AssetClassPath))
 	{
 		const FString OldGeneratedClass = AssetData->GetTagValueRef<FString>(FBlueprintTags::GeneratedClassPath);
 		if (!OldGeneratedClass.IsEmpty())
 		{
-			const FName OldGeneratedClassFName = *Utils::ExportTextPathToObjectName(OldGeneratedClass);
-			CachedBPInheritanceMap.Remove(OldGeneratedClassFName);
+			const FTopLevelAssetPath OldGeneratedClassName(OldGeneratedClass);
+			checkf(!OldGeneratedClassName.IsNull(), TEXT("Short class name used: OldGeneratedClass=%s. Short class names in tags on the Blueprint class should have been converted to path names."), *OldGeneratedClass);
+			CachedBPInheritanceMap.Remove(OldGeneratedClassName);
 
 			// Invalidate caching because CachedBPInheritanceMap got modified
 			TempCachedInheritanceBuffer.bDirty = true;
@@ -4563,9 +4698,10 @@ void FAssetRegistryImpl::UpdateAssetData(Impl::FEventContext& EventContext, FAss
 		const FString NewParentClass = NewAssetData.GetTagValueRef<FString>(FBlueprintTags::ParentClassPath);
 		if (!NewGeneratedClass.IsEmpty() && !NewParentClass.IsEmpty())
 		{
-			const FName NewGeneratedClassFName = *Utils::ExportTextPathToObjectName(*NewGeneratedClass);
-			const FName NewParentClassFName = *Utils::ExportTextPathToObjectName(*NewParentClass);
-			CachedBPInheritanceMap.Add(NewGeneratedClassFName, NewParentClassFName);
+			const FTopLevelAssetPath NewGeneratedClassName(NewGeneratedClass);
+			const FTopLevelAssetPath NewParentClassName(NewParentClass);
+			checkf(!NewGeneratedClassName.IsNull() && !NewParentClassName.IsNull(), TEXT("Short class names used in AddAssetData: GeneratedClass=%s, ParentClass=%s. Short class names in these tags on the Blueprint class should have been converted to path names."), *NewGeneratedClass, *NewParentClass);
+			CachedBPInheritanceMap.Add(NewGeneratedClassName, NewParentClassName);
 
 			// Invalidate caching because CachedBPInheritanceMap got modified
 			TempCachedInheritanceBuffer.bDirty = true;
@@ -4574,7 +4710,7 @@ void FAssetRegistryImpl::UpdateAssetData(Impl::FEventContext& EventContext, FAss
 
 	State.UpdateAssetData(AssetData, NewAssetData);
 	
-	if (!ShouldSkipAsset(AssetData->AssetClass, AssetData->PackageFlags))
+	if (!ShouldSkipAsset(AssetData->AssetClassPath, AssetData->PackageFlags))
 	{
 		EventContext.AssetEvents.Emplace(*AssetData, Impl::FEventContext::EEvent::Updated);
 	}
@@ -4586,19 +4722,20 @@ bool FAssetRegistryImpl::RemoveAssetData(Impl::FEventContext& EventContext, FAss
 
 	if (ensure(AssetData))
 	{
-		if (!ShouldSkipAsset(AssetData->AssetClass, AssetData->PackageFlags))
+		if (!ShouldSkipAsset(AssetData->AssetClassPath, AssetData->PackageFlags))
 		{
 			EventContext.AssetEvents.Emplace(*AssetData, Impl::FEventContext::EEvent::Removed);
 		}
 
 		// Remove from the class map if removing a blueprint
-		if (ClassGeneratorNames.Contains(AssetData->AssetClass))
+		if (ClassGeneratorNames.Contains(AssetData->AssetClassPath))
 		{
 			const FString OldGeneratedClass = AssetData->GetTagValueRef<FString>(FBlueprintTags::GeneratedClassPath);
 			if (!OldGeneratedClass.IsEmpty())
 			{
-				const FName OldGeneratedClassFName = *Utils::ExportTextPathToObjectName(OldGeneratedClass);
-				CachedBPInheritanceMap.Remove(OldGeneratedClassFName);
+				const FTopLevelAssetPath OldGeneratedClassPathName(FPackageName::ExportTextPathToObjectPath(OldGeneratedClass));
+				checkf(!OldGeneratedClassPathName.IsNull(), TEXT("Short class name used: OldGeneratedClass=%s"), *OldGeneratedClass);
+				CachedBPInheritanceMap.Remove(OldGeneratedClassPathName);
 
 				// Invalidate caching because CachedBPInheritanceMap got modified
 				TempCachedInheritanceBuffer.bDirty = true;
@@ -4990,7 +5127,7 @@ void FAssetRegistryImpl::PushProcessLoadedAssetsBatch(Impl::FEventContext& Event
 void FAssetRegistryImpl::UpdateRedirectCollector()
 {
 	// Look for all redirectors in list
-	const TArray<const FAssetData*>& RedirectorAssets = State.GetAssetsByClassName(UObjectRedirector::StaticClass()->GetFName());
+	const TArray<const FAssetData*>& RedirectorAssets = State.GetAssetsByClassPathName(UObjectRedirector::StaticClass()->GetClassPathName());
 
 	for (const FAssetData* AssetData : RedirectorAssets)
 	{
@@ -5320,9 +5457,9 @@ void FAssetRegistryImpl::UpdateInheritanceBuffer(Impl::FClassInheritanceBuffer&
 	}
 	OutBuffer.InheritanceMap.Reserve(NumNativeClasses + CachedBPInheritanceMap.Num());
 	OutBuffer.InheritanceMap = CachedBPInheritanceMap;
-	OutBuffer.InheritanceMap.Add(FName("Object"), FName());
+	OutBuffer.InheritanceMap.Add(FTopLevelAssetPath(TEXT("/Script.CoreUObject"), TEXT("Object")), FTopLevelAssetPath());
 
-	for (TPair<FName, TArray<FName>>& Pair : OutBuffer.ReverseInheritanceMap)
+	for (TPair<FTopLevelAssetPath, TArray<FTopLevelAssetPath>>& Pair : OutBuffer.ReverseInheritanceMap)
 	{
 		Pair.Value.Reset();
 	}
@@ -5330,15 +5467,15 @@ void FAssetRegistryImpl::UpdateInheritanceBuffer(Impl::FClassInheritanceBuffer&
 
 	for (const TPair<UClass*, TSet<UClass*>>& Pair : NativeSubclasses)
 	{
-		FName SuperclassName = Pair.Key->GetFName();
+		FTopLevelAssetPath SuperclassName = Pair.Key->GetClassPathName();
 
-		TArray<FName>* OutputSubclasses = &OutBuffer.ReverseInheritanceMap.FindOrAdd(SuperclassName);
+		TArray<FTopLevelAssetPath>* OutputSubclasses = &OutBuffer.ReverseInheritanceMap.FindOrAdd(SuperclassName);
 		OutputSubclasses->Reserve(Pair.Value.Num());
 		for (UClass* Subclass : Pair.Value)
 		{
 			if (!Subclass->HasAnyClassFlags(CLASS_Deprecated | CLASS_NewerVersionExists))
 			{
-				FName SubclassName = Subclass->GetFName();
+				FTopLevelAssetPath SubclassName = Subclass->GetClassPathName();
 				OutputSubclasses->Add(SubclassName);
 				OutBuffer.InheritanceMap.Add(SubclassName, SuperclassName);
 
@@ -5349,7 +5486,7 @@ void FAssetRegistryImpl::UpdateInheritanceBuffer(Impl::FClassInheritanceBuffer&
 					{
 						if (UClass* InterfaceClass = Interface.Class) // could be nulled out by ForceDelete of a blueprint interface
 						{
-							TArray<FName>& Implementations = OutBuffer.ReverseInheritanceMap.FindOrAdd(InterfaceClass->GetFName());
+							TArray<FTopLevelAssetPath>& Implementations = OutBuffer.ReverseInheritanceMap.FindOrAdd(InterfaceClass->GetClassPathName());
 							Implementations.Add(SubclassName);
 						}
 					}
@@ -5363,12 +5500,12 @@ void FAssetRegistryImpl::UpdateInheritanceBuffer(Impl::FClassInheritanceBuffer&
 	}
 
 	// Add non-native classes to reverse map
-	for (const TPair<FName, FName>& Kvp : CachedBPInheritanceMap)
+	for (const TPair<FTopLevelAssetPath, FTopLevelAssetPath>& Kvp : CachedBPInheritanceMap)
 	{
-		FName ParentClassName = Kvp.Value;
-		if (!ParentClassName.IsNone())
+		const FTopLevelAssetPath& ParentClassName = Kvp.Value;
+		if (!ParentClassName.IsNull())
 		{
-			TArray<FName>& ChildClasses = OutBuffer.ReverseInheritanceMap.FindOrAdd(ParentClassName);
+			TArray<FTopLevelAssetPath>& ChildClasses = OutBuffer.ReverseInheritanceMap.FindOrAdd(ParentClassName);
 			ChildClasses.Add(Kvp.Key);
 		}
 
@@ -5497,20 +5634,20 @@ void FClassInheritanceContext::ConditionalUpdate()
 }
 
 void FAssetRegistryImpl::GetSubClasses(Impl::FClassInheritanceContext& InheritanceContext,
-	const TArray<FName>& InClassNames, const TSet<FName>& ExcludedClassNames, TSet<FName>& SubClassNames) const
+	const TArray<FTopLevelAssetPath>& InClassNames, const TSet<FTopLevelAssetPath>& ExcludedClassNames, TSet<FTopLevelAssetPath>& SubClassNames) const
 {
 	InheritanceContext.ConditionalUpdate();
 
-	TSet<FName> ProcessedClassNames;
-	for (FName ClassName : InClassNames)
+	TSet<FTopLevelAssetPath> ProcessedClassNames;
+	for (const FTopLevelAssetPath& ClassName : InClassNames)
 	{
 		// Now find all subclass names
 		GetSubClasses_Recursive(InheritanceContext, ClassName, SubClassNames, ProcessedClassNames, ExcludedClassNames);
 	}
 }
 
-void FAssetRegistryImpl::GetSubClasses_Recursive(Impl::FClassInheritanceContext& InheritanceContext, FName InClassName,
-	TSet<FName>& SubClassNames, TSet<FName>& ProcessedClassNames, const TSet<FName>& ExcludedClassNames) const
+void FAssetRegistryImpl::GetSubClasses_Recursive(Impl::FClassInheritanceContext& InheritanceContext, FTopLevelAssetPath InClassName,
+	TSet<FTopLevelAssetPath>& SubClassNames, TSet<FTopLevelAssetPath>& ProcessedClassNames, const TSet<FTopLevelAssetPath>& ExcludedClassNames) const
 {
 	if (ExcludedClassNames.Contains(InClassName))
 	{
@@ -5525,10 +5662,10 @@ void FAssetRegistryImpl::GetSubClasses_Recursive(Impl::FClassInheritanceContext&
 		SubClassNames.Add(InClassName);
 		ProcessedClassNames.Add(InClassName);
 
-		const TArray<FName>* FoundSubClassNames = InheritanceContext.Buffer->ReverseInheritanceMap.Find(InClassName);
+		const TArray<FTopLevelAssetPath>* FoundSubClassNames = InheritanceContext.Buffer->ReverseInheritanceMap.Find(InClassName);
 		if (FoundSubClassNames)
 		{
-			for (FName ClassName : (*FoundSubClassNames))
+			for (FTopLevelAssetPath ClassName : (*FoundSubClassNames))
 			{
 				GetSubClasses_Recursive(InheritanceContext, ClassName, SubClassNames, ProcessedClassNames,
 					ExcludedClassNames);
@@ -5791,7 +5928,7 @@ bool FAssetRegistryImpl::SetPrimaryAssetIdForObjectPath(Impl::FEventContext& Eve
 	TagsAndValues.Add(FPrimaryAssetId::PrimaryAssetTypeTag, PrimaryAssetId.PrimaryAssetType.ToString());
 	TagsAndValues.Add(FPrimaryAssetId::PrimaryAssetNameTag, PrimaryAssetId.PrimaryAssetName.ToString());
 
-	FAssetData NewAssetData = FAssetData(AssetData->PackageName, AssetData->PackagePath, AssetData->AssetName, AssetData->AssetClass, TagsAndValues, AssetData->ChunkIDs, AssetData->PackageFlags);
+	FAssetData NewAssetData = FAssetData(AssetData->PackageName, AssetData->PackagePath, AssetData->AssetName, AssetData->AssetClassPath, TagsAndValues, AssetData->ChunkIDs, AssetData->PackageFlags);
 	NewAssetData.TaggedAssetBundles = AssetData->TaggedAssetBundles;
 	UpdateAssetData(EventContext, AssetData, NewAssetData);
 
@@ -5868,7 +6005,7 @@ const TSet<FName>& FAssetRegistryImpl::GetCachedEmptyPackages() const
 	return CachedEmptyPackages;
 }
 
-bool FAssetRegistryImpl::ShouldSkipAsset(FName AssetClass, uint32 PackageFlags) const
+bool FAssetRegistryImpl::ShouldSkipAsset(FTopLevelAssetPath AssetClass, uint32 PackageFlags) const
 {
 #if WITH_ENGINE && WITH_EDITOR
 	return Utils::ShouldSkipAsset(AssetClass, PackageFlags, SkipUncookedClasses, SkipCookedClasses);
@@ -6104,7 +6241,7 @@ const FAssetData* GetMostImportantAsset(TConstArrayView<const FAssetData*> Packa
 	// both determined by class then name:
 	auto AssetDataLessThan = [](const FAssetData* LHS, const FAssetData* RHS)
 	{
-		int32 ClassCompare = LHS->AssetClass.Compare(RHS->AssetClass);
+		int32 ClassCompare = LHS->AssetClassPath.Compare(RHS->AssetClassPath);
 		if (ClassCompare == 0)
 		{
 			return LHS->AssetName.LexicalLess(RHS->AssetName);
@@ -6125,7 +6262,7 @@ const FAssetData* GetMostImportantAsset(TConstArrayView<const FAssetData*> Packa
 			return Asset;
 		}
 		// This is after IsUAsset because Blueprints can be the UAsset but also be considered skipable.
-		if (FFiltering::ShouldSkipAsset(Asset->AssetClass, Asset->PackageFlags))
+		if (FFiltering::ShouldSkipAsset(Asset->AssetClassPath, Asset->PackageFlags))
 		{
 			continue;
 		}
diff --git a/Engine/Source/Runtime/AssetRegistry/Private/AssetRegistry.h b/Engine/Source/Runtime/AssetRegistry/Private/AssetRegistry.h
--- a/Engine/Source/Runtime/AssetRegistry/Private/AssetRegistry.h
+++ b/Engine/Source/Runtime/AssetRegistry/Private/AssetRegistry.h
@@ -33,7 +33,9 @@ public:
 	virtual bool GetAssetsByPackageName(FName PackageName, TArray<FAssetData>& OutAssetData, bool bIncludeOnlyOnDiskAssets = false) const override;
 	virtual bool GetAssetsByPath(FName PackagePath, TArray<FAssetData>& OutAssetData, bool bRecursive = false, bool bIncludeOnlyOnDiskAssets = false) const override;
 	virtual bool GetAssetsByPaths(TArray<FName> PackagePath, TArray<FAssetData>& OutAssetData, bool bRecursive = false, bool bIncludeOnlyOnDiskAssets = false) const override;
+	UE_DEPRECATED(5.1, "Class names are now represented by path names. Please use a version of this function that uses FTopLevelAssetPath.")
 	virtual bool GetAssetsByClass(FName ClassName, TArray<FAssetData>& OutAssetData, bool bSearchSubClasses = false) const override;
+	virtual bool GetAssetsByClass(FTopLevelAssetPath ClassPathName, TArray<FAssetData>& OutAssetData, bool bSearchSubClasses = false) const override;
 	virtual bool GetAssetsByTags(const TArray<FName>& AssetTags, TArray<FAssetData>& OutAssetData) const override;
 	virtual bool GetAssetsByTagValues(const TMultiMap<FName, FString>& AssetTagsAndValues, TArray<FAssetData>& OutAssetData) const override;
 	virtual bool GetAssets(const FARFilter& Filter, TArray<FAssetData>& OutAssetData) const override;
@@ -63,8 +65,12 @@ public:
 	virtual void EnumerateAllPackages(TFunctionRef<void(FName PackageName, const FAssetPackageData& PackageData)> Callback) const override;
 	virtual bool DoesPackageExistOnDisk(FName PackageName, FString* OutCorrectCasePackageName = nullptr, FString* OutExtension = nullptr) const override;
 	virtual FName GetRedirectedObjectPath(const FName ObjectPath) const override;
+	UE_DEPRECATED(5.1, "Class names are now represented by path names. Please use a version of this function that uses FTopLevelAssetPath.")
 	virtual bool GetAncestorClassNames(FName ClassName, TArray<FName>& OutAncestorClassNames) const override;
+	virtual bool GetAncestorClassNames(FTopLevelAssetPath ClassName, TArray<FTopLevelAssetPath>& OutAncestorClassNames) const override;
+	UE_DEPRECATED(5.1, "Class names are now represented by path names. Please use a version of this function that uses FTopLevelAssetPath.")
 	virtual void GetDerivedClassNames(const TArray<FName>& ClassNames, const TSet<FName>& ExcludedClassNames, TSet<FName>& OutDerivedClassNames) const override;
+	virtual void GetDerivedClassNames(const TArray<FTopLevelAssetPath>& ClassNames, const TSet<FTopLevelAssetPath>& ExcludedClassNames, TSet<FTopLevelAssetPath>& OutDerivedClassNames) const override;	
 	virtual void GetAllCachedPaths(TArray<FString>& OutPathList) const override;
 	virtual void EnumerateAllCachedPaths(TFunctionRef<bool(FString)> Callback) const override;
 	virtual void EnumerateAllCachedPaths(TFunctionRef<bool(FName)> Callback) const override;
diff --git a/Engine/Source/Runtime/AssetRegistry/Private/AssetRegistryArchive.cpp b/Engine/Source/Runtime/AssetRegistry/Private/AssetRegistryArchive.cpp
--- a/Engine/Source/Runtime/AssetRegistry/Private/AssetRegistryArchive.cpp
+++ b/Engine/Source/Runtime/AssetRegistry/Private/AssetRegistryArchive.cpp
@@ -155,7 +155,7 @@ void FAssetRegistryWriter::SerializeTagsAndBundles(const FAssetData& Out)
 
 #endif
 
-FAssetRegistryReader::FAssetRegistryReader(FArchive& Inner, int32 NumWorkers)
+FAssetRegistryReader::FAssetRegistryReader(FArchive& Inner, int32 NumWorkers, FAssetRegistryVersion::Type Version)
 	: FArchiveProxy(Inner)
 {
 	check(IsLoading());
@@ -173,7 +173,7 @@ FAssetRegistryReader::FAssetRegistryReader(FArchive& Inner, int32 NumWorkers)
 	else
 	{
 		Names = LoadNameBatch(Inner);
-		Tags = FixedTagPrivate::LoadStore(*this);
+		Tags = FixedTagPrivate::LoadStore(*this, Version);
 	}
 }
 
@@ -254,16 +254,16 @@ bool FAssetRegistryTagSerializationTest::RunTest(const FString& Parameters)
 								{"Key_0",		"StringValue_0"}}));
 	LooseMaps.Add(MakeLooseMap({{"Name",		"NameValue"}, 
 								{"Name_0",		"NameValue_0"}}));
-	LooseMaps.Add(MakeLooseMap({{"FullPath",	"C\'P.O\'"}, 
+	LooseMaps.Add(MakeLooseMap({{"FullPath",	"/S/P.C\'P.O\'"}, 
 								{"PkgPath",		"P.O"},
 								{"ObjPath",		"O"}}));
-	LooseMaps.Add(MakeLooseMap({{"NumPath_0",	"C\'P.O_0\'"}, 
-								{"NumPath_1",	"C\'P_0.O\'"},
-								{"NumPath_2",	"C_0\'P.O\'"},
-								{"NumPath_3",	"C\'P_0.O_0\'"},
-								{"NumPath_4",	"C_0\'P_0.O\'"},
-								{"NumPath_5",	"C_0\'P.O_0\'"},
-								{"NumPath_6",	"C_0\'P_0.O_0\'"}}));
+	LooseMaps.Add(MakeLooseMap({{"NumPath_0",	"/S/P.C\'P.O_0\'"}, 
+								{"NumPath_1",	"/S/P.C\'P_0.O\'"},
+								{"NumPath_2",	"/S/P.C_0\'P.O\'"},
+								{"NumPath_3",	"/S/P.C\'P_0.O_0\'"},
+								{"NumPath_4",	"/S/P.C_0\'P_0.O\'"},
+								{"NumPath_5",	"/S/P.C_0\'P.O_0\'"},
+								{"NumPath_6",	"/S/P.C_0\'P_0.O_0\'"}}));
 	LooseMaps.Add(MakeLooseMap({{"SameSame",	"SameSame"}, 
 								{"AlsoSame",	"SameSame"}}));
 	LooseMaps.Add(MakeLooseMap({{"FilterKey1",	"FilterValue1"}, 
diff --git a/Engine/Source/Runtime/AssetRegistry/Private/AssetRegistryArchive.h b/Engine/Source/Runtime/AssetRegistry/Private/AssetRegistryArchive.h
--- a/Engine/Source/Runtime/AssetRegistry/Private/AssetRegistryArchive.h
+++ b/Engine/Source/Runtime/AssetRegistry/Private/AssetRegistryArchive.h
@@ -15,7 +15,7 @@ class FAssetRegistryReader : public FArchiveProxy
 {
 public:
 	/// @param NumWorkers > 0 for parallel loading
-	FAssetRegistryReader(FArchive& Inner, int32 NumWorkers = 0);
+	FAssetRegistryReader(FArchive& Inner, int32 NumWorkers = 0, FAssetRegistryVersion::Type Version = FAssetRegistryVersion::LatestVersion);
 	~FAssetRegistryReader();
 
 	virtual FArchive& operator<<(FName& Value) override;
diff --git a/Engine/Source/Runtime/AssetRegistry/Private/AssetRegistryConsoleCommands.h b/Engine/Source/Runtime/AssetRegistry/Private/AssetRegistryConsoleCommands.h
--- a/Engine/Source/Runtime/AssetRegistry/Private/AssetRegistryConsoleCommands.h
+++ b/Engine/Source/Runtime/AssetRegistry/Private/AssetRegistryConsoleCommands.h
@@ -100,17 +100,24 @@ public:
 	{
 		if ( Args.Num() < 1 )
 		{
-			UE_LOG(LogAssetRegistry, Log, TEXT("Usage: AssetRegistry.GetByClass Classname"));
+			UE_LOG(LogAssetRegistry, Log, TEXT("Usage: AssetRegistry.GetByClass ClassPathname"));
 			return;
 		}
 
 		TArray<FAssetData> AssetData;
-		const FString Classname = Args[0];
-		IAssetRegistry::GetChecked().GetAssetsByClass(FName(*Classname), AssetData);
-		UE_LOG(LogAssetRegistry, Log, TEXT("GetAssetsByClass for %s:"), *Classname);
-		for (int32 AssetIdx = 0; AssetIdx < AssetData.Num(); ++AssetIdx)
+		const FTopLevelAssetPath ClassPathName(Args[0]);
+		if (!ClassPathName.IsNull())
 		{
-			AssetData[AssetIdx].PrintAssetData();
+			IAssetRegistry::GetChecked().GetAssetsByClass(ClassPathName, AssetData);
+			UE_LOG(LogAssetRegistry, Log, TEXT("GetAssetsByClass for %s:"), *ClassPathName.ToString());
+			for (int32 AssetIdx = 0; AssetIdx < AssetData.Num(); ++AssetIdx)
+			{
+				AssetData[AssetIdx].PrintAssetData();
+			}
+		}
+		else
+		{
+			UE_LOG(LogAssetRegistry, Error, TEXT("\"%s\" is not a valid class path name (E.g. /Script/Engine.Actor)"));
 		}
 	}
 
diff --git a/Engine/Source/Runtime/AssetRegistry/Private/AssetRegistryImpl.h b/Engine/Source/Runtime/AssetRegistry/Private/AssetRegistryImpl.h
--- a/Engine/Source/Runtime/AssetRegistry/Private/AssetRegistryImpl.h
+++ b/Engine/Source/Runtime/AssetRegistry/Private/AssetRegistryImpl.h
@@ -83,9 +83,9 @@ namespace Impl
 	struct FClassInheritanceBuffer
 	{
 		/** Map from Class->SuperClass for all classes including native classes and blueprint classes. Updated on demand. */
-		TMap<FName, FName> InheritanceMap;
+		TMap<FTopLevelAssetPath, FTopLevelAssetPath> InheritanceMap;
 		/** Map from Class->(All subclasses) for all classes including native classes and blueprint classes. Updated on demand. */
-		TMap<FName, TArray<FName>> ReverseInheritanceMap;
+		TMap<FTopLevelAssetPath, TArray<FTopLevelAssetPath>> ReverseInheritanceMap;
 		/** Snapshot of GetRegisteredClassesVersionNumber() at the time of the last update, to invalidate on changes to classes. */
 		uint64 RegisteredClassesVersionNumber = MAX_uint64;
 		/** Dirty flag to invalidate on other changes requiring a recompute. */
@@ -127,8 +127,8 @@ public:
 
 	bool HasAssets(const FName PackagePath, const bool bRecursive) const;
 	FName GetRedirectedObjectPath(const FName ObjectPath) const;
-	bool GetAncestorClassNames(Impl::FClassInheritanceContext& InheritanceContext, FName ClassName,
-		TArray<FName>& OutAncestorClassNames) const;
+	bool GetAncestorClassNames(Impl::FClassInheritanceContext& InheritanceContext, FTopLevelAssetPath ClassName,
+		TArray<FTopLevelAssetPath>& OutAncestorClassNames) const;
 	void CompileFilter(Impl::FClassInheritanceContext& InheritanceContext, const FARFilter& InFilter,
 		FARCompiledFilter& OutCompiledFilter) const;
 	void SetTemporaryCachingMode(bool bEnable);
@@ -212,8 +212,8 @@ public:
 	void RemovePackageData(Impl::FEventContext& EventContext, const FName PackageName);
 
 	/** Returns the names of all subclasses of the class whose name is ClassName */
-	void GetSubClasses(Impl::FClassInheritanceContext& InheritanceContext, const TArray<FName>& InClassNames,
-		const TSet<FName>& ExcludedClassNames, TSet<FName>& SubClassNames) const;
+	void GetSubClasses(Impl::FClassInheritanceContext& InheritanceContext, const TArray<FTopLevelAssetPath>& InClassNames,
+		const TSet<FTopLevelAssetPath>& ExcludedClassNames, TSet<FTopLevelAssetPath>& SubClassNames) const;
 
 	bool IsUpdateDiskCacheAfterLoad() const { return bUpdateDiskCacheAfterLoad; }
 	bool IsInitialSearchCompleted() const { return bInitialSearchCompleted; }
@@ -233,7 +233,7 @@ public:
 	const FAssetPackageData* GetAssetPackageData(FName PackageName);
 
 	/** Same as UE::AssetRegistry::Filtering::ShouldSkipAsset, but can be read from any thread under readlock */
-	bool ShouldSkipAsset(FName AssetClass, uint32 PackageFlags) const;
+	bool ShouldSkipAsset(FTopLevelAssetPath AssetClass, uint32 PackageFlags) const;
 	bool ShouldSkipAsset(const UObject* InAsset) const;
 
 	/** Finds all class names of classes capable of generating new UClasses */
@@ -290,11 +290,17 @@ private:
 	bool RemoveAssetData(Impl::FEventContext& EventContext, FAssetData* AssetData);
 
 #if WITH_EDITOR
+	/** 
+	 * Calls PostLoadAssetRegistryTags on the CDO of the asset class this data represents 
+	 * @param AssetData Existing asset data
+	 */
+	void PostLoadAssetRegistryTags(FAssetData* AssetData);
+
 	/** Update Redirect collector with redirects loaded from asset registry */
 	void UpdateRedirectCollector();
 #endif // WITH_EDITOR
-	void GetSubClasses_Recursive(Impl::FClassInheritanceContext& InheritanceContext, FName InClassName,
-		TSet<FName>& SubClassNames, TSet<FName>& ProcessedClassNames, const TSet<FName>& ExcludedClassNames) const;
+	void GetSubClasses_Recursive(Impl::FClassInheritanceContext& InheritanceContext, FTopLevelAssetPath InClassName,
+		TSet<FTopLevelAssetPath>& SubClassNames, TSet<FTopLevelAssetPath>& ProcessedClassNames, const TSet<FTopLevelAssetPath>& ExcludedClassNames) const;
 
 	/** Internal helper which processes a given state and adds its contents to the current registry */
 	void CachePathsFromState(Impl::FEventContext& EventContext, const FAssetRegistryState& InState);
@@ -321,7 +327,7 @@ private:
 	TSet<FName> CachedEmptyPackages;
 
 	/** The map of classes to their parents, only full for offline blueprints */
-	TMap<FName, FName> CachedBPInheritanceMap;
+	TMap<FTopLevelAssetPath, FTopLevelAssetPath> CachedBPInheritanceMap;
 
 	/** If true, search caching is enabled */
 	bool bIsTempCachingEnabled;
@@ -378,7 +384,7 @@ private:
 	bool bSearchAllAssets;
 
 	/** List of all class names derived from Blueprint (including Blueprint itself) */
-	TSet<FName> ClassGeneratorNames;
+	TSet<FTopLevelAssetPath> ClassGeneratorNames;
 
 	struct FAssetRegistryPackageRedirect
 	{
@@ -413,9 +419,9 @@ private:
 #endif
 #if WITH_ENGINE && WITH_EDITOR
 	/** Class names that return true for IsAsset but which should not be treated as assets in uncooked packages */
-	TSet<FName> SkipUncookedClasses;
+	TSet<FTopLevelAssetPath> SkipUncookedClasses;
 	/** Class names that return true for IsAsset but which should not be treated as assets in cooked packages */
-	TSet<FName> SkipCookedClasses;
+	TSet<FTopLevelAssetPath> SkipCookedClasses;
 #endif
 #if ASSETREGISTRY_ENABLE_PREMADE_REGISTRY_IN_EDITOR
 	UE::AssetRegistry::Premade::FAsyncConsumer AsyncConsumer;
@@ -554,9 +560,6 @@ bool ReadAssetFile(FPackageReader& PackageReader, IAssetRegistry::FLoadPackageRe
  */
 bool IsPathMounted(const FString& Path, const TSet<FString>& MountPointsNoTrailingSlashes, FString& StringBuffer);
 
-/** Helper function to return the name of an object, given the object's export text path */
-FString ExportTextPathToObjectName(const FString& InExportTextPath);
-
 /**
  * Iterate all UObjects where this->IsAsset and !ShouldSkipAsset, create an FAssetData for them and call the callback.
  * Stores packagenames of all assets found in OutPackageNamesWithAssets for later filtering of on-disk assets.
diff --git a/Engine/Source/Runtime/AssetRegistry/Private/AssetRegistryState.cpp b/Engine/Source/Runtime/AssetRegistry/Private/AssetRegistryState.cpp
--- a/Engine/Source/Runtime/AssetRegistry/Private/AssetRegistryState.cpp
+++ b/Engine/Source/Runtime/AssetRegistry/Private/AssetRegistryState.cpp
@@ -148,14 +148,13 @@ void FAssetRegistryState::Reset()
 
 void FAssetRegistryState::FilterTags(const FAssetDataTagMapSharedView& InTagsAndValues, FAssetDataTagMap& OutTagsAndValues, const TSet<FName>* ClassSpecificFilterList, const FAssetRegistrySerializationOptions& Options)
 {
-	static FName WildcardName(TEXT("*"));
-	const TSet<FName>* AllClassesFilterList = Options.CookFilterlistTagsByClass.Find(WildcardName);
+	const TSet<FName>* AllClassesFilterList = Options.CookFilterlistTagsByClass.Find(UE::AssetRegistry::WildcardPathName);
 
 	// Exclude denied tags or include only allowed tags, based on how we were configured in ini
 	for (const auto& TagPair : InTagsAndValues)
 	{
-		const bool bInAllClassesList = AllClassesFilterList && (AllClassesFilterList->Contains(TagPair.Key) || AllClassesFilterList->Contains(WildcardName));
-		const bool bInClassSpecificList = ClassSpecificFilterList && (ClassSpecificFilterList->Contains(TagPair.Key) || ClassSpecificFilterList->Contains(WildcardName));
+		const bool bInAllClassesList = AllClassesFilterList && (AllClassesFilterList->Contains(TagPair.Key) || AllClassesFilterList->Contains(UE::AssetRegistry::WildcardFName));
+		const bool bInClassSpecificList = ClassSpecificFilterList && (ClassSpecificFilterList->Contains(TagPair.Key) || ClassSpecificFilterList->Contains(UE::AssetRegistry::WildcardFName));
 		if (Options.bUseAssetRegistryTagsAllowListInsteadOfDenyList)
 		{
 			// It's an allow list, only include it if it is in the all classes list or in the class specific list
@@ -225,10 +224,10 @@ void FAssetRegistryState::InitializeFromExistingAndPrune(const FAssetRegistrySta
 		}
 
 		FAssetDataTagMap NewTagsAndValues;
-		FAssetRegistryState::FilterTags(AssetData->TagsAndValues, NewTagsAndValues, Options.CookFilterlistTagsByClass.Find(AssetData->AssetClass), Options);
+		FAssetRegistryState::FilterTags(AssetData->TagsAndValues, NewTagsAndValues, Options.CookFilterlistTagsByClass.Find(AssetData->AssetClassPath), Options);
 
 		FAssetData* NewAssetData = new FAssetData(AssetData->PackageName, AssetData->PackagePath, AssetData->AssetName,
-			AssetData->AssetClass, NewTagsAndValues, AssetData->ChunkIDs, AssetData->PackageFlags);
+			AssetData->AssetClassPath, NewTagsAndValues, AssetData->ChunkIDs, AssetData->PackageFlags);
 
 		NewAssetData->TaggedAssetBundles = AssetData->TaggedAssetBundles;
 
@@ -357,7 +356,7 @@ void FAssetRegistryState::InitializeFromExisting(const TMap<FName, FAssetData*>&
 		const FAssetData& AssetData = *Pair.Value;
 
 		FAssetDataTagMap LocalTagsAndValues;
-		FAssetRegistryState::FilterTags(AssetData.TagsAndValues, LocalTagsAndValues, Options.CookFilterlistTagsByClass.Find(AssetData.AssetClass), Options);
+		FAssetRegistryState::FilterTags(AssetData.TagsAndValues, LocalTagsAndValues, Options.CookFilterlistTagsByClass.Find(AssetData.AssetClassPath), Options);
 		
 		if (ExistingData)
 		{
@@ -375,7 +374,7 @@ void FAssetRegistryState::InitializeFromExisting(const TMap<FName, FAssetData*>&
 		else
 		{
 			FAssetData* NewData = new FAssetData(AssetData.PackageName, AssetData.PackagePath, AssetData.AssetName,
-				AssetData.AssetClass, LocalTagsAndValues, AssetData.ChunkIDs, AssetData.PackageFlags);
+				AssetData.AssetClassPath, LocalTagsAndValues, AssetData.ChunkIDs, AssetData.PackageFlags);
 			NewData->TaggedAssetBundles = AssetData.TaggedAssetBundles;
 
 			AddAssetData(NewData);
@@ -608,7 +607,7 @@ bool FAssetRegistryState::HasAssets(const FName PackagePath, bool bARFiltering)
 		{
 			return FoundAssetArray->ContainsByPredicate([](FAssetData* AssetData)
 			{
-				return AssetData && !UE::AssetRegistry::FFiltering::ShouldSkipAsset(AssetData->AssetClass, AssetData->PackageFlags);
+				return AssetData && !UE::AssetRegistry::FFiltering::ShouldSkipAsset(AssetData->AssetClassPath, AssetData->PackageFlags);
 			});
 		}
 		else
@@ -629,14 +628,14 @@ bool FAssetRegistryState::GetAssets(const FARCompiledFilter& Filter, const TSet<
 	bARFiltering);
 }
 
-template<class ArrayType>
-TArray<FAssetData*> FindAssets(const TMap<FName, ArrayType>& Map, const TSet<FName>& Keys)
+template<class ArrayType, typename KeyType>
+TArray<FAssetData*> FindAssets(const TMap<KeyType, ArrayType>& Map, const TSet<KeyType>& Keys)
 {
 	TArray<TArrayView<FAssetData* const>> Matches;
 	Matches.Reserve(Keys.Num());
 	uint32 TotalMatches = 0;
 
-	for (FName Key : Keys)
+	for (const KeyType& Key : Keys)
 	{
 		if (const ArrayType* Assets = Map.Find(Key))
 		{
@@ -682,9 +681,9 @@ bool FAssetRegistryState::EnumerateAssets(const FARCompiledFilter& Filter, const
 	}
 
 	// On disk classes
-	if (Filter.ClassNames.Num() > 0)
+	if (Filter.ClassPaths.Num() > 0)
 	{
-		FilterResults.Emplace(FindAssets(CachedAssetsByClass, Filter.ClassNames));
+		FilterResults.Emplace(FindAssets(CachedAssetsByClass, Filter.ClassPaths));
 	}
 
 	// On disk object paths
@@ -754,7 +753,7 @@ bool FAssetRegistryState::EnumerateAssets(const FARCompiledFilter& Filter, const
 				return true;
 			}
 
-			return bARFiltering && UE::AssetRegistry::FFiltering::ShouldSkipAsset(AssetData->AssetClass, AssetData->PackageFlags);
+			return bARFiltering && UE::AssetRegistry::FFiltering::ShouldSkipAsset(AssetData->AssetClassPath, AssetData->PackageFlags);
 		};
 
 		int32 NumFilterResults = FilterResults.Num();
@@ -827,7 +826,7 @@ bool FAssetRegistryState::EnumerateAllAssets(const TSet<FName>& PackageNamesToSk
 
 		if (AssetData &&
 			!PackageNamesToSkip.Contains(AssetData->PackageName) &&
-			(!bARFiltering || !UE::AssetRegistry::FFiltering::ShouldSkipAsset(AssetData->AssetClass, AssetData->PackageFlags)))
+			(!bARFiltering || !UE::AssetRegistry::FFiltering::ShouldSkipAsset(AssetData->AssetClassPath, AssetData->PackageFlags)))
 		{
 			if (!Callback(*AssetData))
 			{
@@ -1058,7 +1057,8 @@ bool FAssetRegistryState::Save(FArchive& OriginalAr, const FAssetRegistrySeriali
 	Algo::Sort(SortedAssetsByObjectPath, [](TPair<FName, FAssetData*>& A, TPair<FName, FAssetData*>& B) { return A.Key.LexicalLess(B.Key); });
 	for (const TPair<FName, FAssetData*>& Pair : SortedAssetsByObjectPath)
 	{
-		Pair.Value->SerializeForCache(Ar);
+		// Hardcoding FAssetRegistryVersion::LatestVersion here so that branches can get optimized out in the forceinlined SerializeForCache
+		Pair.Value->SerializeForCache(Ar, FAssetRegistryVersion::LatestVersion);
 	}
 
 	// Serialize Dependencies
@@ -1183,7 +1183,7 @@ bool FAssetRegistryState::Load(FArchive& OriginalAr, const FAssetRegistryLoadOpt
 	}
 	else
 	{
-		FAssetRegistryReader Reader(OriginalAr, Options.ParallelWorkers);
+		FAssetRegistryReader Reader(OriginalAr, Options.ParallelWorkers, Version);
 
 		if (Reader.IsError())
 		{
@@ -1230,9 +1230,10 @@ void FAssetRegistryState::Load(Archive&& Ar, FAssetRegistryVersion::Type Version
 	// allocate one single block for all asset data structs (to reduce tens of thousands of heap allocations)
 	TArrayView<FAssetData> PreallocatedAssetDataBuffer(new FAssetData[LocalNumAssets], LocalNumAssets);
 	PreallocatedAssetDataBuffers.Add(PreallocatedAssetDataBuffer.GetData());
+
 	for (FAssetData& NewAssetData : PreallocatedAssetDataBuffer)
 	{
-		NewAssetData.SerializeForCache(Ar);
+		NewAssetData.SerializeForCache(Ar, Version);
 	}
 
 	SetAssetDatas(PreallocatedAssetDataBuffer, Options);
@@ -1403,7 +1404,12 @@ SIZE_T FAssetRegistryState::GetAllocatedSize(bool bLogDetailed) const
 	};
 	SubArray(CachedAssetsByPackageName);
 	SubArray(CachedAssetsByPath);
-	SubArray(CachedAssetsByClass);
+
+	for (auto& Pair : CachedAssetsByClass)
+	{
+		MapArrayMemory += Pair.Value.GetAllocatedSize();
+	}
+
 	SubArray(CachedAssetsByTag);
 
 	if (bLogDetailed)
@@ -1551,6 +1557,16 @@ FDependsNode* FAssetRegistryState::ResolveRedirector(FDependsNode* InDependency,
 	return Result;
 }
 
+template <typename KeyType>
+void ShrinkMultimap(TMap<KeyType, TArray<FAssetData*> >& Map)
+{
+	Map.Shrink();
+	for (auto& Pair : Map)
+	{
+		Pair.Value.Shrink();
+	}
+};
+
 void FAssetRegistryState::SetAssetDatas(TArrayView<FAssetData> AssetDatas, const FAssetRegistryLoadOptions& Options)
 {
 	UE_CLOG(NumAssets != 0, LogAssetRegistry, Fatal, TEXT("Can only load into empty asset registry states. Load into temporary and append using InitializeFromExisting() instead."));
@@ -1581,15 +1597,6 @@ void FAssetRegistryState::SetAssetDatas(TArrayView<FAssetData> AssetDatas, const
 
 	auto SetOtherCaches = [&]()
 	{
-		auto ShrinkMultimap = [](TMap<FName, TArray<FAssetData*> >& Map)
-		{
-			Map.Shrink();
-			for (auto& Pair : Map)
-			{
-				Pair.Value.Shrink();
-			}
-		};
-
 		CachedAssetsByPath.Empty();
 		for (FAssetData& AssetData : AssetDatas)
 		{
@@ -1601,7 +1608,7 @@ void FAssetRegistryState::SetAssetDatas(TArrayView<FAssetData> AssetDatas, const
 		CachedAssetsByClass.Empty();
 		for (FAssetData& AssetData : AssetDatas)
 		{
-			TArray<FAssetData*>& ClassAssets = CachedAssetsByClass.FindOrAdd(AssetData.AssetClass);
+			TArray<FAssetData*>& ClassAssets = CachedAssetsByClass.FindOrAdd(AssetData.AssetClassPath);
 			ClassAssets.Add(&AssetData);
 		}
 		ShrinkMultimap(CachedAssetsByClass);
@@ -1651,7 +1658,7 @@ void FAssetRegistryState::AddAssetData(FAssetData* AssetData)
 
 	TArray<FAssetData*, TInlineAllocator<1>>& PackageAssets = CachedAssetsByPackageName.FindOrAdd(AssetData->PackageName);
 	TArray<FAssetData*>& PathAssets = CachedAssetsByPath.FindOrAdd(AssetData->PackagePath);
-	TArray<FAssetData*>& ClassAssets = CachedAssetsByClass.FindOrAdd(AssetData->AssetClass);
+	TArray<FAssetData*>& ClassAssets = CachedAssetsByClass.FindOrAdd(AssetData->AssetClassPath);
 
 	PackageAssets.Add(AssetData);
 	PathAssets.Add(AssetData);
@@ -1770,10 +1777,10 @@ void FAssetRegistryState::UpdateAssetData(FAssetData* AssetData, const FAssetDat
 	}
 
 	// Update AssetClass
-	if (AssetData->AssetClass != NewAssetData.AssetClass)
+	if (AssetData->AssetClassPath != NewAssetData.AssetClassPath)
 	{
-		TArray<FAssetData*>& NewClassAssets = CachedAssetsByClass.FindOrAdd(NewAssetData.AssetClass);
-		TArray<FAssetData*>* OldClassAssets = CachedAssetsByClass.Find(AssetData->AssetClass);
+		TArray<FAssetData*>& NewClassAssets = CachedAssetsByClass.FindOrAdd(NewAssetData.AssetClassPath);
+		TArray<FAssetData*>* OldClassAssets = CachedAssetsByClass.Find(AssetData->AssetClassPath);
 
 		OldClassAssets->Remove(AssetData);
 		NewClassAssets.Add(AssetData);
@@ -1845,7 +1852,7 @@ void FAssetRegistryState::RemoveAssetData(FAssetData* AssetData, bool bRemoveDep
 
 	TArray<FAssetData*, TInlineAllocator<1>>* OldPackageAssets = CachedAssetsByPackageName.Find(AssetData->PackageName);
 	TArray<FAssetData*>* OldPathAssets = CachedAssetsByPath.Find(AssetData->PackagePath);
-	TArray<FAssetData*>* OldClassAssets = CachedAssetsByClass.Find(AssetData->AssetClass);
+	TArray<FAssetData*>* OldClassAssets = CachedAssetsByClass.Find(AssetData->AssetClassPath);
 	OldPackageAssets->RemoveSingleSwap(AssetData);
 	OldPathAssets->RemoveSingleSwap(AssetData);
 	OldClassAssets->RemoveSingleSwap(AssetData);
@@ -2026,6 +2033,19 @@ bool FAssetRegistryState::IsFilterValid(const FARCompiledFilter& Filter)
 	return UE::AssetRegistry::Utils::IsFilterValid(Filter);
 }
 
+const TArray<const FAssetData*>& FAssetRegistryState::GetAssetsByClassName(const FName ClassName) const
+{
+	static TArray<const FAssetData*> InvalidArray;
+	for (const TPair< FTopLevelAssetPath, TArray<FAssetData*> >& Pair : CachedAssetsByClass)
+	{
+		if (Pair.Key.GetAssetName() == ClassName)
+		{
+			return reinterpret_cast<const TArray<const FAssetData*>&>(Pair.Value);
+		}
+	}
+	return InvalidArray;
+}
+
 namespace UE::AssetRegistry::Utils
 {
 
@@ -2034,7 +2054,7 @@ bool IsFilterValid(const FARCompiledFilter& Filter)
 	if (Filter.PackageNames.Contains(NAME_None) ||
 		Filter.PackagePaths.Contains(NAME_None) ||
 		Filter.ObjectPaths.Contains(NAME_None) ||
-		Filter.ClassNames.Contains(NAME_None) ||
+		Filter.ClassPaths.Contains(FTopLevelAssetPath()) ||
 		Filter.TagsAndValues.Contains(NAME_None)
 		)
 	{
@@ -2086,6 +2106,59 @@ namespace AssetRegistry
 }
 }
 
+template <typename MapType>
+static void PrintAssetDataMap(FString Name, const MapType& AssetMap, TStringBuilder<16>& PageBuffer, const TFunctionRef<void()>& AddLine)
+{
+	PageBuffer.Appendf(TEXT("--- Begin %s ---"), *Name);
+	AddLine();
+
+	TArray<typename MapType::KeyType> Keys;
+	AssetMap.GenerateKeyArray(Keys);
+
+	struct FKeyTypeCompare
+	{
+		FORCEINLINE bool operator()(const typename MapType::KeyType& A, const typename MapType::KeyType& B) const
+		{
+			return A.Compare(B) < 0;
+		}
+	};
+	Keys.Sort(FKeyTypeCompare());
+
+	TArray<FAssetData*> Items;
+	Items.Reserve(1024);
+
+	int32 ValidCount = 0;
+	for (const typename MapType::KeyType& Key : Keys)
+	{
+		const auto& AssetArray = AssetMap.FindChecked(Key);
+		if (AssetArray.Num() == 0)
+		{
+			continue;
+		}
+		++ValidCount;
+
+		Items.Reset();
+		Items.Append(AssetArray);
+		Items.Sort([](const FAssetData& A, const FAssetData& B)
+			{ return A.ObjectPath.ToString() < B.ObjectPath.ToString(); }
+		);
+
+		PageBuffer.Append(TEXT("	"));
+		Key.AppendString(PageBuffer);
+		PageBuffer.Appendf(TEXT(" : %d item(s)"), Items.Num());
+		AddLine();
+		for (const FAssetData* Data : Items)
+		{
+			PageBuffer.Append(TEXT("	 "));
+			Data->ObjectPath.AppendString(PageBuffer);
+			AddLine();
+		}
+	}
+
+	PageBuffer.Appendf(TEXT("--- End %s : %d entries ---"), *Name, ValidCount);
+	AddLine();
+};
+
 void FAssetRegistryState::Dump(const TArray<FString>& Arguments, TArray<FString>& OutPages, int32 LinesPerPage) const
 {
 	int32 ExpectedNumLines = 14 + CachedAssetsByObjectPath.Num() * 5 + CachedDependsNodes.Num() + CachedPackageData.Num();
@@ -2189,50 +2262,6 @@ void FAssetRegistryState::Dump(const TArray<FString>& Arguments, TArray<FString>
 		}
 	};
 
-	auto PrintAssetDataMap = [&AddLine, &PageBuffer](FString Name, const auto& AssetMap)
-	{
-		PageBuffer.Appendf(TEXT("--- Begin %s ---"), *Name);
-		AddLine();
-
-		TArray<FName> Keys;
-		AssetMap.GenerateKeyArray(Keys);
-		Keys.Sort(FNameLexicalLess());
-
-		TArray<FAssetData*> Items;
-		Items.Reserve(1024);
-
-		int32 ValidCount = 0;
-		for (const FName& Key : Keys)
-		{
-			const auto& AssetArray = AssetMap.FindChecked(Key);
-			if (AssetArray.Num() == 0)
-			{
-				continue;
-			}
-			++ValidCount;
-
-			Items.Reset();
-			Items.Append(AssetArray);
-			Items.Sort([](const FAssetData& A, const FAssetData& B)
-				{ return A.ObjectPath.ToString() < B.ObjectPath.ToString(); }
-				);
-
-			PageBuffer.Append(TEXT("	"));
-			Key.AppendString(PageBuffer);
-			PageBuffer.Appendf(TEXT(" : %d item(s)"), Items.Num());
-			AddLine();
-			for (const FAssetData* Data : Items)
-			{
-				PageBuffer.Append(TEXT("	 "));
-				Data->ObjectPath.AppendString(PageBuffer);
-				AddLine();
-			}
-		}
-
-		PageBuffer.Appendf(TEXT("--- End %s : %d entries ---"), *Name, ValidCount);
-		AddLine();
-	};
-
 	if (Arguments.Contains(TEXT("ObjectPath")))
 	{
 		PageBuffer.Append(TEXT("--- Begin CachedAssetsByObjectPath ---"));
@@ -2255,22 +2284,22 @@ void FAssetRegistryState::Dump(const TArray<FString>& Arguments, TArray<FString>
 
 	if (Arguments.Contains(TEXT("PackageName")))
 	{
-		PrintAssetDataMap(TEXT("CachedAssetsByPackageName"), CachedAssetsByPackageName);
+		PrintAssetDataMap(TEXT("CachedAssetsByPackageName"), CachedAssetsByPackageName, PageBuffer, AddLine);
 	}
 
 	if (Arguments.Contains(TEXT("Path")))
 	{
-		PrintAssetDataMap(TEXT("CachedAssetsByPath"), CachedAssetsByPath);
+		PrintAssetDataMap(TEXT("CachedAssetsByPath"), CachedAssetsByPath, PageBuffer, AddLine);
 	}
 
 	if (Arguments.Contains(TEXT("Class")))
 	{
-		PrintAssetDataMap(TEXT("CachedAssetsByClass"), CachedAssetsByClass);
+		PrintAssetDataMap(TEXT("CachedAssetsByClass"), CachedAssetsByClass, PageBuffer, AddLine);
 	}
 
 	if (Arguments.Contains(TEXT("Tag")))
 	{
-		PrintAssetDataMap(TEXT("CachedAssetsByTag"), CachedAssetsByTag);
+		PrintAssetDataMap(TEXT("CachedAssetsByTag"), CachedAssetsByTag, PageBuffer, AddLine);
 	}
 
 	if (Arguments.Contains(TEXT("Dependencies")) && !bDumpDependencyDetails)
diff --git a/Engine/Source/Runtime/AssetRegistry/Private/PackageReader.cpp b/Engine/Source/Runtime/AssetRegistry/Private/PackageReader.cpp
--- a/Engine/Source/Runtime/AssetRegistry/Private/PackageReader.cpp
+++ b/Engine/Source/Runtime/AssetRegistry/Private/PackageReader.cpp
@@ -14,6 +14,7 @@
 #include "UObject/Linker.h"
 #include "UObject/PackageTrailer.h"
 
+
 FPackageReader::FPackageReader()
 	: Loader(nullptr)
 	, PackageFileSize(0)
@@ -355,7 +356,7 @@ bool FPackageReader::ReadAssetDataFromThumbnailCache(TArray<FAssetData*>& AssetD
 		}
 
 		// Create a new FAssetData for this asset and update it with the gathered data
-		AssetDataList.Add(new FAssetData(FName(*PackageName), FName(*PackagePath), FName(*ObjectPathWithoutPackageName), FName(*AssetClassName), FAssetDataTagMap(), PackageFileSummary.ChunkIDs, PackageFileSummary.GetPackageFlags()));
+		AssetDataList.Add(new FAssetData(FName(*PackageName), FName(*PackagePath), FName(*ObjectPathWithoutPackageName), FTopLevelAssetPath(AssetClassName), FAssetDataTagMap(), PackageFileSummary.ChunkIDs, PackageFileSummary.GetPackageFlags()));
 	}
 
 	return true;
@@ -398,23 +399,28 @@ bool FPackageReader::ReadAssetRegistryDataFromCookedPackage(TArray<FAssetData*>&
 			if (Export.bIsAsset)
 			{
 				// We need to get the class name from the import/export maps
-				FName ObjectClassName;
+				FString ObjectClassName;
 				if (Export.ClassIndex.IsNull())
 				{
-					ObjectClassName = UClass::StaticClass()->GetFName();
+					ObjectClassName = UClass::StaticClass()->GetPathName();
 				}
 				else if (Export.ClassIndex.IsExport())
 				{
 					const FObjectExport& ClassExport = ExportMap[Export.ClassIndex.ToExport()];
-					ObjectClassName = ClassExport.ObjectName;
+					ObjectClassName = PackageName;
+					ObjectClassName += '.'; 
+					ClassExport.ObjectName.AppendString(ObjectClassName);
 				}
 				else if (Export.ClassIndex.IsImport())
 				{
 					const FObjectImport& ClassImport = ImportMap[Export.ClassIndex.ToImport()];
-					ObjectClassName = ClassImport.ObjectName;
+					const FObjectImport& ClassPackageImport = ImportMap[ClassImport.OuterIndex.ToImport()];
+					ClassPackageImport.ObjectName.AppendString(ObjectClassName);
+					ObjectClassName += '.';
+					ClassImport.ObjectName.AppendString(ObjectClassName);
 				}
 
-				AssetDataList.Add(new FAssetData(FName(*PackageName), FName(*PackagePath), Export.ObjectName, ObjectClassName, FAssetDataTagMap(), TArray<int32>(), GetPackageFlags()));
+				AssetDataList.Add(new FAssetData(FName(*PackageName), FName(*PackagePath), Export.ObjectName, FTopLevelAssetPath(ObjectClassName), FAssetDataTagMap(), TArray<int32>(), GetPackageFlags()));
 				bFoundAtLeastOneAsset = true;
 			}
 		}
@@ -871,6 +877,107 @@ FArchive& FPackageReader::operator<<( FName& Name )
 
 namespace UE::AssetRegistry
 {
+	class FNameMapAwareArchive : public FArchiveProxy
+	{
+		TArray<FNameEntryId>	NameMap;
+
+	public:
+
+		FNameMapAwareArchive(FArchive& Inner)
+			: FArchiveProxy(Inner)
+		{}
+
+		FORCEINLINE virtual FArchive& operator<<(FName& Name) override
+		{
+			FArchive& Ar = *this;
+			int32 NameIndex;
+			Ar << NameIndex;
+			int32 Number = 0;
+			Ar << Number;
+
+			if (NameMap.IsValidIndex(NameIndex))
+			{
+				// if the name wasn't loaded (because it wasn't valid in this context)
+				FNameEntryId MappedName = NameMap[NameIndex];
+
+				// simply create the name from the NameMap's name and the serialized instance number
+				Name = FName::CreateFromDisplayId(MappedName, Number);
+			}
+			else
+			{
+				Name = FName();
+				SetCriticalError();
+			}
+
+			return *this;
+		}
+
+		void SerializeNameMap(const FPackageFileSummary& PackageFileSummary)
+		{
+			Seek(PackageFileSummary.NameOffset);
+			NameMap.Reserve(PackageFileSummary.NameCount);
+			FNameEntrySerialized NameEntry(ENAME_LinkerConstructor);
+			for (int32 Idx = NameMap.Num(); Idx < PackageFileSummary.NameCount; ++Idx)
+			{
+				*this << NameEntry;
+				NameMap.Emplace(FName(NameEntry).GetDisplayIndex());
+			}
+		}
+	};
+	FString ReconstructFullClassPath(FArchive& BinaryArchive, const FString& PackageName, const FPackageFileSummary& PackageFileSummary, const FString& AssetClassName)
+	{
+		FLinkerTables LinkerTables;
+		FNameMapAwareArchive NameMapArchive(BinaryArchive);
+		NameMapArchive.SerializeNameMap(PackageFileSummary);
+		FName ClassFName(*AssetClassName);
+
+		// Load the linker tables
+		BinaryArchive.Seek(PackageFileSummary.ImportOffset);
+		for (int32 ImportMapIndex = 0; ImportMapIndex < PackageFileSummary.ImportCount; ++ImportMapIndex)
+		{
+			FObjectImport* Import = new(LinkerTables.ImportMap)FObjectImport;
+			NameMapArchive << *Import;
+		}
+		BinaryArchive.Seek(PackageFileSummary.ExportOffset);
+		for (int32 ExportMapIndex = 0; ExportMapIndex < PackageFileSummary.ExportCount; ++ExportMapIndex)
+		{
+			FObjectExport* Export = new(LinkerTables.ExportMap)FObjectExport;
+			NameMapArchive << *Export;
+		}
+
+		FString ClassPathName;
+
+		// Now look through the exports' classes and find the one matching the asset class
+		for (const FObjectExport& Export : LinkerTables.ExportMap)
+		{
+			if (Export.ClassIndex.IsImport())
+			{
+				if (LinkerTables.ImportMap[Export.ClassIndex.ToImport()].ObjectName == ClassFName)
+				{
+					ClassPathName = LinkerTables.GetImportPathName(Export.ClassIndex.ToImport());
+					break;
+				}					
+			}
+			else if (Export.ClassIndex.IsExport())
+			{
+				if (LinkerTables.ExportMap[Export.ClassIndex.ToExport()].ObjectName == ClassFName)
+				{
+					ClassPathName = LinkerTables.GetExportPathName(PackageName, Export.ClassIndex.ToExport());
+					break;
+				}
+			}
+		}
+		if (ClassPathName.IsEmpty())
+		{
+			UE_LOG(LogAssetRegistry, Error, TEXT("Failed to find an import or export matching asset class short name \"%s\"."), *AssetClassName);
+			// Just pass through the short class name
+			ClassPathName = AssetClassName;
+		}
+
+
+		return ClassPathName;
+	}
+
 	// See the corresponding WritePackageData defined in SavePackageUtilities.cpp in CoreUObject module
 	bool ReadPackageDataMain(FArchive& BinaryArchive, const FString& PackageName, const FPackageFileSummary& PackageFileSummary, int64& OutDependencyDataOffset, TArray<FAssetData*>& OutAssetDataList, EReadPackageDataMainErrorCode& OutError)
 	{
@@ -912,7 +1019,7 @@ namespace UE::AssetRegistry
 			if (bLegacyPackage || bNoMapAsset)
 			{
 				FString AssetName = FPackageName::GetLongPackageAssetName(PackageName);
-				OutAssetDataList.Add(new FAssetData(FName(*PackageName), FName(*PackagePath), FName(*AssetName), FName(TEXT("World")), FAssetDataTagMap(), PackageFileSummary.ChunkIDs, PackageFileSummary.GetPackageFlags()));
+				OutAssetDataList.Add(new FAssetData(FName(*PackageName), FName(*PackagePath), FName(*AssetName), FTopLevelAssetPath(TEXT("/Script/Engine"), TEXT("World")), FAssetDataTagMap(), PackageFileSummary.ChunkIDs, PackageFileSummary.GetPackageFlags()));
 			}
 		}
 
@@ -925,6 +1032,7 @@ namespace UE::AssetRegistry
 			int32 TagCount = 0;
 			BinaryArchive << ObjectPath;
 			BinaryArchive << ObjectClassName;
+			// @todo make sure this is a full path name
 			BinaryArchive << TagCount;
 			if (BinaryArchive.IsError() || TagCount < 0 || PackageFileSize < BinaryArchive.Tell() + TagCount * MinBytesPerTag)
 			{
@@ -985,7 +1093,13 @@ namespace UE::AssetRegistry
 			}
 
 			// Create a new FAssetData for this asset and update it with the gathered data
-			OutAssetDataList.Add(new FAssetData(PackageName, ObjectPath, FName(*ObjectClassName), MoveTemp(TagsAndValues), PackageFileSummary.ChunkIDs, PackageFileSummary.GetPackageFlags()));
+			if (!ObjectClassName.IsEmpty() && FPackageName::IsShortPackageName(ObjectClassName))
+			{
+				int64 CurrentPos = BinaryArchive.Tell();
+				ObjectClassName = ReconstructFullClassPath(BinaryArchive, PackageName, PackageFileSummary, ObjectClassName);
+				BinaryArchive.Seek(CurrentPos);
+			}
+			OutAssetDataList.Add(new FAssetData(PackageName, ObjectPath, FTopLevelAssetPath(ObjectClassName), MoveTemp(TagsAndValues), PackageFileSummary.ChunkIDs, PackageFileSummary.GetPackageFlags()));
 		}
 
 		return true;
diff --git a/Engine/Source/Runtime/AssetRegistry/Public/AssetRegistry/AssetRegistryState.h b/Engine/Source/Runtime/AssetRegistry/Public/AssetRegistry/AssetRegistryState.h
--- a/Engine/Source/Runtime/AssetRegistry/Public/AssetRegistry/AssetRegistryState.h
+++ b/Engine/Source/Runtime/AssetRegistry/Public/AssetRegistry/AssetRegistryState.h
@@ -53,8 +53,8 @@ struct FAssetRegistrySerializationOptions
 	/** Filter out searchable names from dependency data */
 	bool bFilterSearchableNames = false;
 
-	/** The map of classname to tag set of tags that are allowed in cooked builds. This is either an allow list or deny list depending on bUseAssetRegistryTagsAllowListInsteadOfDenyList */
-	TMap<FName, TSet<FName>> CookFilterlistTagsByClass;
+	/** The map of class pathname to tag set of tags that are allowed in cooked builds. This is either an allow list or deny list depending on bUseAssetRegistryTagsAllowListInsteadOfDenyList */
+	TMap<FTopLevelAssetPath, TSet<FName>> CookFilterlistTagsByClass;
 
 	/** Tag keys whose values should be stored as FName in cooked builds */
 	TSet<FName> CookTagsAsName;
@@ -254,10 +254,19 @@ public:
 	 * @param ClassName the class name of the assets to look for
 	 * @return An array of AssetData*, empty if nothing found
 	 */
-	const TArray<const FAssetData*>& GetAssetsByClassName(const FName ClassName) const
+	UE_DEPRECATED(5.1, "Class names are now represented by path names. Please use GetAssetsByClassPathName")
+	const TArray<const FAssetData*>& GetAssetsByClassName(const FName ClassName) const;
+
+	/**
+	 * Gets the asset data for the specified asset class
+	 *
+	 * @param ClassPathName the class path name of the assets to look for
+	 * @return An array of AssetData*, empty if nothing found
+	 */
+	const TArray<const FAssetData*>& GetAssetsByClassPathName(const FTopLevelAssetPath ClassPathName) const
 	{
 		static TArray<const FAssetData*> InvalidArray;
-		const TArray<FAssetData*>* FoundAssetArray = CachedAssetsByClass.Find(ClassName);
+		const TArray<FAssetData*>* FoundAssetArray = CachedAssetsByClass.Find(ClassPathName);
 		if (FoundAssetArray)
 		{
 			return reinterpret_cast<const TArray<const FAssetData*>&>(*FoundAssetArray);
@@ -452,7 +461,7 @@ private:
 	TMap<FName, TArray<FAssetData*> > CachedAssetsByPath;
 
 	/** The map of class name to asset data for assets saved to disk */
-	TMap<FName, TArray<FAssetData*> > CachedAssetsByClass;
+	TMap<FTopLevelAssetPath, TArray<FAssetData*> > CachedAssetsByClass;
 
 	/** The map of asset tag to asset data for assets saved to disk */
 	TMap<FName, TArray<FAssetData*> > CachedAssetsByTag;
