Fixing editor startup crash related to parsing assed tag values that were stored as 'None'

diff --git a/Engine/Source/Runtime/AssetRegistry/Private/AssetRegistry.cpp b/Engine/Source/Runtime/AssetRegistry/Private/AssetRegistry.cpp
--- a/Engine/Source/Runtime/AssetRegistry/Private/AssetRegistry.cpp
+++ b/Engine/Source/Runtime/AssetRegistry/Private/AssetRegistry.cpp
@@ -4665,15 +4665,17 @@ void FAssetRegistryImpl::AddAssetData(Impl::FEventContext& EventContext, FAssetD
 	{
 		const FString GeneratedClass = AssetData->GetTagValueRef<FString>(FBlueprintTags::GeneratedClassPath);
 		const FString ParentClass = AssetData->GetTagValueRef<FString>(FBlueprintTags::ParentClassPath);
-		if (!GeneratedClass.IsEmpty() && !ParentClass.IsEmpty())
+		if (!GeneratedClass.IsEmpty() && !ParentClass.IsEmpty() && GeneratedClass != TEXTVIEW("None") && ParentClass != TEXTVIEW("None"))
 		{
 			const FTopLevelAssetPath GeneratedClassPathName(GeneratedClass);
 			const FTopLevelAssetPath ParentClassPathName(ParentClass);
-			checkf(!GeneratedClassPathName.IsNull() && !ParentClassPathName.IsNull(), TEXT("Short class names used in AddAssetData: GeneratedClass=%s, ParentClass=%s. Short class names in these tags on the Blueprint class should have been converted to path names."), *GeneratedClass, *ParentClass);
-			CachedBPInheritanceMap.Add(GeneratedClassPathName, ParentClassPathName);
-			
-			// Invalidate caching because CachedBPInheritanceMap got modified
-			TempCachedInheritanceBuffer.bDirty = true;
+			if (ensureAlwaysMsgf(!GeneratedClassPathName.IsNull() && !ParentClassPathName.IsNull(), TEXT("Short class names used in AddAssetData: GeneratedClass=%s, ParentClass=%s. Short class names in these tags on the Blueprint class should have been converted to path names."), *GeneratedClass, *ParentClass))
+			{
+				CachedBPInheritanceMap.Add(GeneratedClassPathName, ParentClassPathName);
+
+				// Invalidate caching because CachedBPInheritanceMap got modified
+				TempCachedInheritanceBuffer.bDirty = true;
+			}
 		}
 	}
 }
@@ -4684,24 +4686,28 @@ void FAssetRegistryImpl::UpdateAssetData(Impl::FEventContext& EventContext, FAss
 	if (ClassGeneratorNames.Contains(AssetData->AssetClassPath))
 	{
 		const FString OldGeneratedClass = AssetData->GetTagValueRef<FString>(FBlueprintTags::GeneratedClassPath);
-		if (!OldGeneratedClass.IsEmpty())
+		if (!OldGeneratedClass.IsEmpty() && OldGeneratedClass != TEXTVIEW("None"))
 		{
 			const FTopLevelAssetPath OldGeneratedClassName(OldGeneratedClass);
-			checkf(!OldGeneratedClassName.IsNull(), TEXT("Short class name used: OldGeneratedClass=%s. Short class names in tags on the Blueprint class should have been converted to path names."), *OldGeneratedClass);
-			CachedBPInheritanceMap.Remove(OldGeneratedClassName);
+			if (ensureAlwaysMsgf(!OldGeneratedClassName.IsNull(), TEXT("Short class name used: OldGeneratedClass=%s. Short class names in tags on the Blueprint class should have been converted to path names."), *OldGeneratedClass))
+			{
+				CachedBPInheritanceMap.Remove(OldGeneratedClassName);
 
-			// Invalidate caching because CachedBPInheritanceMap got modified
-			TempCachedInheritanceBuffer.bDirty = true;
+				// Invalidate caching because CachedBPInheritanceMap got modified
+				TempCachedInheritanceBuffer.bDirty = true;
+			}
 		}
 
 		const FString NewGeneratedClass = NewAssetData.GetTagValueRef<FString>(FBlueprintTags::GeneratedClassPath);
 		const FString NewParentClass = NewAssetData.GetTagValueRef<FString>(FBlueprintTags::ParentClassPath);
-		if (!NewGeneratedClass.IsEmpty() && !NewParentClass.IsEmpty())
+		if (!NewGeneratedClass.IsEmpty() && !NewParentClass.IsEmpty() && NewGeneratedClass != TEXTVIEW("None") && NewParentClass != TEXTVIEW("None"))
 		{
 			const FTopLevelAssetPath NewGeneratedClassName(NewGeneratedClass);
 			const FTopLevelAssetPath NewParentClassName(NewParentClass);
-			checkf(!NewGeneratedClassName.IsNull() && !NewParentClassName.IsNull(), TEXT("Short class names used in AddAssetData: GeneratedClass=%s, ParentClass=%s. Short class names in these tags on the Blueprint class should have been converted to path names."), *NewGeneratedClass, *NewParentClass);
-			CachedBPInheritanceMap.Add(NewGeneratedClassName, NewParentClassName);
+			if (ensureAlwaysMsgf(!NewGeneratedClassName.IsNull() && !NewParentClassName.IsNull(), TEXT("Short class names used in AddAssetData: GeneratedClass=%s, ParentClass=%s. Short class names in these tags on the Blueprint class should have been converted to path names."), *NewGeneratedClass, *NewParentClass))
+			{
+				CachedBPInheritanceMap.Add(NewGeneratedClassName, NewParentClassName);
+			}
 
 			// Invalidate caching because CachedBPInheritanceMap got modified
 			TempCachedInheritanceBuffer.bDirty = true;
@@ -4731,14 +4737,16 @@ bool FAssetRegistryImpl::RemoveAssetData(Impl::FEventContext& EventContext, FAss
 		if (ClassGeneratorNames.Contains(AssetData->AssetClassPath))
 		{
 			const FString OldGeneratedClass = AssetData->GetTagValueRef<FString>(FBlueprintTags::GeneratedClassPath);
-			if (!OldGeneratedClass.IsEmpty())
+			if (!OldGeneratedClass.IsEmpty() && OldGeneratedClass != TEXTVIEW("None"))
 			{
 				const FTopLevelAssetPath OldGeneratedClassPathName(FPackageName::ExportTextPathToObjectPath(OldGeneratedClass));
-				checkf(!OldGeneratedClassPathName.IsNull(), TEXT("Short class name used: OldGeneratedClass=%s"), *OldGeneratedClass);
-				CachedBPInheritanceMap.Remove(OldGeneratedClassPathName);
+				if (ensureAlwaysMsgf(!OldGeneratedClassPathName.IsNull(), TEXT("Short class name used: OldGeneratedClass=%s"), *OldGeneratedClass))
+				{
+					CachedBPInheritanceMap.Remove(OldGeneratedClassPathName);
 
-				// Invalidate caching because CachedBPInheritanceMap got modified
-				TempCachedInheritanceBuffer.bDirty = true;
+					// Invalidate caching because CachedBPInheritanceMap got modified
+					TempCachedInheritanceBuffer.bDirty = true;
+				}
 			}
 		}
 
