AssetRegistry::GetDerivedClassNames: Add handling of CoreRedirects so that subclasses registered with the oldname of a coreredirect for their parent class will still be found.

This fixes a problem that was exposed when  ANY_PACKAGE was removed; previously classes with a CoreRedirect to a different script package but the same class name would still be found even without searching CoreRedirects, but that is no longer the case and there is some content that was relying on it.

#rnx

diff --git a/Engine/Source/Runtime/AssetRegistry/Private/AssetRegistry.cpp b/Engine/Source/Runtime/AssetRegistry/Private/AssetRegistry.cpp
--- a/Engine/Source/Runtime/AssetRegistry/Private/AssetRegistry.cpp
+++ b/Engine/Source/Runtime/AssetRegistry/Private/AssetRegistry.cpp
@@ -5670,15 +5670,40 @@ void FAssetRegistryImpl::GetSubClasses_Recursive(Impl::FClassInheritanceContext&
 		SubClassNames.Add(InClassName);
 		ProcessedClassNames.Add(InClassName);
 
-		const TArray<FTopLevelAssetPath>* FoundSubClassNames = InheritanceContext.Buffer->ReverseInheritanceMap.Find(InClassName);
-		if (FoundSubClassNames)
+		auto AddSubClasses = [this, &InheritanceContext, &SubClassNames, &ProcessedClassNames, &ExcludedClassNames]
+		(FTopLevelAssetPath ParentClassName)
 		{
-			for (FTopLevelAssetPath ClassName : (*FoundSubClassNames))
+			const TArray<FTopLevelAssetPath>* FoundSubClassNames = InheritanceContext.Buffer->ReverseInheritanceMap.Find(ParentClassName);
+			if (FoundSubClassNames)
 			{
-				GetSubClasses_Recursive(InheritanceContext, ClassName, SubClassNames, ProcessedClassNames,
-					ExcludedClassNames);
+				for (FTopLevelAssetPath ClassName : (*FoundSubClassNames))
+				{
+					GetSubClasses_Recursive(InheritanceContext, ClassName, SubClassNames, ProcessedClassNames,
+						ExcludedClassNames);
+				}
+			}
+		};
+
+		// Add Subclasses of the given classname
+		AddSubClasses(InClassName);
+#if WITH_EDITOR
+		// Also add subclasses of any old names (from CoreRedirects) in case subclasses have not been resaved and have
+		// registered their ParentClass as the old name
+		TArray<FString> OldNames = FLinkerLoad::FindPreviousPathNamesForClass(InClassName.ToString(), false);
+		for (const FString& OldNameString : OldNames)
+		{
+			if (FPackageName::IsShortPackageName(OldNameString))
+			{
+				// TODO: We do not have a way to handle these, because we don't allow searching subclasses for short path
+				// names. We should instead change the construction of CachedBPInheritanceMap and ReverseInheritanceMap 
+				// to modify the ParentClassName with CoreRedirects
+				UE_LOG(LogAssetRegistry, Display, TEXT("GetDerivedClassNames found CoreRedirect from ShortPathName %s; it will not be able to find old blueprint classes that were saved with that parent class before the Redirect was added."),
+					*OldNameString);
+				continue;
 			}
+			AddSubClasses(FTopLevelAssetPath(OldNameString));
 		}
+#endif
 	}
 }
 
