AssetRegistry updating during cook: We call InitializeFromExisting on the AssetRegistryGenerator's state at the end of the cook to copy any updates to AssetDatas from loading the packages. When we do this, we need to update all data on the AssetData. Previously we were only updating the AssetData if the tags changed. But the AssetRegistry does not currently handle CoreRedirects, so it is possible for the ClassPath of the AssetData to change when the package loads. Failure to update the ClassPath was causing bugs when using class CoreRedirects.
#rnx

diff --git a/Engine/Source/Editor/UnrealEd/Private/Commandlets/AssetRegistryGenerator.cpp b/Engine/Source/Editor/UnrealEd/Private/Commandlets/AssetRegistryGenerator.cpp
--- a/Engine/Source/Editor/UnrealEd/Private/Commandlets/AssetRegistryGenerator.cpp
+++ b/Engine/Source/Editor/UnrealEd/Private/Commandlets/AssetRegistryGenerator.cpp
@@ -40,7 +40,6 @@
 #include "Serialization/JsonReader.h"
 #include "Serialization/JsonSerializer.h"
 
-
 DEFINE_LOG_CATEGORY_STATIC(LogAssetRegistryGenerator, Log, All);
 
 #define LOCTEXT_NAMESPACE "AssetRegistryGenerator"
@@ -933,7 +932,7 @@ void FAssetRegistryGenerator::InjectEncryptionData(FAssetRegistryState& TargetSt
 								TagsAndValues.Add(UAssetManager::GetEncryptionKeyAssetTagName(), Guid.ToString());
 								FAssetData NewAssetData = FAssetData(AssetData->PackageName, AssetData->PackagePath, AssetData->AssetName, AssetData->AssetClassPath, TagsAndValues, AssetData->ChunkIDs, AssetData->PackageFlags);
 								NewAssetData.TaggedAssetBundles = AssetData->TaggedAssetBundles;
-								TargetState.UpdateAssetData(AssetData, NewAssetData);
+								TargetState.UpdateAssetData(AssetData, MoveTemp(NewAssetData));
 							}
 						}
 					}
@@ -1126,7 +1125,9 @@ void FAssetRegistryGenerator::UpdateCollectionAssetData()
 			{
 				TagsAndValues.Add(CollectionTagName, FString(CollectionValue));
 			}
-			State.UpdateAssetData(FAssetData(AssetData->PackageName, AssetData->PackagePath, AssetData->AssetName, AssetData->AssetClassPath, TagsAndValues, AssetData->ChunkIDs, AssetData->PackageFlags));
+			FAssetData NewAssetData(*AssetData);
+			NewAssetData.TagsAndValues = FAssetDataTagMapSharedView(MoveTemp(TagsAndValues));
+			State.UpdateAssetData(const_cast<FAssetData*>(AssetData), MoveTemp(NewAssetData));
 		}
 	}
 }
@@ -1448,6 +1449,17 @@ void FAssetRegistryGenerator::AddAssetToFileOrderRecursive(const FName& InPackag
 	}
 }
 
+void FAssetRegistryGenerator::UpdateAssetDatas(bool bForceNoFilter)
+{
+	FAssetRegistrySerializationOptions SaveOptions;
+	AssetRegistry.InitializeSerializationOptions(SaveOptions, TargetPlatform->IniPlatformName());
+	if (bForceNoFilter)
+	{
+		SaveOptions.DisableFilters();
+	}
+	AssetRegistry.InitializeTemporaryAssetRegistryState(State, SaveOptions, true);
+}
+
 bool FAssetRegistryGenerator::SaveAssetRegistry(const FString& SandboxPath, bool bSerializeDevelopmentAssetRegistry, bool bForceNoFilter)
 {
 	UE_LOG(LogAssetRegistryGenerator, Display, TEXT("Saving asset registry v%d."), FAssetRegistryVersion::Type::LatestVersion);
@@ -1469,9 +1481,6 @@ bool FAssetRegistryGenerator::SaveAssetRegistry(const FString& SandboxPath, bool
 		SaveOptions.DisableFilters();
 	}
 
-	// First flush the asset registry and make sure the asset data is in sync, as it may have been updated during cook
-	AssetRegistry.Tick(-1.0f);
-	AssetRegistry.InitializeTemporaryAssetRegistryState(State, SaveOptions, true);
 	// Then possibly apply previous AssetData and AssetPackageData for packages kept from a previous cook
 	UpdateKeptPackages();
 	UpdateCollectionAssetData();
diff --git a/Engine/Source/Editor/UnrealEd/Private/CookOnTheFlyServer.cpp b/Engine/Source/Editor/UnrealEd/Private/CookOnTheFlyServer.cpp
--- a/Engine/Source/Editor/UnrealEd/Private/CookOnTheFlyServer.cpp
+++ b/Engine/Source/Editor/UnrealEd/Private/CookOnTheFlyServer.cpp
@@ -7323,6 +7323,9 @@ void UCookOnTheFlyServer::CookByTheBookFinished()
 			UE_SCOPED_HIERARCHICAL_COOKTIMER(SavingAssetRegistry);
 			SCOPED_BOOT_TIMING("SavingAssetRegistry");
 
+			// Flush the asset registry to get any remaining updates from late loads
+			AssetRegistry->Tick(-1.0f);
+
 			RegisterLocalizationChunkDataGenerator();
 			RegisterShaderChunkDataGenerator();
 
@@ -7405,6 +7408,9 @@ void UCookOnTheFlyServer::CookByTheBookFinished()
 					}
 				}
 				
+				// Sync the Generator's AssetRegistryState with any changes to AssetDatas made during the cook
+				Generator.UpdateAssetDatas(bForceNoFilterAssetsFromAssetRegistry);
+
 				// Add the package hashes to the relevant AssetPackageDatas.
 				// PackageHashes are gated by requiring UPackage::WaitForAsyncFileWrites(), which is called above.
 				FCookSavePackageContext& SaveContext = FindOrCreateSaveContext(TargetPlatform);
diff --git a/Engine/Source/Editor/UnrealEd/Public/Commandlets/AssetRegistryGenerator.h b/Engine/Source/Editor/UnrealEd/Public/Commandlets/AssetRegistryGenerator.h
--- a/Engine/Source/Editor/UnrealEd/Public/Commandlets/AssetRegistryGenerator.h
+++ b/Engine/Source/Editor/UnrealEd/Public/Commandlets/AssetRegistryGenerator.h
@@ -112,6 +112,9 @@ public:
 	 */
 	void RegisterChunkDataGenerator(TSharedRef<IChunkDataGenerator> InChunkDataGenerator);
 
+	/** Called before PreSave; copy AssetData changes from the global AssetRegistry that were made during the cook. */
+	void UpdateAssetDatas(bool bForceNoFilter);
+
 	/**
 	* PreSave
 	* Notify generator that we are about to save the registry and chunk manifests
diff --git a/Engine/Source/Runtime/AssetRegistry/Private/AssetRegistry.cpp b/Engine/Source/Runtime/AssetRegistry/Private/AssetRegistry.cpp
--- a/Engine/Source/Runtime/AssetRegistry/Private/AssetRegistry.cpp
+++ b/Engine/Source/Runtime/AssetRegistry/Private/AssetRegistry.cpp
@@ -4347,7 +4347,7 @@ void FAssetRegistryImpl::AssetSearchDataGathered(Impl::FEventContext& EventConte
 		CA_ASSUME(BackgroundResult.Get() != nullptr);
 
 		// Try to update any asset data that may already exist
-		FAssetData* AssetData = State.CachedAssetsByObjectPath.FindRef(BackgroundResult->ObjectPath);
+		FAssetData* ExistingAssetData = State.CachedAssetsByObjectPath.FindRef(BackgroundResult->ObjectPath);
 
 		const FName PackagePath = BackgroundResult->PackagePath;
 
@@ -4362,10 +4362,10 @@ void FAssetRegistryImpl::AssetSearchDataGathered(Impl::FEventContext& EventConte
 			}
 		}
 
-		if (AssetData)
+		if (ExistingAssetData)
 		{
 			// If this ensure fires then we've somehow processed the same result more than once, and that should never happen
-			if (ensure(AssetData != BackgroundResult.Get()))
+			if (ensure(ExistingAssetData != BackgroundResult.Get()))
 			{
 				// If the current AssetData came from a loaded asset, don't overwrite it with the new one from disk; loaded asset is more authoritative because it has run the postload steps
 #if WITH_EDITOR
@@ -4376,7 +4376,7 @@ void FAssetRegistryImpl::AssetSearchDataGathered(Impl::FEventContext& EventConte
 					PostLoadAssetRegistryTags(BackgroundResult.Get());
 #endif
 					// The asset exists in the cache from disk and has not yet been loaded into memory, update it with the new background data
-					UpdateAssetData(EventContext, AssetData, *BackgroundResult);
+					UpdateAssetData(EventContext, ExistingAssetData, MoveTemp(*BackgroundResult));
 				}
 			}
 		}
@@ -4680,43 +4680,48 @@ void FAssetRegistryImpl::AddAssetData(Impl::FEventContext& EventContext, FAssetD
 	}
 }
 
-void FAssetRegistryImpl::UpdateAssetData(Impl::FEventContext& EventContext, FAssetData* AssetData, const FAssetData& NewAssetData)
+void FAssetRegistryImpl::UpdateAssetData(Impl::FEventContext& EventContext, FAssetData* AssetData, FAssetData&& NewAssetData)
 {
 	// Update the class map if updating a blueprint
 	if (ClassGeneratorNames.Contains(AssetData->AssetClassPath))
 	{
 		const FString OldGeneratedClass = AssetData->GetTagValueRef<FString>(FBlueprintTags::GeneratedClassPath);
-		if (!OldGeneratedClass.IsEmpty() && OldGeneratedClass != TEXTVIEW("None"))
+		const FString OldParentClass = AssetData->GetTagValueRef<FString>(FBlueprintTags::ParentClassPath);
+		const FString NewGeneratedClass = NewAssetData.GetTagValueRef<FString>(FBlueprintTags::GeneratedClassPath);
+		const FString NewParentClass = NewAssetData.GetTagValueRef<FString>(FBlueprintTags::ParentClassPath);
+		if (OldGeneratedClass != NewGeneratedClass || OldParentClass != NewParentClass)
 		{
-			const FTopLevelAssetPath OldGeneratedClassName(OldGeneratedClass);
-			if (ensureAlwaysMsgf(!OldGeneratedClassName.IsNull(), TEXT("Short class name used: OldGeneratedClass=%s. Short class names in tags on the Blueprint class should have been converted to path names."), *OldGeneratedClass))
+			if (!OldGeneratedClass.IsEmpty() && OldGeneratedClass != TEXTVIEW("None"))
 			{
-				CachedBPInheritanceMap.Remove(OldGeneratedClassName);
+				const FTopLevelAssetPath OldGeneratedClassName(OldGeneratedClass);
+				if (ensureAlwaysMsgf(!OldGeneratedClassName.IsNull(), TEXT("Short class name used: OldGeneratedClass=%s. Short class names in tags on the Blueprint class should have been converted to path names."), *OldGeneratedClass))
+				{
+					CachedBPInheritanceMap.Remove(OldGeneratedClassName);
 
-				// Invalidate caching because CachedBPInheritanceMap got modified
-				TempCachedInheritanceBuffer.bDirty = true;
+					// Invalidate caching because CachedBPInheritanceMap got modified
+					TempCachedInheritanceBuffer.bDirty = true;
+				}
 			}
-		}
 
-		const FString NewGeneratedClass = NewAssetData.GetTagValueRef<FString>(FBlueprintTags::GeneratedClassPath);
-		const FString NewParentClass = NewAssetData.GetTagValueRef<FString>(FBlueprintTags::ParentClassPath);
-		if (!NewGeneratedClass.IsEmpty() && !NewParentClass.IsEmpty() && NewGeneratedClass != TEXTVIEW("None") && NewParentClass != TEXTVIEW("None"))
-		{
-			const FTopLevelAssetPath NewGeneratedClassName(NewGeneratedClass);
-			const FTopLevelAssetPath NewParentClassName(NewParentClass);
-			if (ensureAlwaysMsgf(!NewGeneratedClassName.IsNull() && !NewParentClassName.IsNull(), TEXT("Short class names used in AddAssetData: GeneratedClass=%s, ParentClass=%s. Short class names in these tags on the Blueprint class should have been converted to path names."), *NewGeneratedClass, *NewParentClass))
+			if (!NewGeneratedClass.IsEmpty() && !NewParentClass.IsEmpty() && NewGeneratedClass != TEXTVIEW("None") && NewParentClass != TEXTVIEW("None"))
 			{
-				CachedBPInheritanceMap.Add(NewGeneratedClassName, NewParentClassName);
-			}
+				const FTopLevelAssetPath NewGeneratedClassName(NewGeneratedClass);
+				const FTopLevelAssetPath NewParentClassName(NewParentClass);
+				if (ensureAlwaysMsgf(!NewGeneratedClassName.IsNull() && !NewParentClassName.IsNull(), TEXT("Short class names used in AddAssetData: GeneratedClass=%s, ParentClass=%s. Short class names in these tags on the Blueprint class should have been converted to path names."), *NewGeneratedClass, *NewParentClass))
+				{
+					CachedBPInheritanceMap.Add(NewGeneratedClassName, NewParentClassName);
+				}
 
-			// Invalidate caching because CachedBPInheritanceMap got modified
-			TempCachedInheritanceBuffer.bDirty = true;
+				// Invalidate caching because CachedBPInheritanceMap got modified
+				TempCachedInheritanceBuffer.bDirty = true;
+			}
 		}
 	}
 
-	State.UpdateAssetData(AssetData, NewAssetData);
+	bool bModified;
+	State.UpdateAssetData(AssetData, MoveTemp(NewAssetData), &bModified);
 	
-	if (!ShouldSkipAsset(AssetData->AssetClassPath, AssetData->PackageFlags))
+	if (bModified && !ShouldSkipAsset(AssetData->AssetClassPath, AssetData->PackageFlags))
 	{
 		EventContext.AssetEvents.Emplace(*AssetData, Impl::FEventContext::EEvent::Updated);
 	}
@@ -5111,16 +5116,7 @@ void FAssetRegistryImpl::PushProcessLoadedAssetsBatch(Impl::FEventContext& Event
 		}
 		else
 		{
-			if (NewAssetData.TagsAndValues != (*DataFromGather)->TagsAndValues)
-			{
-				// We need to actually update disk cache
-				UpdateAssetData(EventContext, *DataFromGather, NewAssetData);
-			}
-			else
-			{
-				// Bundle tags might have changed but CachedAssetsByTag is up to date
-				(*DataFromGather)->TaggedAssetBundles = NewAssetData.TaggedAssetBundles;
-			}
+			UpdateAssetData(EventContext, *DataFromGather, MoveTemp(NewAssetData));
 		}
 	}
 
@@ -5961,9 +5957,9 @@ bool FAssetRegistryImpl::SetPrimaryAssetIdForObjectPath(Impl::FEventContext& Eve
 	TagsAndValues.Add(FPrimaryAssetId::PrimaryAssetTypeTag, PrimaryAssetId.PrimaryAssetType.ToString());
 	TagsAndValues.Add(FPrimaryAssetId::PrimaryAssetNameTag, PrimaryAssetId.PrimaryAssetName.ToString());
 
-	FAssetData NewAssetData = FAssetData(AssetData->PackageName, AssetData->PackagePath, AssetData->AssetName, AssetData->AssetClassPath, TagsAndValues, AssetData->ChunkIDs, AssetData->PackageFlags);
-	NewAssetData.TaggedAssetBundles = AssetData->TaggedAssetBundles;
-	UpdateAssetData(EventContext, AssetData, NewAssetData);
+	FAssetData NewAssetData(*AssetData);
+	NewAssetData.TagsAndValues = FAssetDataTagMapSharedView(MoveTemp(TagsAndValues));
+	UpdateAssetData(EventContext, AssetData, MoveTemp(NewAssetData));
 
 	return true;
 }
diff --git a/Engine/Source/Runtime/AssetRegistry/Private/AssetRegistryImpl.h b/Engine/Source/Runtime/AssetRegistry/Private/AssetRegistryImpl.h
--- a/Engine/Source/Runtime/AssetRegistry/Private/AssetRegistryImpl.h
+++ b/Engine/Source/Runtime/AssetRegistry/Private/AssetRegistryImpl.h
@@ -284,7 +284,7 @@ private:
 	void AddAssetData(Impl::FEventContext& EventContext, FAssetData* AssetData);
 
 	/** Updates an existing asset data with the new value and updates lookup maps */
-	void UpdateAssetData(Impl::FEventContext& EventContext, FAssetData* AssetData, const FAssetData& NewAssetData);
+	void UpdateAssetData(Impl::FEventContext& EventContext, FAssetData* AssetData, FAssetData&& NewAssetData);
 
 	/** Removes the asset data from the lookup maps */
 	bool RemoveAssetData(Impl::FEventContext& EventContext, FAssetData* AssetData);
diff --git a/Engine/Source/Runtime/AssetRegistry/Private/AssetRegistryState.cpp b/Engine/Source/Runtime/AssetRegistry/Private/AssetRegistryState.cpp
--- a/Engine/Source/Runtime/AssetRegistry/Private/AssetRegistryState.cpp
+++ b/Engine/Source/Runtime/AssetRegistry/Private/AssetRegistryState.cpp
@@ -26,7 +26,6 @@
 #include "UObject/NameBatchSerialization.h"
 #include "UObject/PrimaryAssetId.h"
 
-
 FAssetRegistryState& FAssetRegistryState::operator=(FAssetRegistryState&& Rhs)
 {
 	Reset();
@@ -360,23 +359,14 @@ void FAssetRegistryState::InitializeFromExisting(const TMap<FName, FAssetData*>&
 		
 		if (ExistingData)
 		{
-			// Bundle tags might have changed even if other tags haven't
-			ExistingData->TaggedAssetBundles = AssetData.TaggedAssetBundles;
-
-			// If tags have changed we need to update CachedAssetsByTag
-			if (LocalTagsAndValues != ExistingData->TagsAndValues)
-			{
-				FAssetData TempData = *ExistingData;
-				TempData.TagsAndValues = FAssetDataTagMapSharedView(MoveTemp(LocalTagsAndValues));
-				UpdateAssetData(ExistingData, TempData);
-			}
+			FAssetData NewData(AssetData);
+			NewData.TagsAndValues = FAssetDataTagMapSharedView(MoveTemp(LocalTagsAndValues));
+			UpdateAssetData(ExistingData, MoveTemp(NewData));
 		}
 		else
 		{
-			FAssetData* NewData = new FAssetData(AssetData.PackageName, AssetData.PackagePath, AssetData.AssetName,
-				AssetData.AssetClassPath, LocalTagsAndValues, AssetData.ChunkIDs, AssetData.PackageFlags);
-			NewData->TaggedAssetBundles = AssetData.TaggedAssetBundles;
-
+			FAssetData* NewData = new FAssetData(AssetData);
+			NewData->TagsAndValues = FAssetDataTagMapSharedView(MoveTemp(LocalTagsAndValues));
 			AddAssetData(NewData);
 		}
 	}
@@ -1713,27 +1703,19 @@ void FAssetRegistryState::UpdateAssetData(const FAssetData& NewAssetData)
 	}
 }
 
-void FAssetRegistryState::UpdateAssetData(FAssetData* AssetData, const FAssetData& NewAssetData)
+void FAssetRegistryState::UpdateAssetData(FAssetData* AssetData, const FAssetData& NewAssetData, bool* bOutModified)
 {
-	// Determine if tags need to be remapped
-	bool bTagsChanged = AssetData->TagsAndValues.Num() != NewAssetData.TagsAndValues.Num();
+	UpdateAssetData(AssetData, FAssetData(NewAssetData), bOutModified);
+}
 
-	// If the old and new asset data has the same number of tags, see if any are different (its ok if values are different)
-	if (!bTagsChanged)
-	{
-		for (auto TagIt = AssetData->TagsAndValues.CreateConstIterator(); TagIt; ++TagIt)
-		{
-			if (!NewAssetData.TagsAndValues.Contains(TagIt.Key()))
-			{
-				bTagsChanged = true;
-				break;
-			}
-		}
-	}
+void FAssetRegistryState::UpdateAssetData(FAssetData* AssetData, FAssetData&& NewAssetData, bool* bOutModified)
+{
+	bool bKeyFieldIsModified = false;
 
 	// Update ObjectPath
 	if (AssetData->ObjectPath != NewAssetData.ObjectPath)
 	{
+		bKeyFieldIsModified = true;
 		int32 NumRemoved = CachedAssetsByObjectPath.Remove(AssetData->ObjectPath);
 		check(NumRemoved <= 1);
 		if (NumRemoved == 0)
@@ -1759,6 +1741,7 @@ void FAssetRegistryState::UpdateAssetData(FAssetData* AssetData, const FAssetDat
 	// Update PackageName
 	if (AssetData->PackageName != NewAssetData.PackageName)
 	{
+		bKeyFieldIsModified = true;
 		TArray<FAssetData*, TInlineAllocator<1>>& NewPackageAssets = CachedAssetsByPackageName.FindOrAdd(NewAssetData.PackageName);
 		TArray<FAssetData*, TInlineAllocator<1>>* OldPackageAssets = CachedAssetsByPackageName.Find(AssetData->PackageName);
 
@@ -1769,6 +1752,7 @@ void FAssetRegistryState::UpdateAssetData(FAssetData* AssetData, const FAssetDat
 	// Update PackagePath
 	if (AssetData->PackagePath != NewAssetData.PackagePath)
 	{
+		bKeyFieldIsModified = true;
 		TArray<FAssetData*>& NewPathAssets = CachedAssetsByPath.FindOrAdd(NewAssetData.PackagePath);
 		TArray<FAssetData*>* OldPathAssets = CachedAssetsByPath.Find(AssetData->PackagePath);
 
@@ -1776,9 +1760,12 @@ void FAssetRegistryState::UpdateAssetData(FAssetData* AssetData, const FAssetDat
 		NewPathAssets.Add(AssetData);
 	}
 
+	// AssetName is not a keyfield; compared below
+
 	// Update AssetClass
 	if (AssetData->AssetClassPath != NewAssetData.AssetClassPath)
 	{
+		bKeyFieldIsModified = true;
 		TArray<FAssetData*>& NewClassAssets = CachedAssetsByClass.FindOrAdd(NewAssetData.AssetClassPath);
 		TArray<FAssetData*>* OldClassAssets = CachedAssetsByClass.Find(AssetData->AssetClassPath);
 
@@ -1786,9 +1773,12 @@ void FAssetRegistryState::UpdateAssetData(FAssetData* AssetData, const FAssetDat
 		NewClassAssets.Add(AssetData);
 	}
 
+	// PackageFlags is not a keyfield; compared below
+
 	// Update Tags
-	if (bTagsChanged)
-	{
+	if (AssetData->TagsAndValues != NewAssetData.TagsAndValues)
+	{ 
+		bKeyFieldIsModified = true;
 		for (auto TagIt = AssetData->TagsAndValues.CreateConstIterator(); TagIt; ++TagIt)
 		{
 			const FName FNameKey = TagIt.Key();
@@ -1812,8 +1802,29 @@ void FAssetRegistryState::UpdateAssetData(FAssetData* AssetData, const FAssetDat
 		}
 	}
 
+	// TaggedAssetBundles is not a keyfield; compared below
+
+	// ChunkIds is not a keyfield; compared below
+
+	if (bOutModified)
+	{
+		// Computing equality is expensive; if the caller needs to know it, do cheap compares first 
+		// so we can skip the more expensive compares if the inequality is already known
+		// This is not possible for keyfields - we have to take action on those even if inequality is already known -
+		// so we start with whether bKeyFieldIsModified
+		*bOutModified = bKeyFieldIsModified ||
+			AssetData->AssetName != NewAssetData.AssetName ||
+			AssetData->PackageFlags != NewAssetData.PackageFlags ||
+			AssetData->ChunkIDs != NewAssetData.ChunkIDs ||
+			(AssetData->TaggedAssetBundles.IsValid() != NewAssetData.TaggedAssetBundles.IsValid() ||
+				(AssetData->TaggedAssetBundles.IsValid() &&
+					AssetData->TaggedAssetBundles.Get() != NewAssetData.TaggedAssetBundles.Get() && // First check whether the pointers are the same
+					*AssetData->TaggedAssetBundles != *NewAssetData.TaggedAssetBundles // If the pointers differ, check whether the contents differ
+					));
+	}
+
 	// Copy in new values
-	*AssetData = NewAssetData;
+	*AssetData = MoveTemp(NewAssetData);
 }
 
 bool FAssetRegistryState::UpdateAssetDataPackageFlags(FName PackageName, uint32 PackageFlags)
diff --git a/Engine/Source/Runtime/AssetRegistry/Public/AssetRegistry/AssetRegistryState.h b/Engine/Source/Runtime/AssetRegistry/Public/AssetRegistry/AssetRegistryState.h
--- a/Engine/Source/Runtime/AssetRegistry/Public/AssetRegistry/AssetRegistryState.h
+++ b/Engine/Source/Runtime/AssetRegistry/Public/AssetRegistry/AssetRegistryState.h
@@ -343,7 +343,8 @@ public:
 	void UpdateAssetData(const FAssetData& NewAssetData);
 
 	/** Updates an existing asset data with the new value and updates lookup maps */
-	void UpdateAssetData(FAssetData* AssetData, const FAssetData& NewAssetData);
+	void UpdateAssetData(FAssetData* AssetData, const FAssetData& NewAssetData, bool* bOutModified = nullptr);
+	void UpdateAssetData(FAssetData* AssetData, FAssetData&& NewAssetData, bool* bOutModified = nullptr);
 
 	/**
 	 * Updates all asset data package flags in the specified package
diff --git a/Engine/Source/Runtime/CoreUObject/Private/AssetRegistry/AssetDataTagMap.cpp b/Engine/Source/Runtime/CoreUObject/Private/AssetRegistry/AssetDataTagMap.cpp
--- a/Engine/Source/Runtime/CoreUObject/Private/AssetRegistry/AssetDataTagMap.cpp
+++ b/Engine/Source/Runtime/CoreUObject/Private/AssetRegistry/AssetDataTagMap.cpp
@@ -459,7 +459,7 @@ namespace FixedTagPrivate
 	{
 		TCHAR Buffer[FName::StringBufferSize];
 		uint32 Len = B.ToString(Buffer);
-		return Len == A.Len() && FPlatformString::Strnicmp(A.GetData(), Buffer, Len);
+		return Len == A.Len() && FPlatformString::Strnicmp(A.GetData(), Buffer, Len) == 0;
 	}
 
 	static bool EqualsInsensitive(const FStringView& Str, const FNumberlessExportPath& Path)
@@ -480,7 +480,12 @@ namespace FixedTagPrivate
 		case EValueType::Name:					return EqualsInsensitive(Str, Store.Names[Index]);
 		case EValueType::NumberlessExportPath:	return EqualsInsensitive(Str, Store.NumberlessExportPaths[Index]);
 		case EValueType::ExportPath:			return EqualsInsensitive(Str, Store.ExportPaths[Index]);
-		case EValueType::LocalizedText:			return EqualsInsensitive(Str, *Store.Texts[Index].ToString());
+		case EValueType::LocalizedText:
+		{
+			FText StrAsText;
+			FromComplexString(FString(Str), StrAsText);
+			return EqualsInsensitive(StrAsText, Store.Texts[Index]);
+		}
 		}
 
 		check(false);
@@ -612,7 +617,6 @@ namespace FixedTagPrivate
 
 		if ((A.Num == B.Num) & (A.HasNumberlessKeys == B.HasNumberlessKeys))
 		{
-			check(A.StoreIndex != B.StoreIndex);
 			const FStore& StoreA = GStores[A.StoreIndex];
 			const FStore& StoreB = GStores[B.StoreIndex];
 
@@ -1463,8 +1467,32 @@ bool operator==(const FAssetDataTagMapSharedView& A, const FAssetDataTagMapShare
 	}
 	else if (A.IsFixed() != B.IsFixed())
 	{
-		// This is very wasteful but currently only used by unit tests
-		return A.IsFixed() ? *B.Loose == A.CopyMap() : *A.Loose == B.CopyMap();
+		FixedTagPrivate::FMapHandle Fixed;
+		FAssetDataTagMap* Loose;
+		if (A.IsFixed())
+		{
+			Fixed = A.Fixed;
+			Loose = B.Loose;
+		}
+		else
+		{
+			Fixed = B.Fixed;
+			Loose = A.Loose;
+		}
+		// Since the num is the same and these are unique maps, we only have to test whether all keys in Fixed exist with equal value in Loose
+		// Once we've done that, we don't have to test whether any keys in Loose are missing from Fixed.
+		bool bEqual = true;
+		Fixed.ForEachPair([&](FixedTagPrivate::FNumberedPair FixedPair)
+			{
+				FString* LooseValue = Loose->Find(FixedPair.Key);
+				if (!LooseValue)
+				{
+					bEqual = false;
+					return;
+				}
+				bEqual &= (FAssetTagValueRef(Fixed.StoreIndex, FixedPair.Value) == *LooseValue);
+			});
+		return bEqual;
 	}
 	else if (A.IsFixed())
 	{
