AssetRegistry: Shutdown the Gatherer thread before the engine shuts down, to prevent accessing UObjectRedirector::StaticClass while it is being torn down.
#rnx

diff --git a/Engine/Source/Runtime/AssetRegistry/Private/AssetRegistry.cpp b/Engine/Source/Runtime/AssetRegistry/Private/AssetRegistry.cpp
--- a/Engine/Source/Runtime/AssetRegistry/Private/AssetRegistry.cpp
+++ b/Engine/Source/Runtime/AssetRegistry/Private/AssetRegistry.cpp
@@ -868,6 +868,8 @@ void UAssetRegistryImpl::InitializeEvents(UE::AssetRegistry::Impl::FInitializeCo
 	}
 #endif // WITH_EDITOR
 
+	FCoreDelegates::OnEnginePreExit.AddUObject(this, &UAssetRegistryImpl::OnEnginePreExit);
+
 	// Listen for new content paths being added or removed at runtime.  These are usually plugin-specific asset paths that
 	// will be loaded a bit later on.
 	FPackageName::OnContentPathMounted().AddUObject(this, &UAssetRegistryImpl::OnContentPathMounted);
@@ -1305,6 +1307,12 @@ void FAssetRegistryImpl::RefreshNativeClasses()
 
 }
 
+void UAssetRegistryImpl::OnEnginePreExit()
+{
+	FWriteScopeLock InterfaceScopeLock(InterfaceLock);
+	GuardedData.OnEnginePreExit();
+}
+
 UAssetRegistryImpl::~UAssetRegistryImpl()
 {
 	FWriteScopeLock InterfaceScopeLock(InterfaceLock);
@@ -1313,6 +1321,7 @@ UAssetRegistryImpl::~UAssetRegistryImpl()
 	FPackageName::OnContentPathMounted().RemoveAll(this);
 	FPackageName::OnContentPathDismounted().RemoveAll(this);
 	FCoreDelegates::OnPostEngineInit.RemoveAll(this);
+	FCoreDelegates::OnEnginePreExit.RemoveAll(this);
 	IPluginManager::Get().OnLoadingPhaseComplete().RemoveAll(this);
 
 #if WITH_EDITOR
@@ -1379,6 +1388,12 @@ UAssetRegistryImpl& UAssetRegistryImpl::Get()
 namespace UE::AssetRegistry
 {
 
+void FAssetRegistryImpl::OnEnginePreExit()
+{
+	// Shut down the GlobalGatherer's gather threads, before we start tearing down the engine
+	GlobalGatherer.Reset();
+}
+
 void FAssetRegistryImpl::ConstructGatherer()
 {
 	if (GlobalGatherer.IsValid())
diff --git a/Engine/Source/Runtime/AssetRegistry/Private/AssetRegistry.h b/Engine/Source/Runtime/AssetRegistry/Private/AssetRegistry.h
--- a/Engine/Source/Runtime/AssetRegistry/Private/AssetRegistry.h
+++ b/Engine/Source/Runtime/AssetRegistry/Private/AssetRegistry.h
@@ -170,6 +170,7 @@ protected:
 	virtual bool SetPrimaryAssetIdForObjectPath(const FName ObjectPath, FPrimaryAssetId PrimaryAssetId) override;
 
 private:
+	void OnEnginePreExit();
 	void InitializeEvents(UE::AssetRegistry::Impl::FInitializeContext& Context);
 	void Broadcast(UE::AssetRegistry::Impl::FEventContext& EventContext);
 
diff --git a/Engine/Source/Runtime/AssetRegistry/Private/AssetRegistryImpl.h b/Engine/Source/Runtime/AssetRegistry/Private/AssetRegistryImpl.h
--- a/Engine/Source/Runtime/AssetRegistry/Private/AssetRegistryImpl.h
+++ b/Engine/Source/Runtime/AssetRegistry/Private/AssetRegistryImpl.h
@@ -120,7 +120,7 @@ public:
 	FAssetRegistryImpl();
 	/** Construct the AssetRegistryImpl, including initial scans if applicable. */
 	void Initialize(Impl::FInitializeContext& Context);
-
+	void OnEnginePreExit();
 
 	// Helpers for functions of the same name from UAssetRegistryImpl
 
