Fixed asset registry ScanModifiedAssetFiles : It didn't process queued loaded assets.

diff --git a/Engine/Source/Runtime/AssetRegistry/Private/AssetRegistry.cpp b/Engine/Source/Runtime/AssetRegistry/Private/AssetRegistry.cpp
--- a/Engine/Source/Runtime/AssetRegistry/Private/AssetRegistry.cpp
+++ b/Engine/Source/Runtime/AssetRegistry/Private/AssetRegistry.cpp
@@ -5170,6 +5170,15 @@ void UAssetRegistryImpl::ScanModifiedAssetFiles(const TArray<FString>& InFilePat
 		FWriteScopeLock InterfaceScopeLock(InterfaceLock);
 		GuardedData.ScanModifiedAssetFiles(EventContext, InFilePaths);
 	}
+
+#if WITH_EDITOR
+	// Our caller expects up to date results after calling this function,
+	// but in-memory results will override the on-disk results we just scanned,
+	// and our in-memory results might be out of date due to being queued but not yet processed.
+	// So ProcessLoadedAssetsToUpdateCache before returning to make sure results are up to date.
+	ProcessLoadedAssetsToUpdateCache(EventContext, -1., UE::AssetRegistry::Impl::EGatherStatus::Complete);
+#endif
+
 	Broadcast(EventContext);
 }
 
