AssetRegistry: Downgrade log if AssetRegistry fails to load in runtime from Fatal to Error. Add more feedback about why it failed.
#rnx

diff --git a/Engine/Source/Runtime/AssetRegistry/Private/AssetRegistry.cpp b/Engine/Source/Runtime/AssetRegistry/Private/AssetRegistry.cpp
--- a/Engine/Source/Runtime/AssetRegistry/Private/AssetRegistry.cpp
+++ b/Engine/Source/Runtime/AssetRegistry/Private/AssetRegistry.cpp
@@ -169,6 +169,16 @@ TArray<FString, TInlineAllocator<2>> GetPriorityPaths()
 	return Paths;
 }
 
+enum class ELoadResult : uint8
+{
+	Succeeded = 0,
+	NotFound = 1,
+	FailedToLoad = 2,
+	Inactive = 3,
+	AlreadyConsumed = 4,
+	UninitializedMemberLoadResult = 5,
+};
+
 // Loads cooked AssetRegistry.bin using an async preload task if available and sync otherwise
 class FPreloader
 {
@@ -179,7 +189,7 @@ public:
 		Failed,
 		Deferred
 	};
-	using FConsumeFunction = TFunction<void(bool bSucceeded, FAssetRegistryState&& ARState)>;
+	using FConsumeFunction = TFunction<void(ELoadResult LoadResult, FAssetRegistryState&& ARState)>;
 
 	FPreloader()
 	{
@@ -268,7 +278,7 @@ private:
 		return false;
 	}
 
-	bool TryLoad()
+	ELoadResult TryLoad()
 	{
 		checkf(!ARPath.IsEmpty(), TEXT("TryLoad must not be called until after TrySetPath has succeeded."));
 
@@ -276,10 +286,11 @@ private:
 		const int32 ThreadReduction = 2; // This thread + main thread already has work to do 
 		int32 MaxWorkers = CanLoadAsync() ? FPlatformMisc::NumberOfCoresIncludingHyperthreads() - ThreadReduction : 0;
 		Options.ParallelWorkers = FMath::Clamp(MaxWorkers, 0, 16);
-		bLoadSucceeded = FAssetRegistryState::LoadFromDisk(*ARPath, Options, Payload);
+		bool bLoadSucceeded = FAssetRegistryState::LoadFromDisk(*ARPath, Options, Payload);
 		UE_CLOG(!bLoadSucceeded, LogAssetRegistry, Warning, TEXT("Premade AssetRegistry path %s existed but failed to load."), *ARPath);
 		UE_CLOG(bLoadSucceeded, LogAssetRegistry, Log, TEXT("Premade AssetRegistry loaded from '%s'"), *ARPath);
-		return bLoadSucceeded;
+		LoadResult = bLoadSucceeded ? ELoadResult::Succeeded : ELoadResult::FailedToLoad;
+		return LoadResult;
 	}
 
 	void DelayedInitialize()
@@ -337,7 +348,7 @@ private:
 		// This function is active only after State has been set to Loading and PreloadReady has been Reset
 		// Until this function triggers PreloadReady, it has exclusive ownership of bLoadSucceeded and Payload
 		// Load outside the lock so that ConsumeOrDefer does not have to wait for the Load before it can defer and exit
-		bool bSucceeded = TryLoad();
+		ELoadResult LocalResult = TryLoad();
 		// Trigger outside the lock so that a locked Consume function that is waiting on PreloadReady can wait inside the lock.
 		PreloadReady->Trigger();
 
@@ -361,7 +372,7 @@ private:
 		{
 			// No further threads will read/write payload at this point (until destructor, which is called after all async threads are complete
 			// so we can use it outside the lock
-			LocalConsumeCallback(bSucceeded, MoveTemp(Payload));
+			LocalConsumeCallback(LocalResult, MoveTemp(Payload));
 			Shutdown();
 		}
 	}
@@ -375,18 +386,19 @@ private:
 		if (LoadState == EState::WillNeverPreload || LoadState == EState::Consumed || ConsumeCallback)
 		{
 			Lock.Unlock(); // Unlock before calling external code in Consume callback
-			ConsumeSynchronous(false, FAssetRegistryState());
+			ELoadResult LocalResult = (LoadState == EState::Consumed || ConsumeCallback) ? ELoadResult::AlreadyConsumed : ELoadResult::Inactive;
+			ConsumeSynchronous(LocalResult, FAssetRegistryState());
 			return EConsumeResult::Failed;
 		}
 
 		if (LoadState == EState::LoadSynchronous)
 		{
-			bool bSucceeded = TrySetPath() && TryLoad();
+			ELoadResult LocalResult = TrySetPath() ? TryLoad() : ELoadResult::NotFound;
 			LoadState = EState::Consumed;
 			Lock.Unlock(); // Unlock before calling external code in Consume callback
-			ConsumeSynchronous(bSucceeded, MoveTemp(Payload));
+			ConsumeSynchronous(LocalResult, MoveTemp(Payload));
 			Shutdown(); // Shutdown can be called outside the lock since AsyncThread doesn't exist
-			return bSucceeded ? EConsumeResult::Succeeded : EConsumeResult::Failed;
+			return LocalResult == ELoadResult::Succeeded ? EConsumeResult::Succeeded : EConsumeResult::Failed;
 		}
 
 		// Cancel any further searching in Paks since we will no longer accept preloads starting after this point
@@ -408,14 +420,14 @@ private:
 
 		// TryAsyncLoad might not yet have set state to Loaded
 		check(LoadState == EState::Loaded || LoadState == EState::Loading || LoadState == EState::NotFound);
-		bool bSucceeded = LoadState == EState::NotFound ? false : bLoadSucceeded;
+		ELoadResult LocalResult = LoadState == EState::NotFound ? ELoadResult::NotFound : LoadResult;
 		LoadState = EState::Consumed;
 
 		// No further async threads exist that will read/write payload at this point so we can use it outside the lock
 		Lock.Unlock(); // Unlock before calling external code in Consume callback
-		ConsumeSynchronous(bSucceeded, MoveTemp(Payload));
+		ConsumeSynchronous(LocalResult, MoveTemp(Payload));
 		Shutdown(); // Shutdown can be called outside the lock since we have set state to Consumed and the Async thread will notice and exit
-		return bSucceeded ? EConsumeResult::Succeeded : EConsumeResult::Failed;;
+		return LocalResult == ELoadResult::Succeeded ? EConsumeResult::Succeeded : EConsumeResult::Failed;
 	}
 
 	/** Called when the Preloader has no further work to do, to free resources early since destruction occurs at end of process. */
@@ -460,7 +472,7 @@ private:
 	EState LoadState = EState::WillNeverPreload;
 
 	/** Result of TryLoad.Thread ownership rules are the same as the rules for Payload. */
-	bool bLoadSucceeded = false;
+	ELoadResult LoadResult = ELoadResult::UninitializedMemberLoadResult;
 }
 GPreloader;
 
@@ -509,10 +521,10 @@ void FAsyncConsumer::Wait(UAssetRegistryImpl& UARI, FWriteScopeLock& ScopeLock)
 	}
 }
 
-void FAsyncConsumer::Consume(UAssetRegistryImpl& UARI, UE::AssetRegistry::Impl::FEventContext& EventContext, bool bSucceeded, FAssetRegistryState&& ARState)
+void FAsyncConsumer::Consume(UAssetRegistryImpl& UARI, UE::AssetRegistry::Impl::FEventContext& EventContext, ELoadResult LoadResult, FAssetRegistryState&& ARState)
 {
 	// Called within the lock
-	UARI.GuardedData.LoadPremadeAssetRegistry(EventContext, bSucceeded, MoveTemp(ARState), FAssetRegistryState::EInitializationMode::OnlyUpdateNew);
+	UARI.GuardedData.LoadPremadeAssetRegistry(EventContext, LoadResult, MoveTemp(ARState), FAssetRegistryState::EInitializationMode::OnlyUpdateNew);
 	check(ReferenceCount >= 1);
 	check(Consumed != nullptr);
 	Consumed->Trigger();
@@ -550,20 +562,20 @@ void FAssetRegistryImpl::ConsumeOrDeferPreloadedPremade(UAssetRegistryImpl& UARI
 		return;
 	}
 #if ASSETREGISTRY_ENABLE_PREMADE_REGISTRY_IN_EDITOR
-	FPreloader::FConsumeFunction ConsumeFromAsyncThread = [this, &UARI](bool bSucceeded, FAssetRegistryState&& ARState)
+	FPreloader::FConsumeFunction ConsumeFromAsyncThread = [this, &UARI](Premade::ELoadResult LoadResult, FAssetRegistryState&& ARState)
 	{
 		Impl::FEventContext EventContext;
 		{
 			FWriteScopeLock InterfaceScopeLock(UARI.InterfaceLock);
-			AsyncConsumer.Consume(UARI, EventContext, bSucceeded, MoveTemp(ARState));
+			AsyncConsumer.Consume(UARI, EventContext, LoadResult, MoveTemp(ARState));
 		}
 		UARI.Broadcast(EventContext);
 	};
-	auto ConsumeOnCurrentThread = [ConsumeFromAsyncThread](bool bSucceeded, FAssetRegistryState&& ARState) mutable
+	auto ConsumeOnCurrentThread = [ConsumeFromAsyncThread](Premade::ELoadResult LoadResult, FAssetRegistryState&& ARState) mutable
 	{
-		Async(EAsyncExecution::TaskGraph, [bSucceeded, ARState=MoveTemp(ARState), ConsumeFromAsyncThread=MoveTemp(ConsumeFromAsyncThread)]() mutable
+		Async(EAsyncExecution::TaskGraph, [LoadResult, ARState=MoveTemp(ARState), ConsumeFromAsyncThread=MoveTemp(ConsumeFromAsyncThread)]() mutable
 		{
-			ConsumeFromAsyncThread(bSucceeded, MoveTemp(ARState));
+			ConsumeFromAsyncThread(LoadResult, MoveTemp(ARState));
 		});
 	};
 	if (CanConsumeAsync())
@@ -574,9 +586,9 @@ void FAssetRegistryImpl::ConsumeOrDeferPreloadedPremade(UAssetRegistryImpl& UARI
 	else
 #endif
 	{
-		GPreloader.Consume([this, &EventContext](bool bSucceeded, FAssetRegistryState&& ARState)
+		GPreloader.Consume([this, &EventContext](Premade::ELoadResult LoadResult, FAssetRegistryState&& ARState)
 			{
-				LoadPremadeAssetRegistry(EventContext, bSucceeded, MoveTemp(ARState), FAssetRegistryState::EInitializationMode::Rebuild);
+				LoadPremadeAssetRegistry(EventContext, LoadResult, MoveTemp(ARState), FAssetRegistryState::EInitializationMode::Rebuild);
 			});
 	}
 }
@@ -654,12 +666,12 @@ FAssetRegistryImpl::FAssetRegistryImpl()
 }
 
 void FAssetRegistryImpl::LoadPremadeAssetRegistry(Impl::FEventContext& EventContext,
-	bool bSucceeded, FAssetRegistryState&& ARState, FAssetRegistryState::EInitializationMode Mode)
+	Premade::ELoadResult LoadResult, FAssetRegistryState&& ARState, FAssetRegistryState::EInitializationMode Mode)
 {
 	SCOPED_BOOT_TIMING("LoadPremadeAssetRegistry");
 	if (SerializationOptions.bSerializeAssetRegistry)
 	{
-		if (bSucceeded)
+		if (LoadResult == Premade::ELoadResult::Succeeded)
 		{
 			if (Mode == FAssetRegistryState::EInitializationMode::Rebuild)
 			{
@@ -675,7 +687,7 @@ void FAssetRegistryImpl::LoadPremadeAssetRegistry(Impl::FEventContext& EventCont
 		else
 		{
 			UE_CLOG(FPlatformProperties::RequiresCookedData() && (IsRunningGame() || IsRunningDedicatedServer()),
-				LogAssetRegistry, Fatal, TEXT("Failed to load premade asset registry."));
+				LogAssetRegistry, Error, TEXT("Failed to load premade asset registry. LoadResult == %d."), static_cast<int32>(LoadResult));
 		}
 	}
 
diff --git a/Engine/Source/Runtime/AssetRegistry/Private/AssetRegistryImpl.h b/Engine/Source/Runtime/AssetRegistry/Private/AssetRegistryImpl.h
--- a/Engine/Source/Runtime/AssetRegistry/Private/AssetRegistryImpl.h
+++ b/Engine/Source/Runtime/AssetRegistry/Private/AssetRegistryImpl.h
@@ -26,7 +26,11 @@ struct FFileChangeData;
 class FPackageReader;
 class UAssetRegistryImpl;
 namespace UE::AssetRegistry::Impl { struct FEventContext; }
+namespace UE::AssetRegistry::Premade { enum class ELoadResult : uint8; }
 namespace UE::AssetRegistry::Premade { struct FAsyncConsumer; }
+namespace UE::AssetRegistry::Impl { struct FInitializeContext; }
+namespace UE::AssetRegistry::Impl { struct FScanPathContext; }
+namespace UE::AssetRegistry::Impl { struct FClassInheritanceContext; }
 
 #if WITH_EDITOR
 namespace UE::AssetDependencyGatherer::Private { class FRegisteredAssetDependencyGatherer; }
@@ -51,7 +55,7 @@ struct FAsyncConsumer
 	 */
 	void Wait(UAssetRegistryImpl& UARI, FWriteScopeLock& ScopeLock);
 	/** Callback from the async thread to consume the Premade ARState */
-	void Consume(UAssetRegistryImpl& UARI, UE::AssetRegistry::Impl::FEventContext& EventContext, bool bSucceeded, FAssetRegistryState&& ARState);
+	void Consume(UAssetRegistryImpl& UARI, UE::AssetRegistry::Impl::FEventContext& EventContext, ELoadResult LoadResult, FAssetRegistryState&& ARState);
 
 private:
 	/**
@@ -65,18 +69,9 @@ private:
 }
 #endif
 
-namespace UE
-{
-namespace AssetRegistry
+namespace UE::AssetRegistry
 {
 
-namespace Impl
-{
-	struct FInitializeContext;
-	struct FScanPathContext;
-	struct FClassInheritanceContext;
-}
-
 namespace Impl
 {
 	/** Container for global class inheritance data; the AssetRegistry may use a persistent buffer or a stack local buffer. */
@@ -312,7 +307,7 @@ private:
 	void ConsumeOrDeferPreloadedPremade(UAssetRegistryImpl& UARI, Impl::FEventContext& EventContext);
 	/** Moves a premade asset registry state into this AR */
 	void LoadPremadeAssetRegistry(Impl::FEventContext& Context,
-		bool bSucceeded, FAssetRegistryState&& ARState, FAssetRegistryState::EInitializationMode Mode);
+		Premade::ELoadResult LoadResult, FAssetRegistryState&& ARState, FAssetRegistryState::EInitializationMode Mode);
 
 private:
 
@@ -583,5 +578,3 @@ void EnumerateMemoryAssets(const FARCompiledFilter& InFilter, TSet<FName>& OutPa
 }
 
 }
-}
-
