Avoid temporary string/name conversion until needed

#rnx

diff --git a/Engine/Source/Runtime/AssetRegistry/Private/AssetRegistry.cpp b/Engine/Source/Runtime/AssetRegistry/Private/AssetRegistry.cpp
--- a/Engine/Source/Runtime/AssetRegistry/Private/AssetRegistry.cpp
+++ b/Engine/Source/Runtime/AssetRegistry/Private/AssetRegistry.cpp
@@ -1865,10 +1865,13 @@ void EnumerateMemoryAssetsHelper(const FARCompiledFilter& InFilter, TSet<FName>&
 
 			// Package path
 			const FString PackageNameStr = InMemoryPackage->GetName();
-			const FName PackagePath = FName(*FPackageName::GetLongPackagePath(PackageNameStr));
-			if (InFilter.PackagePaths.Num() > 0 && !InFilter.PackagePaths.Contains(PackagePath))
+			if (InFilter.PackagePaths.Num() > 0)
 			{
-				return;
+				const FName PackagePath = FName(*FPackageName::GetLongPackagePath(PackageNameStr));
+				if (!InFilter.PackagePaths.Contains(PackagePath))
+				{
+					return;
+				}
 			}
 
 			// Could perhaps save some FName -> String conversions by creating this a bit earlier using the UObject constructor
