AssetRegistryModule: Add TryGet function that can return null instead of asserting during Engine shutdown.
#rnx

diff --git a/Engine/Plugins/Developer/Concert/ConcertSync/ConcertSyncClient/Source/ConcertSyncClient/Private/ConcertClientPackageBridge.cpp b/Engine/Plugins/Developer/Concert/ConcertSync/ConcertSyncClient/Source/ConcertSyncClient/Private/ConcertClientPackageBridge.cpp
--- a/Engine/Plugins/Developer/Concert/ConcertSync/ConcertSyncClient/Source/ConcertSyncClient/Private/ConcertClientPackageBridge.cpp
+++ b/Engine/Plugins/Developer/Concert/ConcertSync/ConcertSyncClient/Source/ConcertSyncClient/Private/ConcertClientPackageBridge.cpp
@@ -95,9 +95,13 @@ FConcertClientPackageBridge::~FConcertClientPackageBridge()
 		// Unregister Asset Registry Events
 		if (FAssetRegistryModule* AssetRegistryModule = FModuleManager::GetModulePtr<FAssetRegistryModule>("AssetRegistry"))
 		{
-			AssetRegistryModule->Get().OnInMemoryAssetCreated().RemoveAll(this);
-			AssetRegistryModule->Get().OnInMemoryAssetDeleted().RemoveAll(this);
-			AssetRegistryModule->Get().OnAssetRenamed().RemoveAll(this);
+			IAssetRegistry* AssetRegistry = AssetRegistryModule->TryGet();
+			if (AssetRegistry)
+			{
+				AssetRegistry->OnInMemoryAssetCreated().RemoveAll(this);
+				AssetRegistry->OnInMemoryAssetDeleted().RemoveAll(this);
+				AssetRegistry->OnAssetRenamed().RemoveAll(this);
+			}
 		}
 
 		// Unregister Map Change Events
diff --git a/Engine/Plugins/Editor/AssetSearch/Source/Private/AssetSearchManager.cpp b/Engine/Plugins/Editor/AssetSearch/Source/Private/AssetSearchManager.cpp
--- a/Engine/Plugins/Editor/AssetSearch/Source/Private/AssetSearchManager.cpp
+++ b/Engine/Plugins/Editor/AssetSearch/Source/Private/AssetSearchManager.cpp
@@ -292,9 +292,13 @@ void FAssetSearchManager::StopScanningAssets()
 {
 	if (FAssetRegistryModule* AssetRegistryModule = FModuleManager::GetModulePtr<FAssetRegistryModule>("AssetRegistry"))
 	{
-		AssetRegistryModule->Get().OnAssetAdded().RemoveAll(this);
-		AssetRegistryModule->Get().OnAssetRemoved().RemoveAll(this);
-		AssetRegistryModule->Get().OnFilesLoaded().RemoveAll(this);
+		IAssetRegistry* AssetRegistry = AssetRegistryModule->TryGet();
+		if (AssetRegistry)
+		{
+			AssetRegistry->OnAssetAdded().RemoveAll(this);
+			AssetRegistry->OnAssetRemoved().RemoveAll(this);
+			AssetRegistry->OnFilesLoaded().RemoveAll(this);
+		}
 	}
 
 	ProcessAssetQueue.Reset();
diff --git a/Engine/Plugins/EnhancedInput/Source/InputEditor/Private/InputCustomizations.cpp b/Engine/Plugins/EnhancedInput/Source/InputEditor/Private/InputCustomizations.cpp
--- a/Engine/Plugins/EnhancedInput/Source/InputEditor/Private/InputCustomizations.cpp
+++ b/Engine/Plugins/EnhancedInput/Source/InputEditor/Private/InputCustomizations.cpp
@@ -94,9 +94,13 @@ FEnhancedInputDeveloperSettingsCustomization::~FEnhancedInputDeveloperSettingsCu
 	// Unregister settings panel listeners
 	if (FAssetRegistryModule* AssetRegistryModule = FModuleManager::GetModulePtr<FAssetRegistryModule>("AssetRegistry"))
 	{
-		AssetRegistryModule->Get().OnAssetAdded().RemoveAll(this);
-		AssetRegistryModule->Get().OnAssetRemoved().RemoveAll(this);
-		AssetRegistryModule->Get().OnAssetRenamed().RemoveAll(this);
+		IAssetRegistry* AssetRegistry = AssetRegistryModule->TryGet();
+		if (AssetRegistry)
+		{
+			AssetRegistry->OnAssetAdded().RemoveAll(this);
+			AssetRegistry->OnAssetRemoved().RemoveAll(this);
+			AssetRegistry->OnAssetRenamed().RemoveAll(this);
+		}
 	}
 	
 	CachedDetailBuilder = nullptr;
diff --git a/Engine/Plugins/Experimental/PCG/Source/PCGEditor/Private/SPCGEditorGraphNodePalette.cpp b/Engine/Plugins/Experimental/PCG/Source/PCGEditor/Private/SPCGEditorGraphNodePalette.cpp
--- a/Engine/Plugins/Experimental/PCG/Source/PCGEditor/Private/SPCGEditorGraphNodePalette.cpp
+++ b/Engine/Plugins/Experimental/PCG/Source/PCGEditor/Private/SPCGEditorGraphNodePalette.cpp
@@ -59,11 +59,14 @@ SPCGEditorGraphNodePalette::~SPCGEditorGraphNodePalette()
 {
 	if (FModuleManager::Get().IsModuleLoaded(TEXT("AssetRegistry")))
 	{
-		FAssetRegistryModule& AssetRegistryModule = FModuleManager::GetModuleChecked<FAssetRegistryModule>(TEXT("AssetRegistry"));
-		AssetRegistryModule.Get().OnAssetAdded().RemoveAll(this);
-		AssetRegistryModule.Get().OnAssetRemoved().RemoveAll(this);
-		AssetRegistryModule.Get().OnAssetUpdated().RemoveAll(this);
-		AssetRegistryModule.Get().OnAssetRenamed().RemoveAll(this);
+		IAssetRegistry* AssetRegistry = FModuleManager::GetModuleChecked<FAssetRegistryModule>(TEXT("AssetRegistry")).TryGet();
+		if (AssetRegistry) // AssetRegistry can be null during EngineShutdown even if module is available
+		{
+			AssetRegistry->OnAssetAdded().RemoveAll(this);
+			AssetRegistry->OnAssetRemoved().RemoveAll(this);
+			AssetRegistry->OnAssetUpdated().RemoveAll(this);
+			AssetRegistry->OnAssetRenamed().RemoveAll(this);
+		}
 	}
 }
 
diff --git a/Engine/Plugins/Experimental/PythonScriptPlugin/Source/PythonScriptPlugin/Private/PythonScriptPlugin.cpp b/Engine/Plugins/Experimental/PythonScriptPlugin/Source/PythonScriptPlugin/Private/PythonScriptPlugin.cpp
--- a/Engine/Plugins/Experimental/PythonScriptPlugin/Source/PythonScriptPlugin/Private/PythonScriptPlugin.cpp
+++ b/Engine/Plugins/Experimental/PythonScriptPlugin/Source/PythonScriptPlugin/Private/PythonScriptPlugin.cpp
@@ -1078,8 +1078,12 @@ void FPythonScriptPlugin::ShutdownPython()
 
 	if (FAssetRegistryModule* AssetRegistryModule = FModuleManager::GetModulePtr<FAssetRegistryModule>("AssetRegistry"))
 	{
-		AssetRegistryModule->Get().OnAssetRenamed().RemoveAll(this);
-		AssetRegistryModule->Get().OnAssetRemoved().RemoveAll(this);
+		IAssetRegistry* AssetRegistry = AssetRegistryModule->TryGet();
+		if (AssetRegistry)
+		{
+			AssetRegistry->OnAssetRenamed().RemoveAll(this);
+			AssetRegistry->OnAssetRemoved().RemoveAll(this);
+		}
 	}
 
 #if WITH_EDITOR
diff --git a/Engine/Plugins/Experimental/RemoteSession/Source/RemoteSessionEditor/Private/Widgets/SRemoteSessionStream.cpp b/Engine/Plugins/Experimental/RemoteSession/Source/RemoteSessionEditor/Private/Widgets/SRemoteSessionStream.cpp
--- a/Engine/Plugins/Experimental/RemoteSession/Source/RemoteSessionEditor/Private/Widgets/SRemoteSessionStream.cpp
+++ b/Engine/Plugins/Experimental/RemoteSession/Source/RemoteSessionEditor/Private/Widgets/SRemoteSessionStream.cpp
@@ -504,7 +504,11 @@ void SRemoteSessionStream::EnabledStreaming(bool bInStreaming)
 			}
 			if (FAssetRegistryModule* AssetRegistryModule = FModuleManager::GetModulePtr<FAssetRegistryModule>(TEXT("AssetRegistry")))
 			{
-				AssetRegistryModule->Get().OnAssetRemoved().RemoveAll(this);
+				IAssetRegistry* AssetRegistry = AssetRegistryModule->TryGet();
+				if (AssetRegistry)
+				{
+					AssetRegistry->OnAssetRemoved().RemoveAll(this);
+				}
 			}
 			FEditorSupportDelegates::PrepareToCleanseEditorObject.RemoveAll(this);
 			GEditor->OnBlueprintPreCompile().RemoveAll(this);
diff --git a/Engine/Plugins/Experimental/VirtualCamera/Source/VirtualCamera/Private/VirtualCameraActor.cpp b/Engine/Plugins/Experimental/VirtualCamera/Source/VirtualCamera/Private/VirtualCameraActor.cpp
--- a/Engine/Plugins/Experimental/VirtualCamera/Source/VirtualCamera/Private/VirtualCameraActor.cpp
+++ b/Engine/Plugins/Experimental/VirtualCamera/Source/VirtualCamera/Private/VirtualCameraActor.cpp
@@ -463,7 +463,11 @@ bool AVirtualCameraActor::StopStreaming()
 			LevelEditorModule.OnMapChanged().RemoveAll(this);
 			if (FAssetRegistryModule* AssetRegistryModule = FModuleManager::GetModulePtr<FAssetRegistryModule>(AssetRegistryName))
 			{
-				AssetRegistryModule->Get().OnAssetRemoved().RemoveAll(this);
+				IAssetRegistry* AssetRegistry = AssetRegistryModule->TryGet();
+				if (AssetRegistry)
+				{
+					AssetRegistry->OnAssetRemoved().RemoveAll(this);
+				}
 			}
 			FEditorSupportDelegates::PrepareToCleanseEditorObject.RemoveAll(this);
 			GEditor->OnBlueprintPreCompile().RemoveAll(this);
diff --git a/Engine/Plugins/FX/Niagara/Source/Niagara/Private/NiagaraParameterDefinitionsSubscriber.cpp b/Engine/Plugins/FX/Niagara/Source/Niagara/Private/NiagaraParameterDefinitionsSubscriber.cpp
--- a/Engine/Plugins/FX/Niagara/Source/Niagara/Private/NiagaraParameterDefinitionsSubscriber.cpp
+++ b/Engine/Plugins/FX/Niagara/Source/Niagara/Private/NiagaraParameterDefinitionsSubscriber.cpp
@@ -62,7 +62,11 @@ void INiagaraParameterDefinitionsSubscriber::CleanupDefinitionsSubscriptions()
 	{
 		if(FAssetRegistryModule* AssetRegistryModule = FModuleManager::GetModulePtr<FAssetRegistryModule>("AssetRegistry"))
 		{
-			AssetRegistryModule->Get().OnFilesLoaded().Remove(OnDeferredSyncAllNameMatchParametersHandle);
+			IAssetRegistry* AssetRegistry = AssetRegistryModule->TryGet();
+			if (AssetRegistry)
+			{
+				AssetRegistry->OnFilesLoaded().Remove(OnDeferredSyncAllNameMatchParametersHandle);
+			}
 		}
 	}
 }
diff --git a/Engine/Plugins/Runtime/StateTree/Source/StateTreeEditorModule/Private/StateTreeNodeClassCache.cpp b/Engine/Plugins/Runtime/StateTree/Source/StateTreeEditorModule/Private/StateTreeNodeClassCache.cpp
--- a/Engine/Plugins/Runtime/StateTree/Source/StateTreeEditorModule/Private/StateTreeNodeClassCache.cpp
+++ b/Engine/Plugins/Runtime/StateTree/Source/StateTreeEditorModule/Private/StateTreeNodeClassCache.cpp
@@ -105,10 +105,13 @@ FStateTreeNodeClassCache::~FStateTreeNodeClassCache()
 	// Unregister with the Asset Registry to be informed when it is done loading up files.
 	if (FModuleManager::Get().IsModuleLoaded(TEXT("AssetRegistry")))
 	{
-		FAssetRegistryModule& AssetRegistryModule = FModuleManager::GetModuleChecked<FAssetRegistryModule>(TEXT("AssetRegistry"));
-		AssetRegistryModule.Get().OnFilesLoaded().RemoveAll(this);
-		AssetRegistryModule.Get().OnAssetAdded().RemoveAll(this);
-		AssetRegistryModule.Get().OnAssetRemoved().RemoveAll(this);
+		IAssetRegistry* AssetRegistry = FModuleManager::GetModuleChecked<FAssetRegistryModule>(TEXT("AssetRegistry")).TryGet();
+		if (AssetRegistry)
+		{
+			AssetRegistry->OnFilesLoaded().RemoveAll(this);
+			AssetRegistry->OnAssetAdded().RemoveAll(this);
+			AssetRegistry->OnAssetRemoved().RemoveAll(this);
+		}
 
 		// Unregister to have Populate called when doing a Reload.
 		FCoreUObjectDelegates::ReloadCompleteDelegate.RemoveAll(this);
diff --git a/Engine/Plugins/VirtualProduction/Takes/Source/TakeRecorder/Private/Widgets/STakeRecorderCockpit.cpp b/Engine/Plugins/VirtualProduction/Takes/Source/TakeRecorder/Private/Widgets/STakeRecorderCockpit.cpp
--- a/Engine/Plugins/VirtualProduction/Takes/Source/TakeRecorder/Private/Widgets/STakeRecorderCockpit.cpp
+++ b/Engine/Plugins/VirtualProduction/Takes/Source/TakeRecorder/Private/Widgets/STakeRecorderCockpit.cpp
@@ -72,7 +72,11 @@ STakeRecorderCockpit::~STakeRecorderCockpit()
 
 	if (FAssetRegistryModule* AssetRegistryModule = FModuleManager::GetModulePtr<FAssetRegistryModule>("AssetRegistry"))
 	{
-		AssetRegistryModule->Get().OnFilesLoaded().Remove(OnAssetRegistryFilesLoadedHandle);
+		IAssetRegistry* AssetRegistry = AssetRegistryModule->TryGet();
+		if (AssetRegistry)
+		{
+			AssetRegistry->OnFilesLoaded().Remove(OnAssetRegistryFilesLoadedHandle);
+		}
 	}
 
 	if (!ensure(TransactionIndex == INDEX_NONE))
diff --git a/Engine/Source/Developer/AssetTools/Private/SDiscoveringAssetsDialog.cpp b/Engine/Source/Developer/AssetTools/Private/SDiscoveringAssetsDialog.cpp
--- a/Engine/Source/Developer/AssetTools/Private/SDiscoveringAssetsDialog.cpp
+++ b/Engine/Source/Developer/AssetTools/Private/SDiscoveringAssetsDialog.cpp
@@ -22,8 +22,11 @@ SDiscoveringAssetsDialog::~SDiscoveringAssetsDialog()
 {
 	if ( FModuleManager::Get().IsModuleLoaded(TEXT("AssetRegistry")) )
 	{
-		FAssetRegistryModule& AssetRegistryModule = FModuleManager::GetModuleChecked<FAssetRegistryModule>("AssetRegistry");
-		AssetRegistryModule.Get().OnFilesLoaded().RemoveAll(this);
+		IAssetRegistry* AssetRegistry = FModuleManager::GetModuleChecked<FAssetRegistryModule>("AssetRegistry").TryGet();
+		if (AssetRegistry)
+		{
+			AssetRegistry->OnFilesLoaded().RemoveAll(this);
+		}
 	}
 }
 
diff --git a/Engine/Source/Developer/AutomationWindow/Private/SAutomationWindow.cpp b/Engine/Source/Developer/AutomationWindow/Private/SAutomationWindow.cpp
--- a/Engine/Source/Developer/AutomationWindow/Private/SAutomationWindow.cpp
+++ b/Engine/Source/Developer/AutomationWindow/Private/SAutomationWindow.cpp
@@ -123,8 +123,11 @@
 #if WITH_EDITOR
 	if ( FModuleManager::Get().IsModuleLoaded(TEXT("AssetRegistry")) )
 	{
-		FAssetRegistryModule& AssetRegistryModule = FModuleManager::GetModuleChecked<FAssetRegistryModule>(TEXT("AssetRegistry"));
-		AssetRegistryModule.Get().OnFileLoadProgressUpdated().RemoveAll(this);
+		IAssetRegistry* AssetRegistry = FModuleManager::GetModuleChecked<FAssetRegistryModule>(TEXT("AssetRegistry")).TryGet();
+		if (AssetRegistry)
+		{
+			AssetRegistry->OnFileLoadProgressUpdated().RemoveAll(this);
+		}
 	}
 #endif
 
diff --git a/Engine/Source/Editor/AIGraph/Private/AIGraphTypes.cpp b/Engine/Source/Editor/AIGraph/Private/AIGraphTypes.cpp
--- a/Engine/Source/Editor/AIGraph/Private/AIGraphTypes.cpp
+++ b/Engine/Source/Editor/AIGraph/Private/AIGraphTypes.cpp
@@ -164,10 +164,13 @@ FGraphNodeClassHelper::~FGraphNodeClassHelper()
 	// Unregister with the Asset Registry to be informed when it is done loading up files.
 	if (FModuleManager::Get().IsModuleLoaded(TEXT("AssetRegistry")))
 	{
-		FAssetRegistryModule& AssetRegistryModule = FModuleManager::GetModuleChecked<FAssetRegistryModule>(TEXT("AssetRegistry"));
-		AssetRegistryModule.Get().OnFilesLoaded().RemoveAll(this);
-		AssetRegistryModule.Get().OnAssetAdded().RemoveAll(this);
-		AssetRegistryModule.Get().OnAssetRemoved().RemoveAll(this);
+		IAssetRegistry* AssetRegistry = FModuleManager::GetModuleChecked<FAssetRegistryModule>(TEXT("AssetRegistry")).TryGet();
+		if (AssetRegistry)
+		{
+			AssetRegistry->OnFilesLoaded().RemoveAll(this);
+			AssetRegistry->OnAssetAdded().RemoveAll(this);
+			AssetRegistry->OnAssetRemoved().RemoveAll(this);
+		}
 
 		// Unregister to have Populate called when doing a Reload.
 		FCoreUObjectDelegates::ReloadCompleteDelegate.RemoveAll(this);
@@ -512,7 +515,7 @@ void FGraphNodeClassHelper::UpdateAvailableBlueprintClasses()
 {
 	if (FModuleManager::Get().IsModuleLoaded(TEXT("AssetRegistry")))
 	{
-		FAssetRegistryModule& AssetRegistryModule = FModuleManager::GetModuleChecked<FAssetRegistryModule>(TEXT("AssetRegistry"));
+		IAssetRegistry& AssetRegistry = FModuleManager::GetModuleChecked<FAssetRegistryModule>(TEXT("AssetRegistry")).Get();
 		const bool bSearchSubClasses = true;
 
 		TArray<FTopLevelAssetPath> ClassNames;
@@ -524,7 +527,7 @@ void FGraphNodeClassHelper::UpdateAvailableBlueprintClasses()
 			ClassNames.Add(It.Key()->GetClassPathName());
 
 			DerivedClassNames.Empty(DerivedClassNames.Num());
-			AssetRegistryModule.Get().GetDerivedClassNames(ClassNames, TSet<FTopLevelAssetPath>(), DerivedClassNames);
+			AssetRegistry.GetDerivedClassNames(ClassNames, TSet<FTopLevelAssetPath>(), DerivedClassNames);
 
 			int32& Count = It.Value();
 			Count = DerivedClassNames.Num();
diff --git a/Engine/Source/Editor/BlueprintGraph/Private/BlueprintActionDatabase.cpp b/Engine/Source/Editor/BlueprintGraph/Private/BlueprintActionDatabase.cpp
--- a/Engine/Source/Editor/BlueprintGraph/Private/BlueprintActionDatabase.cpp
+++ b/Engine/Source/Editor/BlueprintGraph/Private/BlueprintActionDatabase.cpp
@@ -1177,10 +1177,13 @@ FBlueprintActionDatabase::~FBlueprintActionDatabase()
 
 	if (FModuleManager::Get().IsModuleLoaded(TEXT("AssetRegistry")))
 	{
-		IAssetRegistry& AssetRegistry = FModuleManager::LoadModuleChecked<FAssetRegistryModule>(TEXT("AssetRegistry")).Get();
-		AssetRegistry.OnAssetAdded().Remove(OnAssetAddedDelegateHandle);
-		AssetRegistry.OnAssetRemoved().Remove(OnAssetRemovedDelegateHandle);
-		AssetRegistry.OnAssetRenamed().Remove(OnAssetRenamedDelegateHandle);
+		IAssetRegistry* AssetRegistry = FModuleManager::LoadModuleChecked<FAssetRegistryModule>(TEXT("AssetRegistry")).TryGet();
+		if (AssetRegistry)
+		{
+			AssetRegistry->OnAssetAdded().Remove(OnAssetAddedDelegateHandle);
+			AssetRegistry->OnAssetRemoved().Remove(OnAssetRemovedDelegateHandle);
+			AssetRegistry->OnAssetRenamed().Remove(OnAssetRenamedDelegateHandle);
+		}
 	}
 
 	FEditorDelegates::OnAssetsPreDelete.Remove(OnAssetsPreDeleteDelegateHandle);
diff --git a/Engine/Source/Editor/ClassViewer/Private/SClassViewer.cpp b/Engine/Source/Editor/ClassViewer/Private/SClassViewer.cpp
--- a/Engine/Source/Editor/ClassViewer/Private/SClassViewer.cpp
+++ b/Engine/Source/Editor/ClassViewer/Private/SClassViewer.cpp
@@ -1030,10 +1030,13 @@ FClassHierarchy::~FClassHierarchy()
 	// Unregister with the Asset Registry to be informed when it is done loading up files.
 	if (FAssetRegistryModule* AssetRegistryModule = FModuleManager::Get().GetModulePtr<FAssetRegistryModule>(AssetRegistryConstants::ModuleName))
 	{
-		IAssetRegistry& AssetRegistry = AssetRegistryModule->Get();
-		AssetRegistry.OnFilesLoaded().Remove(OnFilesLoadedRequestPopulateClassHierarchyDelegateHandle);
-		AssetRegistry.OnAssetAdded().RemoveAll( this );
-		AssetRegistry.OnAssetRemoved().RemoveAll( this );
+		IAssetRegistry* AssetRegistry = AssetRegistryModule->TryGet();
+		if (AssetRegistry)
+		{
+			AssetRegistry->OnFilesLoaded().Remove(OnFilesLoadedRequestPopulateClassHierarchyDelegateHandle);
+			AssetRegistry->OnAssetAdded().RemoveAll(this);
+			AssetRegistry->OnAssetRemoved().RemoveAll(this);
+		}
 
 		// Unregister to have Populate called when doing a Reload.
 		FCoreUObjectDelegates::ReloadCompleteDelegate.RemoveAll(this);
diff --git a/Engine/Source/Editor/EditorWidgets/Private/SAssetDiscoveryIndicator.cpp b/Engine/Source/Editor/EditorWidgets/Private/SAssetDiscoveryIndicator.cpp
--- a/Engine/Source/Editor/EditorWidgets/Private/SAssetDiscoveryIndicator.cpp
+++ b/Engine/Source/Editor/EditorWidgets/Private/SAssetDiscoveryIndicator.cpp
@@ -21,9 +21,13 @@ SAssetDiscoveryIndicator::~SAssetDiscoveryIndicator()
 {
 	if ( FModuleManager::Get().IsModuleLoaded(TEXT("AssetRegistry")) )
 	{
-		FAssetRegistryModule& AssetRegistryModule = FModuleManager::GetModuleChecked<FAssetRegistryModule>(TEXT("AssetRegistry"));
-		AssetRegistryModule.Get().OnFileLoadProgressUpdated().RemoveAll( this );
-		AssetRegistryModule.Get().OnFilesLoaded().RemoveAll( this );
+		// The engine might be shutting down so we need to use TryGet and handle null
+		IAssetRegistry* AssetRegistry = FModuleManager::GetModuleChecked<FAssetRegistryModule>(TEXT("AssetRegistry")).TryGet();
+		if (AssetRegistry) 
+		{
+			AssetRegistry->OnFileLoadProgressUpdated().RemoveAll(this);
+			AssetRegistry->OnFilesLoaded().RemoveAll(this);
+		}
 	}
 }
 
diff --git a/Engine/Source/Editor/FoliageEdit/Private/FoliageEdMode.cpp b/Engine/Source/Editor/FoliageEdit/Private/FoliageEdMode.cpp
--- a/Engine/Source/Editor/FoliageEdit/Private/FoliageEdMode.cpp
+++ b/Engine/Source/Editor/FoliageEdit/Private/FoliageEdMode.cpp
@@ -467,8 +467,11 @@ void FEdModeFoliage::Exit()
 
 	if (FModuleManager::Get().IsModuleLoaded(TEXT("AssetRegistry")))
 	{
-		FAssetRegistryModule& AssetRegistryModule = FModuleManager::GetModuleChecked<FAssetRegistryModule>(TEXT("AssetRegistry"));
-		AssetRegistryModule.Get().OnAssetRemoved().RemoveAll(this);
+		IAssetRegistry* AssetRegistry = FModuleManager::GetModuleChecked<FAssetRegistryModule>(TEXT("AssetRegistry")).TryGet();
+		if (AssetRegistry)
+		{
+			AssetRegistry->OnAssetRemoved().RemoveAll(this);
+		}
 	}
 
 	FCoreUObjectDelegates::OnObjectsReplaced.RemoveAll(this);
diff --git a/Engine/Source/Editor/FoliageEdit/Private/FoliageEditModule.cpp b/Engine/Source/Editor/FoliageEdit/Private/FoliageEditModule.cpp
--- a/Engine/Source/Editor/FoliageEdit/Private/FoliageEditModule.cpp
+++ b/Engine/Source/Editor/FoliageEdit/Private/FoliageEditModule.cpp
@@ -163,8 +163,11 @@ public:
 
 		if (FModuleManager::Get().IsModuleLoaded(TEXT("AssetRegistry")))
 		{
-			FAssetRegistryModule& AssetRegistryModule = FModuleManager::GetModuleChecked<FAssetRegistryModule>(TEXT("AssetRegistry"));
-			AssetRegistryModule.Get().OnAssetRemoved().RemoveAll(this);
+			IAssetRegistry* AssetRegistry = FModuleManager::GetModuleChecked<FAssetRegistryModule>(TEXT("AssetRegistry")).TryGet();
+			if (AssetRegistry)
+			{
+				AssetRegistry->OnAssetRemoved().RemoveAll(this);
+			}
 		}
 	}
 
diff --git a/Engine/Source/Editor/Kismet/Private/BlueprintNamespaceRegistry.cpp b/Engine/Source/Editor/Kismet/Private/BlueprintNamespaceRegistry.cpp
--- a/Engine/Source/Editor/Kismet/Private/BlueprintNamespaceRegistry.cpp
+++ b/Engine/Source/Editor/Kismet/Private/BlueprintNamespaceRegistry.cpp
@@ -78,10 +78,13 @@ void FBlueprintNamespaceRegistry::Shutdown()
 
 	if (FModuleManager::Get().IsModuleLoaded(TEXT("AssetRegistry")))
 	{
-		IAssetRegistry& AssetRegistry = FModuleManager::LoadModuleChecked<FAssetRegistryModule>(TEXT("AssetRegistry")).Get();
-		AssetRegistry.OnAssetAdded().Remove(OnAssetAddedDelegateHandle);
-		AssetRegistry.OnAssetRemoved().Remove(OnAssetRemovedDelegateHandle);
-		AssetRegistry.OnAssetRenamed().Remove(OnAssetRenamedDelegateHandle);
+		IAssetRegistry* AssetRegistry = FModuleManager::LoadModuleChecked<FAssetRegistryModule>(TEXT("AssetRegistry")).TryGet();
+		if (AssetRegistry)
+		{
+			AssetRegistry->OnAssetAdded().Remove(OnAssetAddedDelegateHandle);
+			AssetRegistry->OnAssetRemoved().Remove(OnAssetRemovedDelegateHandle);
+			AssetRegistry->OnAssetRenamed().Remove(OnAssetRenamedDelegateHandle);
+		}
 	}
 
 	bIsInitialized = false;
diff --git a/Engine/Source/Editor/Kismet/Private/FindInBlueprintManager.cpp b/Engine/Source/Editor/Kismet/Private/FindInBlueprintManager.cpp
--- a/Engine/Source/Editor/Kismet/Private/FindInBlueprintManager.cpp
+++ b/Engine/Source/Editor/Kismet/Private/FindInBlueprintManager.cpp
@@ -2028,11 +2028,14 @@ void FFindInBlueprintSearchManager::Initialize()
 	// Must ensure we do not attempt to load the AssetRegistry Module while saving a package, however, if it is loaded already we can safely obtain it
 	if (!GIsSavingPackage || (GIsSavingPackage && FModuleManager::Get().IsModuleLoaded(TEXT("AssetRegistry"))))
 	{
-		AssetRegistryModule = &FModuleManager::LoadModuleChecked<FAssetRegistryModule>(TEXT("AssetRegistry"));
-		AssetRegistryModule->Get().OnAssetAdded().AddRaw(this, &FFindInBlueprintSearchManager::OnAssetAdded);
-		AssetRegistryModule->Get().OnAssetRemoved().AddRaw(this, &FFindInBlueprintSearchManager::OnAssetRemoved);
-		AssetRegistryModule->Get().OnAssetRenamed().AddRaw(this, &FFindInBlueprintSearchManager::OnAssetRenamed);
-		AssetRegistryModule->Get().OnFilesLoaded().AddRaw(this, &FFindInBlueprintSearchManager::OnAssetRegistryFilesLoaded);
+		IAssetRegistry* AssetRegistry = FModuleManager::LoadModuleChecked<FAssetRegistryModule>(TEXT("AssetRegistry")).TryGet();
+		if (AssetRegistry)
+		{
+			AssetRegistry->OnAssetAdded().AddRaw(this, &FFindInBlueprintSearchManager::OnAssetAdded);
+			AssetRegistry->OnAssetRemoved().AddRaw(this, &FFindInBlueprintSearchManager::OnAssetRemoved);
+			AssetRegistry->OnAssetRenamed().AddRaw(this, &FFindInBlueprintSearchManager::OnAssetRenamed);
+			AssetRegistry->OnFilesLoaded().AddRaw(this, &FFindInBlueprintSearchManager::OnAssetRegistryFilesLoaded);
+		}
 	}
 	else
 	{
diff --git a/Engine/Source/Editor/Kismet/Private/WatchPointViewer.cpp b/Engine/Source/Editor/Kismet/Private/WatchPointViewer.cpp
--- a/Engine/Source/Editor/Kismet/Private/WatchPointViewer.cpp
+++ b/Engine/Source/Editor/Kismet/Private/WatchPointViewer.cpp
@@ -419,9 +419,12 @@ public:
 
 		if (FModuleManager::Get().IsModuleLoaded(TEXT("AssetRegistry")))
 		{
-			FAssetRegistryModule& AssetRegistryModule = FModuleManager::GetModuleChecked<FAssetRegistryModule>(TEXT("AssetRegistry"));
-			AssetRegistryModule.Get().OnAssetRemoved().RemoveAll(this);
-			AssetRegistryModule.Get().OnAssetRenamed().RemoveAll(this);
+			IAssetRegistry* AssetRegistry = FModuleManager::GetModuleChecked<FAssetRegistryModule>(TEXT("AssetRegistry")).TryGet();
+			if (AssetRegistry)
+			{
+				AssetRegistry->OnAssetRemoved().RemoveAll(this);
+				AssetRegistry->OnAssetRenamed().RemoveAll(this);
+			}
 		}
 	}
 
diff --git a/Engine/Source/Editor/Persona/Private/SAssetFamilyShortcutBar.cpp b/Engine/Source/Editor/Persona/Private/SAssetFamilyShortcutBar.cpp
--- a/Engine/Source/Editor/Persona/Private/SAssetFamilyShortcutBar.cpp
+++ b/Engine/Source/Editor/Persona/Private/SAssetFamilyShortcutBar.cpp
@@ -180,11 +180,14 @@ public:
 	{
 		if (FModuleManager::Get().IsModuleLoaded(TEXT("AssetRegistry")))
 		{
-			FAssetRegistryModule& AssetRegistryModule = FModuleManager::GetModuleChecked<FAssetRegistryModule>(TEXT("AssetRegistry"));
-			AssetRegistryModule.Get().OnFilesLoaded().RemoveAll(this);
-			AssetRegistryModule.Get().OnAssetAdded().RemoveAll(this);
-			AssetRegistryModule.Get().OnAssetRemoved().RemoveAll(this);
-			AssetRegistryModule.Get().OnAssetRenamed().RemoveAll(this);
+			IAssetRegistry* AssetRegistry = FModuleManager::GetModuleChecked<FAssetRegistryModule>(TEXT("AssetRegistry")).TryGet();
+			if (AssetRegistry)
+			{
+				AssetRegistry->OnFilesLoaded().RemoveAll(this);
+				AssetRegistry->OnAssetAdded().RemoveAll(this);
+				AssetRegistry->OnAssetRemoved().RemoveAll(this);
+				AssetRegistry->OnAssetRenamed().RemoveAll(this);
+			}
 		}
 
 		AssetFamily->GetOnAssetOpened().RemoveAll(this);
diff --git a/Engine/Source/Editor/PlacementMode/Private/PlacementModeModule.cpp b/Engine/Source/Editor/PlacementMode/Private/PlacementModeModule.cpp
--- a/Engine/Source/Editor/PlacementMode/Private/PlacementModeModule.cpp
+++ b/Engine/Source/Editor/PlacementMode/Private/PlacementModeModule.cpp
@@ -349,10 +349,14 @@ void FPlacementModeModule::PreUnloadCallback()
 	FAssetRegistryModule* AssetRegistryModule = FModuleManager::GetModulePtr<FAssetRegistryModule>("AssetRegistry");
 	if (AssetRegistryModule)
 	{
-		AssetRegistryModule->Get().OnAssetRemoved().RemoveAll(this);
-		AssetRegistryModule->Get().OnAssetRenamed().RemoveAll(this);
-		AssetRegistryModule->Get().OnAssetAdded().RemoveAll(this);
-		AssetRegistryModule->Get().OnFilesLoaded().RemoveAll(this);
+		IAssetRegistry* AssetRegistry = AssetRegistryModule->TryGet();
+		if (AssetRegistry)
+		{
+			AssetRegistry->OnAssetRemoved().RemoveAll(this);
+			AssetRegistry->OnAssetRenamed().RemoveAll(this);
+			AssetRegistry->OnAssetAdded().RemoveAll(this);
+			AssetRegistry->OnFilesLoaded().RemoveAll(this);
+		}
 	}
 }
 
diff --git a/Engine/Source/Editor/StructViewer/Private/SStructViewer.cpp b/Engine/Source/Editor/StructViewer/Private/SStructViewer.cpp
--- a/Engine/Source/Editor/StructViewer/Private/SStructViewer.cpp
+++ b/Engine/Source/Editor/StructViewer/Private/SStructViewer.cpp
@@ -806,10 +806,13 @@ FStructHierarchy::~FStructHierarchy()
 	// Unregister with the Asset Registry to be informed when it is done loading up files.
 	if (FModuleManager::Get().IsModuleLoaded(TEXT("AssetRegistry")))
 	{
-		FAssetRegistryModule& AssetRegistryModule = FModuleManager::GetModuleChecked<FAssetRegistryModule>(TEXT("AssetRegistry"));
-		AssetRegistryModule.Get().OnFilesLoaded().RemoveAll(this);
-		AssetRegistryModule.Get().OnAssetAdded().RemoveAll(this);
-		AssetRegistryModule.Get().OnAssetRemoved().RemoveAll(this);
+		IAssetRegistry* AssetRegistry = FModuleManager::GetModuleChecked<FAssetRegistryModule>(TEXT("AssetRegistry")).TryGet();
+		if (AssetRegistry)
+		{
+			AssetRegistry->OnFilesLoaded().RemoveAll(this);
+			AssetRegistry->OnAssetAdded().RemoveAll(this);
+			AssetRegistry->OnAssetRemoved().RemoveAll(this);
+		}
 
 		// Unregister to have Populate called when doing a Reload.
 		FCoreUObjectDelegates::ReloadCompleteDelegate.RemoveAll(this);
diff --git a/Engine/Source/Editor/UnrealEd/Private/AutoReimport/AssetSourceFilenameCache.cpp b/Engine/Source/Editor/UnrealEd/Private/AutoReimport/AssetSourceFilenameCache.cpp
--- a/Engine/Source/Editor/UnrealEd/Private/AutoReimport/AssetSourceFilenameCache.cpp
+++ b/Engine/Source/Editor/UnrealEd/Private/AutoReimport/AssetSourceFilenameCache.cpp
@@ -39,12 +39,16 @@ FAssetSourceFilenameCache& FAssetSourceFilenameCache::Get()
 
 void FAssetSourceFilenameCache::Shutdown()
 {
-	auto* ModulePtr = FModuleManager::GetModulePtr<FAssetRegistryModule>("AssetRegistry");
-	if (ModulePtr)
+	FAssetRegistryModule* AssetRegistryModule = FModuleManager::GetModulePtr<FAssetRegistryModule>("AssetRegistry");
+	if (AssetRegistryModule)
 	{
-		ModulePtr->Get().OnAssetAdded().RemoveAll(this);
-		ModulePtr->Get().OnAssetRemoved().RemoveAll(this);
-		ModulePtr->Get().OnAssetRenamed().RemoveAll(this);
+		IAssetRegistry* AssetRegistry = AssetRegistryModule->TryGet();
+		if (AssetRegistry)
+		{
+			AssetRegistry->OnAssetAdded().RemoveAll(this);
+			AssetRegistry->OnAssetRemoved().RemoveAll(this);
+			AssetRegistry->OnAssetRenamed().RemoveAll(this);
+		}
 	}
 	
 	AssetRenamedEvent.Clear();
diff --git a/Engine/Source/Editor/UnrealEd/Private/AutoReimport/AutoReimportManager.cpp b/Engine/Source/Editor/UnrealEd/Private/AutoReimport/AutoReimportManager.cpp
--- a/Engine/Source/Editor/UnrealEd/Private/AutoReimport/AutoReimportManager.cpp
+++ b/Engine/Source/Editor/UnrealEd/Private/AutoReimport/AutoReimportManager.cpp
@@ -373,7 +373,11 @@ void FAutoReimportManager::Destroy()
 	if (AssetRegistryModule)
 	{
 		FAssetSourceFilenameCache::Get().OnAssetRenamed().RemoveAll(this);
-		AssetRegistryModule->Get().OnInMemoryAssetDeleted().RemoveAll(this);
+		IAssetRegistry* AssetRegistry = AssetRegistryModule->TryGet();
+		if (AssetRegistry)
+		{
+			AssetRegistry->OnInMemoryAssetDeleted().RemoveAll(this);
+		}
 	}
 	
 	if (UEditorLoadingSavingSettings* Settings = GetMutableDefault<UEditorLoadingSavingSettings>())
diff --git a/Engine/Source/Editor/UnrealEd/Private/EditorEngine.cpp b/Engine/Source/Editor/UnrealEd/Private/EditorEngine.cpp
--- a/Engine/Source/Editor/UnrealEd/Private/EditorEngine.cpp
+++ b/Engine/Source/Editor/UnrealEd/Private/EditorEngine.cpp
@@ -1293,7 +1293,11 @@ void UEditorEngine::FinishDestroy()
 		FAssetRegistryModule* AssetRegistryModule = FModuleManager::GetModulePtr<FAssetRegistryModule>("AssetRegistry");
 		if (AssetRegistryModule)
 		{
-			AssetRegistryModule->Get().OnInMemoryAssetCreated().RemoveAll(this);
+			IAssetRegistry* AssetRegistry = AssetRegistryModule->TryGet();
+			if (AssetRegistry)
+			{
+				AssetRegistry->OnInMemoryAssetCreated().RemoveAll(this);
+			}
 		}
 
 		// Shut down transaction tracking system.
diff --git a/Engine/Source/Runtime/AssetRegistry/Private/AssetRegistry.cpp b/Engine/Source/Runtime/AssetRegistry/Private/AssetRegistry.cpp
--- a/Engine/Source/Runtime/AssetRegistry/Private/AssetRegistry.cpp
+++ b/Engine/Source/Runtime/AssetRegistry/Private/AssetRegistry.cpp
@@ -1325,70 +1325,78 @@ void UAssetRegistryImpl::OnEnginePreExit()
 	GuardedData.OnEnginePreExit();
 }
 
-UAssetRegistryImpl::~UAssetRegistryImpl()
+void UAssetRegistryImpl::FinishDestroy()
 {
-	FWriteScopeLock InterfaceScopeLock(InterfaceLock);
+	{
+		FWriteScopeLock InterfaceScopeLock(InterfaceLock);
 
-	// Stop listening for content mount point events
-	FPackageName::OnContentPathMounted().RemoveAll(this);
-	FPackageName::OnContentPathDismounted().RemoveAll(this);
-	FCoreDelegates::OnPostEngineInit.RemoveAll(this);
-	FCoreDelegates::OnEnginePreExit.RemoveAll(this);
-	IPluginManager::Get().OnLoadingPhaseComplete().RemoveAll(this);
+		// Stop listening for content mount point events
+		FPackageName::OnContentPathMounted().RemoveAll(this);
+		FPackageName::OnContentPathDismounted().RemoveAll(this);
+		FCoreDelegates::OnPostEngineInit.RemoveAll(this);
+		FCoreDelegates::OnEnginePreExit.RemoveAll(this);
+		IPluginManager::Get().OnLoadingPhaseComplete().RemoveAll(this);
 
 #if WITH_EDITOR
-	if (GIsEditor)
-	{
-		// If the directory module is still loaded, unregister any delegates
-		if (FModuleManager::Get().IsModuleLoaded("DirectoryWatcher"))
+		if (GIsEditor)
 		{
-			FDirectoryWatcherModule& DirectoryWatcherModule = FModuleManager::GetModuleChecked<FDirectoryWatcherModule>("DirectoryWatcher");
-			IDirectoryWatcher* DirectoryWatcher = DirectoryWatcherModule.Get();
-
-			if (DirectoryWatcher)
+			// If the directory module is still loaded, unregister any delegates
+			if (FModuleManager::Get().IsModuleLoaded("DirectoryWatcher"))
 			{
-				TArray<FString> RootContentPaths;
-				FPackageName::QueryRootContentPaths(RootContentPaths);
-				for (TArray<FString>::TConstIterator RootPathIt(RootContentPaths); RootPathIt; ++RootPathIt)
+				FDirectoryWatcherModule& DirectoryWatcherModule = FModuleManager::GetModuleChecked<FDirectoryWatcherModule>("DirectoryWatcher");
+				IDirectoryWatcher* DirectoryWatcher = DirectoryWatcherModule.Get();
+
+				if (DirectoryWatcher)
 				{
-					const FString& RootPath = *RootPathIt;
-					const FString& ContentFolder = FPackageName::LongPackageNameToFilename(RootPath);
-					DirectoryWatcher->UnregisterDirectoryChangedCallback_Handle(ContentFolder, OnDirectoryChangedDelegateHandles.FindRef(RootPath));
+					TArray<FString> RootContentPaths;
+					FPackageName::QueryRootContentPaths(RootContentPaths);
+					for (TArray<FString>::TConstIterator RootPathIt(RootContentPaths); RootPathIt; ++RootPathIt)
+					{
+						const FString& RootPath = *RootPathIt;
+						const FString& ContentFolder = FPackageName::LongPackageNameToFilename(RootPath);
+						DirectoryWatcher->UnregisterDirectoryChangedCallback_Handle(ContentFolder, OnDirectoryChangedDelegateHandles.FindRef(RootPath));
+					}
 				}
 			}
 		}
-	}
 
-	if (GuardedData.IsUpdateDiskCacheAfterLoad())
-	{
-		FCoreUObjectDelegates::OnAssetLoaded.RemoveAll(this);
-	}
+		if (GuardedData.IsUpdateDiskCacheAfterLoad())
+		{
+			FCoreUObjectDelegates::OnAssetLoaded.RemoveAll(this);
+		}
 
-	if (bAddMetaDataTagsToOnGetExtraObjectTags)
-	{
-		UObject::FAssetRegistryTag::OnGetExtraObjectTags.RemoveAll(this);
-	}
+		if (bAddMetaDataTagsToOnGetExtraObjectTags)
+		{
+			UObject::FAssetRegistryTag::OnGetExtraObjectTags.RemoveAll(this);
+		}
 #endif // WITH_EDITOR
 
-	if (HasAnyFlags(RF_ClassDefaultObject))
-	{
-		check(UE::AssetRegistry::Private::IAssetRegistrySingleton::Singleton == this && IAssetRegistryInterface::Default == &GAssetRegistryInterface);
-		UE::AssetRegistry::Private::IAssetRegistrySingleton::Singleton = nullptr;
-		IAssetRegistryInterface::Default = nullptr;
-	}
-
-	// Clear all listeners
-	PathAddedEvent.Clear();
-	PathRemovedEvent.Clear();
-	AssetAddedEvent.Clear();
-	AssetRemovedEvent.Clear();
-	AssetRenamedEvent.Clear();
-	AssetUpdatedEvent.Clear();
-	AssetUpdatedOnDiskEvent.Clear();
-	InMemoryAssetCreatedEvent.Clear();
-	InMemoryAssetDeletedEvent.Clear();
-	FileLoadedEvent.Clear();
-	FileLoadProgressUpdatedEvent.Clear();
+		if (HasAnyFlags(RF_ClassDefaultObject))
+		{
+			check(UE::AssetRegistry::Private::IAssetRegistrySingleton::Singleton == this && IAssetRegistryInterface::Default == &GAssetRegistryInterface);
+			UE::AssetRegistry::Private::IAssetRegistrySingleton::Singleton = nullptr;
+			IAssetRegistryInterface::Default = nullptr;
+		}
+
+		// Clear all listeners
+		PathAddedEvent.Clear();
+		PathRemovedEvent.Clear();
+		AssetAddedEvent.Clear();
+		AssetRemovedEvent.Clear();
+		AssetRenamedEvent.Clear();
+		AssetUpdatedEvent.Clear();
+		AssetUpdatedOnDiskEvent.Clear();
+		InMemoryAssetCreatedEvent.Clear();
+		InMemoryAssetDeletedEvent.Clear();
+		FileLoadedEvent.Clear();
+		FileLoadProgressUpdatedEvent.Clear();
+	}
+
+	Super::FinishDestroy();
+}
+
+UAssetRegistryImpl::~UAssetRegistryImpl()
+{
 }
 
 UAssetRegistryImpl& UAssetRegistryImpl::Get()
diff --git a/Engine/Source/Runtime/AssetRegistry/Private/AssetRegistry.h b/Engine/Source/Runtime/AssetRegistry/Private/AssetRegistry.h
--- a/Engine/Source/Runtime/AssetRegistry/Private/AssetRegistry.h
+++ b/Engine/Source/Runtime/AssetRegistry/Private/AssetRegistry.h
@@ -24,6 +24,7 @@ public:
 	UAssetRegistryImpl(const FObjectInitializer& ObjectInitializer);
 	UAssetRegistryImpl(FVTableHelper& Helper);
 	virtual ~UAssetRegistryImpl();
+	virtual void FinishDestroy() override;
 
 	/** Gets the asset registry singleton for asset registry module use */
 	static UAssetRegistryImpl& Get();
diff --git a/Engine/Source/Runtime/AssetRegistry/Public/AssetRegistry/AssetRegistryModule.h b/Engine/Source/Runtime/AssetRegistry/Public/AssetRegistry/AssetRegistryModule.h
--- a/Engine/Source/Runtime/AssetRegistry/Public/AssetRegistry/AssetRegistryModule.h
+++ b/Engine/Source/Runtime/AssetRegistry/Public/AssetRegistry/AssetRegistryModule.h
@@ -30,6 +30,18 @@ public:
 		return IAssetRegistry::GetChecked();
 	}
 
+	/** Reports whether Get is valid to call. Will be true except during engine shutdown. */
+	bool IsValid() const
+	{
+		return IAssetRegistry::Get() != nullptr;
+	}
+
+	/** Returns AssetRegistry pointer if valid, or nullptr. Will be non-null except during engine shutdown. */
+	IAssetRegistry* TryGet() const
+	{
+		return IAssetRegistry::Get();
+	}
+
 	static IAssetRegistry& GetRegistry()
 	{
 		return IAssetRegistry::GetChecked();
diff --git a/Engine/Source/Runtime/Engine/Private/Internationalization/EnginePackageLocalizationCache.cpp b/Engine/Source/Runtime/Engine/Private/Internationalization/EnginePackageLocalizationCache.cpp
--- a/Engine/Source/Runtime/Engine/Private/Internationalization/EnginePackageLocalizationCache.cpp
+++ b/Engine/Source/Runtime/Engine/Private/Internationalization/EnginePackageLocalizationCache.cpp
@@ -24,11 +24,13 @@ FEnginePackageLocalizationCache::~FEnginePackageLocalizationCache()
 	if (FModuleManager::Get().IsModuleLoaded(TEXT("AssetRegistry")))
 	{
 		FAssetRegistryModule& AssetRegistryModule = FModuleManager::GetModuleChecked<FAssetRegistryModule>(TEXT("AssetRegistry"));
-		IAssetRegistry& AssetRegistry = AssetRegistryModule.Get();
-
-		AssetRegistry.OnAssetAdded().RemoveAll(this);
-		AssetRegistry.OnAssetRemoved().RemoveAll(this);
-		AssetRegistry.OnAssetRenamed().RemoveAll(this);
+		IAssetRegistry* AssetRegistry = AssetRegistryModule.TryGet();
+		if (AssetRegistry)
+		{
+			AssetRegistry->OnAssetAdded().RemoveAll(this);
+			AssetRegistry->OnAssetRemoved().RemoveAll(this);
+			AssetRegistry->OnAssetRenamed().RemoveAll(this);
+		}
 	}
 }
 
diff --git a/Engine/Source/Runtime/Engine/Private/Slate/SlateGameResources.cpp b/Engine/Source/Runtime/Engine/Private/Slate/SlateGameResources.cpp
--- a/Engine/Source/Runtime/Engine/Private/Slate/SlateGameResources.cpp
+++ b/Engine/Source/Runtime/Engine/Private/Slate/SlateGameResources.cpp
@@ -34,10 +34,12 @@ FSlateGameResources::~FSlateGameResources()
 	if ( GIsEditor && FModuleManager::Get().IsModuleLoaded( TEXT("AssetRegistry") ) )
 	{
 		FAssetRegistryModule& AssetRegistryModule = FModuleManager::GetModuleChecked<FAssetRegistryModule>( TEXT("AssetRegistry") );
-		IAssetRegistry& AssetRegistry = AssetRegistryModule.Get();
-
-		AssetRegistry.OnAssetAdded().RemoveAll( this );
-		AssetRegistry.OnAssetRemoved().RemoveAll( this );
+		IAssetRegistry* AssetRegistry = AssetRegistryModule.TryGet();
+		if (AssetRegistry)
+		{
+			AssetRegistry->OnAssetAdded().RemoveAll(this);
+			AssetRegistry->OnAssetRemoved().RemoveAll(this);
+		}
 	}
 }
 
