AssetRegistry support for -nothreading: The editor needs to call SearchAllAssets synchronously for the editor when asynchronous is not available.

diff --git a/Engine/Source/Runtime/AssetRegistry/Private/AssetRegistry.cpp b/Engine/Source/Runtime/AssetRegistry/Private/AssetRegistry.cpp
--- a/Engine/Source/Runtime/AssetRegistry/Private/AssetRegistry.cpp
+++ b/Engine/Source/Runtime/AssetRegistry/Private/AssetRegistry.cpp
@@ -634,6 +634,7 @@ struct FInitializeContext
 	TArray<FString> RootContentPaths;
 	bool bRedirectorsNeedSubscribe = false;
 	bool bUpdateDiskCacheAfterLoad = false;
+	bool bNeedsSearchAllAssetsAtStartSynchronous = false;
 };
 
 }
@@ -774,6 +775,11 @@ void FAssetRegistryImpl::Initialize(Impl::FInitializeContext& Context)
 		{
 			SearchAllAssetsInitialAsync(Context.Events);
 		}
+		else
+		{
+			// For the Editor we need to take responsibility for the synchronous search; Commandlets will handle that themselves
+			Context.bNeedsSearchAllAssetsAtStartSynchronous = GIsEditor && !IsRunningCommandlet();
+		}
 	}
 
 	ConsumeOrDeferPreloadedPremade(Context.UARI, Context.Events);
@@ -878,6 +884,10 @@ void UAssetRegistryImpl::InitializeEvents(UE::AssetRegistry::Impl::FInitializeCo
 	{
 		UObject::FAssetRegistryTag::OnGetExtraObjectTags.AddUObject(this, &UAssetRegistryImpl::OnGetExtraObjectTags);
 	}
+	if (Context.bNeedsSearchAllAssetsAtStartSynchronous)
+	{
+		FCoreDelegates::OnFEngineLoopInitComplete.AddUObject(this, &UAssetRegistryImpl::OnFEngineLoopInitCompleteSearchAllAssets);
+	}
 #endif // WITH_EDITOR
 
 	FCoreDelegates::OnEnginePreExit.AddUObject(this, &UAssetRegistryImpl::OnEnginePreExit);
@@ -1319,6 +1329,13 @@ void FAssetRegistryImpl::RefreshNativeClasses()
 
 }
 
+#if WITH_EDITOR
+void UAssetRegistryImpl::OnFEngineLoopInitCompleteSearchAllAssets()
+{
+	SearchAllAssets(true);
+}
+#endif
+
 void UAssetRegistryImpl::OnEnginePreExit()
 {
 	FWriteScopeLock InterfaceScopeLock(InterfaceLock);
@@ -1369,6 +1386,8 @@ void UAssetRegistryImpl::FinishDestroy()
 		{
 			UObject::FAssetRegistryTag::OnGetExtraObjectTags.RemoveAll(this);
 		}
+		FCoreDelegates::OnFEngineLoopInitComplete.RemoveAll(this);
+
 #endif // WITH_EDITOR
 
 		if (HasAnyFlags(RF_ClassDefaultObject))
diff --git a/Engine/Source/Runtime/AssetRegistry/Private/AssetRegistry.h b/Engine/Source/Runtime/AssetRegistry/Private/AssetRegistry.h
--- a/Engine/Source/Runtime/AssetRegistry/Private/AssetRegistry.h
+++ b/Engine/Source/Runtime/AssetRegistry/Private/AssetRegistry.h
@@ -172,6 +172,9 @@ protected:
 
 private:
 	void OnEnginePreExit();
+#if WITH_EDITOR
+	void OnFEngineLoopInitCompleteSearchAllAssets();
+#endif
 	void InitializeEvents(UE::AssetRegistry::Impl::FInitializeContext& Context);
 	void Broadcast(UE::AssetRegistry::Impl::FEventContext& EventContext);
 
