DumpAssetRegistryCommandlet: Converts the AssetRegistry.bin files generated by the cook into text.
#rnx

diff --git a/Engine/Source/Editor/UnrealEd/Classes/Commandlets/DumpAssetRegistryCommandlet.h b/Engine/Source/Editor/UnrealEd/Classes/Commandlets/DumpAssetRegistryCommandlet.h
new file mode 100644
--- /dev/null
+++ b/Engine/Source/Editor/UnrealEd/Classes/Commandlets/DumpAssetRegistryCommandlet.h
@@ -0,0 +1,29 @@
+// Copyright Epic Games, Inc. All Rights Reserved.
+#pragma once
+
+#include "Commandlets/Commandlet.h"
+#include "Containers/Array.h"
+#include "Containers/UnrealString.h"
+
+#include "DumpAssetRegistryCommandlet.generated.h"
+
+UCLASS()
+class UDumpAssetRegistryCommandlet: public UCommandlet
+{
+	GENERATED_UCLASS_BODY()
+
+public:
+	// Begin UCommandlet Interface
+	virtual int32 Main(const FString& Params) override;
+	// End UCommandlet Interface
+
+private:
+	bool TryParseArgs();
+	bool TryDumpAssetRegistry();
+
+	TArray<FString> FormattingArgs;
+	int32 LinesPerPage = 0;
+	bool bLowerCase = false;
+	FString Path;
+	FString OutDir;
+};
diff --git a/Engine/Source/Editor/UnrealEd/Private/Commandlets/DumpAssetRegistryCommandlet.cpp b/Engine/Source/Editor/UnrealEd/Private/Commandlets/DumpAssetRegistryCommandlet.cpp
new file mode 100644
--- /dev/null
+++ b/Engine/Source/Editor/UnrealEd/Private/Commandlets/DumpAssetRegistryCommandlet.cpp
@@ -0,0 +1,153 @@
+// Copyright Epic Games, Inc. All Rights Reserved.
+
+#include "Commandlets/DumpAssetRegistryCommandlet.h"
+
+#include "AssetRegistry/AssetRegistryState.h"
+#include "HAL/FileManager.h"
+#include "Logging/LogMacros.h"
+#include "Misc/CommandLine.h"
+#include "Misc/FileHelper.h"
+#include "Misc/Guid.h"
+#include "Misc/Parse.h"
+#include "Misc/Paths.h"
+#include "Misc/StringBuilder.h"
+#include "Serialization/ArrayReader.h"
+
+DEFINE_LOG_CATEGORY_STATIC(LogAssetRegistryDump, Log, All);
+
+UDumpAssetRegistryCommandlet::UDumpAssetRegistryCommandlet(const FObjectInitializer& Initializer)
+	: Super(Initializer)
+{
+}
+
+int32 UDumpAssetRegistryCommandlet::Main(const FString& FullCommandLine)
+{
+	UE_LOG(LogAssetRegistryDump, Display, TEXT("--------------------------------------------------------------------------------------------"));
+	UE_LOG(LogAssetRegistryDump, Display, TEXT("Running DumpAssetRegistry Commandlet"));
+	ON_SCOPE_EXIT
+	{
+		UE_LOG(LogAssetRegistryDump, Display, TEXT("Completed DumpAssetRegistry Commandlet"));
+		UE_LOG(LogAssetRegistryDump, Display, TEXT("--------------------------------------------------------------------------------------------"));
+	};
+
+	if (!TryParseArgs())
+	{
+		return 1;
+	}
+	if (!TryDumpAssetRegistry())
+	{
+		return 1;
+	}
+
+	return 0;
+}
+
+bool UDumpAssetRegistryCommandlet::TryParseArgs()
+{
+	const TCHAR* CommandLine = FCommandLine::Get();
+	TArray<FString> Tokens;
+	TArray<FString> Switches;
+	ParseCommandLine(CommandLine, Tokens, Switches);
+
+	TArray<FString> AllFormattingArgs = {
+		TEXT("All"), TEXT("ObjectPath"), TEXT("PackageName"), TEXT("Path"), TEXT("Class"), TEXT("Tag"),
+		TEXT("Dependencies"), TEXT("DependencyDetails"), TEXT("LegacyDependencies"), TEXT("PackageData")
+	};
+	FormattingArgs.Reset();
+	LinesPerPage = 10000;
+	Path.Reset();
+	OutDir.Reset();
+	bLowerCase = false;
+	for (const FString& Switch : Switches)
+	{
+		if (FParse::Value(*Switch, TEXT("LinesPerPage="), LinesPerPage))
+		{
+		}
+		else if (FParse::Value(*Switch, TEXT("Path="), Path))
+		{
+		}
+		else if (FParse::Value(*Switch, TEXT("OutDir="), OutDir))
+		{
+		}
+		else if (Switch == TEXT("LowerCase"))
+		{
+			bLowerCase = true;
+		}
+		else if (AllFormattingArgs.Contains(Switch))
+		{
+			FormattingArgs.Add(Switch);
+		}
+	}
+
+	if (Path.IsEmpty())
+	{
+		TStringBuilder<256> AllFormattingArgsText;
+		for (const FString& Arg : AllFormattingArgs)
+		{
+			AllFormattingArgsText << TEXT("-") << Arg << TEXT(", ");
+		}
+		AllFormattingArgsText.RemoveSuffix(2); // Remove trailing ", "
+		
+		UE_LOG(LogAssetRegistryDump, Error, TEXT("Missing path argument."));
+		UE_LOG(LogAssetRegistryDump, Display, TEXT("Usage: -Path=<Path> [-OutDir=<Path>] [-LinesPerPage=<int>] [-FormatSwitch]..."));
+		UE_LOG(LogAssetRegistryDump, Display, TEXT("FormatSwitches: %s"), *AllFormattingArgsText);
+		return false;
+	}
+	if (FormattingArgs.IsEmpty())
+	{
+		FormattingArgs.Add(TEXT("All"));
+	}
+	if (OutDir.IsEmpty())
+	{
+		OutDir = FPaths::Combine(FPaths::ProjectDir(), TEXT("Saved"), TEXT("Reports"), TEXT("AssetRegistryDump"));
+	}
+	return true;
+}
+
+bool UDumpAssetRegistryCommandlet::TryDumpAssetRegistry()
+{
+	IFileManager& FileManager = IFileManager::Get();
+	if (!FileManager.FileExists(*Path))
+	{
+		UE_LOG(LogAssetRegistryDump, Error, TEXT("File '%s' does not exist."), *Path);
+		return false;
+	}
+	FArrayReader SerializedAssetData;
+	if (!FFileHelper::LoadFileToArray(SerializedAssetData, *Path))
+	{
+		UE_LOG(LogAssetRegistryDump, Error, TEXT("Failed to load file '%s'."), *Path);
+		return false;
+	}
+	FAssetRegistryState State;
+	if (!State.Load(SerializedAssetData))
+	{
+		UE_LOG(LogAssetRegistryDump, Error, TEXT("Failed to parse file '%s' as asset registry."), *Path);
+		return false;
+	}
+	if (!FileManager.DirectoryExists(*OutDir))
+	{
+		if (!FileManager.MakeDirectory(*OutDir, true /* Tree */))
+		{
+			UE_LOG(LogAssetRegistryDump, Error, TEXT("Failed to create OutDir '%s'."), *OutDir);
+			return false;
+		}
+	}
+
+	TArray<FString> Pages;
+	State.Dump(FormattingArgs, Pages, LinesPerPage);
+	int PageIndex = 0;
+	TStringBuilder<256> FileName;
+	for (FString& PageText : Pages)
+	{
+		FileName.Reset();
+		FileName.Appendf(TEXT("%s_%05d.txt"), *(OutDir / TEXT("Page")), PageIndex++);
+		if (bLowerCase)
+		{
+			PageText.ToLowerInline();
+		}
+		FFileHelper::SaveStringToFile(PageText, *FileName);
+	}
+
+	UE_LOG(LogAssetRegistryDump, Display, TEXT("Wrote %d files to %s."), Pages.Num(), *OutDir);
+	return true;
+}
\ No newline at end of file
diff --git a/Engine/Source/Runtime/AssetRegistry/Private/AssetRegistryState.cpp b/Engine/Source/Runtime/AssetRegistry/Private/AssetRegistryState.cpp
--- a/Engine/Source/Runtime/AssetRegistry/Private/AssetRegistryState.cpp
+++ b/Engine/Source/Runtime/AssetRegistry/Private/AssetRegistryState.cpp
@@ -2130,7 +2130,8 @@ namespace AssetRegistry
 }
 
 template <typename MapType>
-static void PrintAssetDataMap(FString Name, const MapType& AssetMap, TStringBuilder<16>& PageBuffer, const TFunctionRef<void()>& AddLine)
+static void PrintAssetDataMap(FString Name, const MapType& AssetMap, TStringBuilder<16>& PageBuffer, const TFunctionRef<void()>& AddLine,
+	TUniqueFunction<void(const typename MapType::KeyType& Key, const FAssetData& Data)>&& PrintValue = {})
 {
 	PageBuffer.Appendf(TEXT("--- Begin %s ---"), *Name);
 	AddLine();
@@ -2166,14 +2167,18 @@ static void PrintAssetDataMap(FString Name, const MapType& AssetMap, TStringBuil
 			{ return A.ObjectPath.ToString() < B.ObjectPath.ToString(); }
 		);
 
-		PageBuffer.Append(TEXT("	"));
+		PageBuffer.Append(TEXT("\t"));
 		Key.AppendString(PageBuffer);
 		PageBuffer.Appendf(TEXT(" : %d item(s)"), Items.Num());
 		AddLine();
 		for (const FAssetData* Data : Items)
 		{
-			PageBuffer.Append(TEXT("	 "));
+			PageBuffer.Append(TEXT("\t\t"));
 			Data->ObjectPath.AppendString(PageBuffer);
+			if (PrintValue)
+			{
+				PrintValue(Key, *Data);
+			}
 			AddLine();
 		}
 	}
@@ -2187,21 +2192,23 @@ void FAssetRegistryState::Dump(const TArray<FString>& Arguments, TArray<FString>
 	int32 ExpectedNumLines = 14 + CachedAssetsByObjectPath.Num() * 5 + CachedDependsNodes.Num() + CachedPackageData.Num();
 	const int32 EstimatedLinksPerNode = 10*2; // Each dependency shows up once as a dependency and once as a reference
 	const int32 EstimatedCharactersPerLine = 100;
-	const bool bDumpDependencyDetails = Arguments.Contains(TEXT("DependencyDetails"));
+	bool bAllFields = Arguments.Contains(TEXT("All"));
+
+	const bool bDumpDependencyDetails = bAllFields || Arguments.Contains(TEXT("DependencyDetails"));
 	if (bDumpDependencyDetails)
 	{
 		ExpectedNumLines += CachedDependsNodes.Num() * (3 + EstimatedLinksPerNode);
 	}
-	LinesPerPage = FMath::Max(LinesPerPage, 1);
-	const int32 ExpectedNumPages = ExpectedNumLines / LinesPerPage;
-	const int32 PageEndSearchLength = LinesPerPage / 20;
+	LinesPerPage = FMath::Max(0, LinesPerPage);
+	const int32 ExpectedNumPages = LinesPerPage > 0 ? (ExpectedNumLines / LinesPerPage) : 1;
+	const int32 PageEndSearchLength = FMath::Min(LinesPerPage, ExpectedNumLines) / 20;
 	const uint32 HashStartValue = MAX_uint32 - 49979693; // Pick a large starting value to bias against picking empty string
 	const uint32 HashMultiplier = 67867967;
 	TStringBuilder<16> PageBuffer;
 	TStringBuilder<16> OverflowText;
 
 	OutPages.Reserve(ExpectedNumPages);
-	PageBuffer.AddUninitialized(LinesPerPage * EstimatedCharactersPerLine);// TODO: Add Reserve function to TStringBuilder
+	PageBuffer.AddUninitialized(FMath::Min(LinesPerPage, ExpectedNumLines) * EstimatedCharactersPerLine);// TODO: Add Reserve function to TStringBuilder
 	PageBuffer.Reset();
 	OverflowText.AddUninitialized(PageEndSearchLength * EstimatedCharactersPerLine);
 	OverflowText.Reset();
@@ -2274,7 +2281,7 @@ void FAssetRegistryState::Dump(const TArray<FString>& Arguments, TArray<FString>
 		else
 		{
 			++NumLinesInPage;
-			if (NumLinesInPage != LinesPerPage)
+			if (LinesPerPage == 0 || NumLinesInPage < LinesPerPage)
 			{
 				PageBuffer.Append(LINE_TERMINATOR);
 			}
@@ -2285,7 +2292,7 @@ void FAssetRegistryState::Dump(const TArray<FString>& Arguments, TArray<FString>
 		}
 	};
 
-	if (Arguments.Contains(TEXT("ObjectPath")))
+	if (bAllFields || Arguments.Contains(TEXT("ObjectPath")))
 	{
 		PageBuffer.Append(TEXT("--- Begin CachedAssetsByObjectPath ---"));
 		AddLine();
@@ -2305,27 +2312,31 @@ void FAssetRegistryState::Dump(const TArray<FString>& Arguments, TArray<FString>
 		AddLine();
 	}
 
-	if (Arguments.Contains(TEXT("PackageName")))
+	if (bAllFields || Arguments.Contains(TEXT("PackageName")))
 	{
 		PrintAssetDataMap(TEXT("CachedAssetsByPackageName"), CachedAssetsByPackageName, PageBuffer, AddLine);
 	}
 
-	if (Arguments.Contains(TEXT("Path")))
+	if (bAllFields || Arguments.Contains(TEXT("Path")))
 	{
 		PrintAssetDataMap(TEXT("CachedAssetsByPath"), CachedAssetsByPath, PageBuffer, AddLine);
 	}
 
-	if (Arguments.Contains(TEXT("Class")))
+	if (bAllFields || Arguments.Contains(TEXT("Class")))
 	{
 		PrintAssetDataMap(TEXT("CachedAssetsByClass"), CachedAssetsByClass, PageBuffer, AddLine);
 	}
 
-	if (Arguments.Contains(TEXT("Tag")))
+	if (bAllFields || Arguments.Contains(TEXT("Tag")))
 	{
-		PrintAssetDataMap(TEXT("CachedAssetsByTag"), CachedAssetsByTag, PageBuffer, AddLine);
+		PrintAssetDataMap(TEXT("CachedAssetsByTag"), CachedAssetsByTag, PageBuffer, AddLine,
+			[&PageBuffer, &AddLine](const FName& TagName, const FAssetData& Data)
+			{
+				PageBuffer << TEXT(", ") << Data.TagsAndValues.FindTag(TagName).ToLoose();
+			});
 	}
 
-	if (Arguments.Contains(TEXT("Dependencies")) && !bDumpDependencyDetails)
+	if ((bAllFields || Arguments.Contains(TEXT("Dependencies"))) && !bDumpDependencyDetails)
 	{
 		PageBuffer.Appendf(TEXT("--- Begin CachedDependsNodes ---"));
 		AddLine();
@@ -2359,7 +2370,7 @@ void FAssetRegistryState::Dump(const TArray<FString>& Arguments, TArray<FString>
 		CachedDependsNodes.GenerateValueArray(Nodes);
 		Nodes.Sort(SortByAssetID);
 
-		if (Arguments.Contains(TEXT("LegacyDependencies")))
+		if (Arguments.Contains(TEXT("LegacyDependencies"))) // LegacyDependencies are not show by all; they have to be directly requested
 		{
 			EDependencyCategory CategoryTypes[] = { EDependencyCategory::Package, EDependencyCategory::Package,EDependencyCategory::SearchableName,EDependencyCategory::Manage, EDependencyCategory::Manage, EDependencyCategory::None };
 			EDependencyQuery CategoryQueries[] = { EDependencyQuery::Hard, EDependencyQuery::Soft, EDependencyQuery::NoRequirements, EDependencyQuery::Direct, EDependencyQuery::Indirect, EDependencyQuery::NoRequirements };
@@ -2464,7 +2475,7 @@ void FAssetRegistryState::Dump(const TArray<FString>& Arguments, TArray<FString>
 		PageBuffer.Appendf(TEXT("--- End CachedDependsNodes : %d entries ---"), CachedDependsNodes.Num());
 		AddLine();
 	}
-	if (Arguments.Contains(TEXT("PackageData")))
+	if (bAllFields || Arguments.Contains(TEXT("PackageData")))
 	{
 		PageBuffer.Append(TEXT("--- Begin CachedPackageData ---"));
 		AddLine();
