Move some asset registry verification code under a cvar now that it has been shown to work. It can cause a performance regression in some circumstances so shouldn't be on by default

diff --git a/Engine/Source/Runtime/AssetRegistry/Private/AssetRegistry.cpp b/Engine/Source/Runtime/AssetRegistry/Private/AssetRegistry.cpp
--- a/Engine/Source/Runtime/AssetRegistry/Private/AssetRegistry.cpp
+++ b/Engine/Source/Runtime/AssetRegistry/Private/AssetRegistry.cpp
@@ -2421,7 +2421,7 @@ bool UAssetRegistryImpl::DoesPackageExistOnDisk(FName PackageName, FString* OutC
 			FReadScopeLock InterfaceScopeLock(InterfaceLock);
 			AssetPackageData = GuardedData.GetState().GetAssetPackageData(PackageName);
 		}
-		constexpr bool bVerifyNegativeResults = true;
+		const static bool bVerifyNegativeResults = FParse::Param(FCommandLine::Get(), TEXT("AssetRegistryValidatePackageExists"));
 		if (bVerifyNegativeResults && !AssetPackageData)
 		{
 			FString FileName;
