AssetDataGatherer ANY_PACKAGE optimization: cache the ImportMap on FPackageReader and pass it into UE::AssetRegistry::ReadPackageDataMain so that if ReconstructFullClassPath needs it it will not need to be loaded again.
Also fix bug in corrupt-package error logging for FPackageReader. FMessageLog is not threadsafe, and FPackageReader can be called from multiple threads, so we need to change it to use threadsafe UE_LOG instead.
#rnx

diff --git a/Engine/Source/Runtime/AssetRegistry/Private/PackageDependencyData.cpp b/Engine/Source/Runtime/AssetRegistry/Private/PackageDependencyData.cpp
--- a/Engine/Source/Runtime/AssetRegistry/Private/PackageDependencyData.cpp
+++ b/Engine/Source/Runtime/AssetRegistry/Private/PackageDependencyData.cpp
@@ -2,7 +2,7 @@
 
 #include "PackageDependencyData.h"
 
-FName FPackageDependencyData::GetImportPackageName(const TArray<FObjectImport>& ImportMap, int32 ImportIndex)
+FName FPackageDependencyData::GetImportPackageName(TConstArrayView<FObjectImport> ImportMap, int32 ImportIndex)
 {
 	for (int32 NumCycles = 0; NumCycles < ImportMap.Num(); ++NumCycles)
 	{
@@ -30,7 +30,7 @@ FName FPackageDependencyData::GetImportPackageName(const TArray<FObjectImport>&
 	return NAME_None;
 }
 
-void FPackageDependencyData::LoadDependenciesFromPackageHeader(FName SourcePackageName, TArray<FObjectImport>& ImportMap,
+void FPackageDependencyData::LoadDependenciesFromPackageHeader(FName SourcePackageName, TConstArrayView<FObjectImport> ImportMap,
 	TArray<FName>& SoftPackageReferenceList, TMap<FPackageIndex, TArray<FName>>& SearchableNames,
 	TBitArray<>& ImportUsedInGame, TBitArray<>& SoftPackageUsedInGame)
 {
@@ -77,7 +77,7 @@ void FPackageDependencyData::LoadDependenciesFromPackageHeader(FName SourcePacka
 			{
 				continue;
 			}
-			FObjectImport& Resource = ImportMap[ImportIndex];
+			const FObjectImport& Resource = ImportMap[ImportIndex];
 			FPackageIndex OuterLinkerIndex = Resource.OuterIndex;
 			if (!OuterLinkerIndex.IsNull())
 			{
diff --git a/Engine/Source/Runtime/AssetRegistry/Private/PackageDependencyData.h b/Engine/Source/Runtime/AssetRegistry/Private/PackageDependencyData.h
--- a/Engine/Source/Runtime/AssetRegistry/Private/PackageDependencyData.h
+++ b/Engine/Source/Runtime/AssetRegistry/Private/PackageDependencyData.h
@@ -6,6 +6,7 @@
 
 #include "AssetRegistry/AssetData.h"
 #include "AssetRegistry/IAssetRegistry.h"
+#include "Containers/ArrayView.h"
 #include "Containers/BitArray.h"
 #include "UObject/Linker.h"
 
@@ -63,7 +64,7 @@ public:
 		Ar << SearchableNameDependencies;
 	}
 
-	void LoadDependenciesFromPackageHeader(FName PackageName, TArray<FObjectImport>& ImportMap,
+	void LoadDependenciesFromPackageHeader(FName PackageName, TConstArrayView<FObjectImport> ImportMap,
 		TArray<FName>& SoftPackageReferenceList, TMap<FPackageIndex, TArray<FName>>& SearchableNames,
 		TBitArray<>& ImportUsedInGame, TBitArray<>& SoftPackageUsedInGame);
 
@@ -77,7 +78,7 @@ public:
 	}
 
 private:
-	FName GetImportPackageName(const TArray<FObjectImport>& ImportMap, int32 ImportIndex);
+	FName GetImportPackageName(TConstArrayView<FObjectImport> ImportMap, int32 ImportIndex);
 
 
 };
diff --git a/Engine/Source/Runtime/AssetRegistry/Private/PackageReader.cpp b/Engine/Source/Runtime/AssetRegistry/Private/PackageReader.cpp
--- a/Engine/Source/Runtime/AssetRegistry/Private/PackageReader.cpp
+++ b/Engine/Source/Runtime/AssetRegistry/Private/PackageReader.cpp
@@ -200,7 +200,10 @@ bool FPackageReader::StartSerializeSection(int64 Offset)
 	{\
 		FFormatNamedArguments CorruptPackageWarningArguments; \
 		CorruptPackageWarningArguments.Add(TEXT("FileName"), FText::FromString(PackageFileName)); \
-		FMessageLog("AssetRegistry").Warning(FText::Format(NSLOCTEXT("AssetRegistry", MessageKey, "Cannot read AssetRegistry Data in {FileName}, skipping it. Error: " MessageKey "."), CorruptPackageWarningArguments)); \
+		UE_LOG(LogAssetRegistry, Warning, TEXT("%s"), \
+			*FText::Format(NSLOCTEXT("AssetRegistry", MessageKey, \
+			"Cannot read AssetRegistry Data in {FileName}, skipping it. Error: " MessageKey "."), \
+			CorruptPackageWarningArguments).ToString()); \
 	} while (false)
 
 bool FPackageReader::ReadAssetRegistryData(TArray<FAssetData*>& AssetDataList, bool& bOutIsCookedWithoutAssetData)
@@ -211,6 +214,19 @@ bool FPackageReader::ReadAssetRegistryData(TArray<FAssetData*>& AssetDataList, b
 		return ReadAssetRegistryDataFromCookedPackage(AssetDataList, bOutIsCookedWithoutAssetData);
 	}
 
+	if (!SerializeNameMap())
+	{
+		return false;
+	}
+	if (!SerializeImportMap())
+	{
+		return false;
+	}
+	if (!SerializeExportMap())
+	{
+		return false;
+	}
+
 	if (!StartSerializeSection(PackageFileSummary.AssetRegistryDataOffset))
 	{
 		if (!ReadAssetDataFromThumbnailCache(AssetDataList))
@@ -232,7 +248,8 @@ bool FPackageReader::ReadAssetRegistryData(TArray<FAssetData*>& AssetDataList, b
 	using namespace UE::AssetRegistry;
 
 	EReadPackageDataMainErrorCode ErrorCode;
-	if (!ReadPackageDataMain(*this, PackageName, PackageFileSummary, AssetRegistryDependencyDataOffset, AssetDataList, ErrorCode))
+	if (!ReadPackageDataMain(*this, PackageName, PackageFileSummary, AssetRegistryDependencyDataOffset, AssetDataList, ErrorCode,
+		&ImportMap, &ExportMap))
 	{
 		switch (ErrorCode)
 		{
@@ -256,12 +273,12 @@ bool FPackageReader::ReadAssetRegistryData(TArray<FAssetData*>& AssetDataList, b
 }
 
 bool FPackageReader::SerializeAssetRegistryDependencyData(TBitArray<>& OutImportUsedInGame, TBitArray<>& OutSoftPackageUsedInGame,
-	const TArray<FObjectImport>& ImportMap, const TArray<FName>& SoftPackageReferenceList)
+	const TArray<FObjectImport>& InImportMap, const TArray<FName>& SoftPackageReferenceList)
 {
 	if (AssetRegistryDependencyDataOffset == INDEX_NONE)
 	{
 		// For old package versions that did not write out the dependency flags, set default values of the flags
-		OutImportUsedInGame.Init(true, ImportMap.Num());
+		OutImportUsedInGame.Init(true, InImportMap.Num());
 		OutSoftPackageUsedInGame.Init(true, SoftPackageReferenceList.Num());
 		return true;
 	}
@@ -272,7 +289,7 @@ bool FPackageReader::SerializeAssetRegistryDependencyData(TBitArray<>& OutImport
 	}
 
 	if (!UE::AssetRegistry::ReadPackageDataDependencies(*this, OutImportUsedInGame, OutSoftPackageUsedInGame) ||
-		OutImportUsedInGame.Num() != ImportMap.Num() ||
+		OutImportUsedInGame.Num() != InImportMap.Num() ||
 		OutSoftPackageUsedInGame.Num() != SoftPackageReferenceList.Num())
 	{
 		UE_PACKAGEREADER_CORRUPTPACKAGE_WARNING("SerializeAssetRegistryDependencyData", PackageFilename);
@@ -380,21 +397,19 @@ bool FPackageReader::ReadAssetRegistryDataFromCookedPackage(TArray<FAssetData*>&
 	{
 		const FString PackagePath = FPackageName::GetLongPackagePath(PackageName);
 
-		TArray<FObjectImport> ImportMap;
-		TArray<FObjectExport> ExportMap;
 		if (!SerializeNameMap())
 		{
 			return false;
 		}
-		if (!SerializeImportMap(ImportMap))
+		if (!SerializeImportMap())
 		{
 			return false;
 		}
-		if (!SerializeExportMap(ExportMap))
+		if (!SerializeExportMap())
 		{
 			return false;
 		}
-		for (FObjectExport& Export : ExportMap)
+		for (const FObjectExport& Export : ExportMap)
 		{
 			if (Export.bIsAsset)
 			{
@@ -449,8 +464,7 @@ bool FPackageReader::ReadDependencyData(FPackageDependencyData& OutDependencyDat
 		return false;
 	}
 
-	TArray<FObjectImport> ImportMap;
-	if (!SerializeImportMap(ImportMap))
+	if (!SerializeImportMap())
 	{
 		return false;
 	}
@@ -509,6 +523,10 @@ bool FPackageReader::ReadDependencyData(FPackageDependencyData& OutDependencyDat
 
 bool FPackageReader::SerializeNameMap()
 {
+	if (NameMap.Num() > 0)
+	{
+		return true;
+	}
 	if( PackageFileSummary.NameCount > 0 )
 	{
 		if (!StartSerializeSection(PackageFileSummary.NameOffset))
@@ -532,6 +550,7 @@ bool FPackageReader::SerializeNameMap()
 			if (IsError())
 			{
 				UE_PACKAGEREADER_CORRUPTPACKAGE_WARNING("SerializeNameMapInvalidName", PackageFilename);
+				NameMap.Reset();
 				return false;
 			}
 			NameMap.Add(FName(NameEntry));
@@ -541,8 +560,13 @@ bool FPackageReader::SerializeNameMap()
 	return true;
 }
 
-bool FPackageReader::SerializeImportMap(TArray<FObjectImport>& OutImportMap)
+bool FPackageReader::SerializeImportMap()
 {
+	if (ImportMap.Num() > 0)
+	{
+		return true;
+	}
+
 	if( PackageFileSummary.ImportCount > 0 )
 	{
 		if (!StartSerializeSection(PackageFileSummary.ImportOffset))
@@ -557,13 +581,14 @@ bool FPackageReader::SerializeImportMap(TArray<FObjectImport>& OutImportMap)
 			UE_PACKAGEREADER_CORRUPTPACKAGE_WARNING("SerializeImportMapInvalidImportCount", PackageFilename);
 			return false;
 		}
+		ImportMap.Reserve(PackageFileSummary.ImportCount);
 		for ( int32 ImportMapIdx = 0; ImportMapIdx < PackageFileSummary.ImportCount; ++ImportMapIdx )
 		{
-			FObjectImport& Import = OutImportMap.Emplace_GetRef();
-			*this << Import;
+			*this << ImportMap.Emplace_GetRef();
 			if (IsError())
 			{
 				UE_PACKAGEREADER_CORRUPTPACKAGE_WARNING("SerializeImportMapInvalidImport", PackageFilename);
+				ImportMap.Reset();
 				return false;
 			}
 		}
@@ -575,7 +600,7 @@ bool FPackageReader::SerializeImportMap(TArray<FObjectImport>& OutImportMap)
 static FName CoreUObjectPackageName(TEXT("/Script/CoreUObject"));
 static FName ScriptStructName(TEXT("ScriptStruct"));
 
-bool FPackageReader::SerializeImportedClasses(const TArray<FObjectImport>& ImportMap, TArray<FName>& OutClassNames)
+bool FPackageReader::SerializeImportedClasses(const TArray<FObjectImport>& InImportMap, TArray<FName>& OutClassNames)
 {
 	OutClassNames.Reset();
 
@@ -617,9 +642,9 @@ bool FPackageReader::SerializeImportedClasses(const TArray<FObjectImport>& Impor
 	// as a serialized-external-struct-pointer then we will miss it. In a future implementation we will 
 	// change the PackageReader to report all imports, and allow the AssetRegistry to decide which ones
 	// are classes based on its class database.
-	for (int32 ImportIndex = 0; ImportIndex < ImportMap.Num(); ++ImportIndex)
+	for (int32 ImportIndex = 0; ImportIndex < InImportMap.Num(); ++ImportIndex)
 	{
-		const FObjectImport& ObjectImport = ImportMap[ImportIndex];
+		const FObjectImport& ObjectImport = InImportMap[ImportIndex];
 		if (ObjectImport.ClassPackage == CoreUObjectPackageName && ObjectImport.ClassName == ScriptStructName)
 		{
 			ClassImportIndices.Add(ImportIndex);
@@ -632,7 +657,7 @@ bool FPackageReader::SerializeImportedClasses(const TArray<FObjectImport>& Impor
 	{
 		ParentChain.Reset();
 		ClassObjectPath.Reset();
-		if (!ImportMap.IsValidIndex(ClassImportIndex))
+		if (!InImportMap.IsValidIndex(ClassImportIndex))
 		{
 			UE_PACKAGEREADER_CORRUPTPACKAGE_WARNING("SerializeImportedClassesInvalidClassIndex", PackageFilename);
 			return false;
@@ -641,12 +666,12 @@ bool FPackageReader::SerializeImportedClasses(const TArray<FObjectImport>& Impor
 		int32 CurrentParentIndex = ClassImportIndex;
 		for (;;)
 		{
-			const FObjectImport& ObjectImport = ImportMap[CurrentParentIndex];
+			const FObjectImport& ObjectImport = InImportMap[CurrentParentIndex];
 			ParentChain.Add(ObjectImport.ObjectName);
 			if (ObjectImport.OuterIndex.IsImport())
 			{
 				CurrentParentIndex = ObjectImport.OuterIndex.ToImport();
-				if (!ImportMap.IsValidIndex(CurrentParentIndex))
+				if (!InImportMap.IsValidIndex(CurrentParentIndex))
 				{
 					UE_PACKAGEREADER_CORRUPTPACKAGE_WARNING("SerializeImportedClassesInvalidImportInParentChain",
 						PackageFilename);
@@ -689,8 +714,13 @@ bool FPackageReader::SerializeImportedClasses(const TArray<FObjectImport>& Impor
 	return true;
 }
 
-bool FPackageReader::SerializeExportMap(TArray<FObjectExport>& OutExportMap)
+bool FPackageReader::SerializeExportMap()
 {
+	if (ExportMap.Num() > 0)
+	{
+		return true;
+	}
+
 	if (PackageFileSummary.ExportCount > 0)
 	{
 		if (!StartSerializeSection(PackageFileSummary.ExportOffset))
@@ -705,13 +735,14 @@ bool FPackageReader::SerializeExportMap(TArray<FObjectExport>& OutExportMap)
 			UE_PACKAGEREADER_CORRUPTPACKAGE_WARNING("SerializeExportMapInvalidExportCount", PackageFilename);
 			return false;
 		}
+		ExportMap.Reserve(PackageFileSummary.ExportCount);
 		for (int32 ExportMapIdx = 0; ExportMapIdx < PackageFileSummary.ExportCount; ++ExportMapIdx)
 		{
-			FObjectExport* Export = new(OutExportMap)FObjectExport;
-			*this << *Export;
+			*this << ExportMap.Emplace_GetRef();
 			if (IsError())
 			{
 				UE_PACKAGEREADER_CORRUPTPACKAGE_WARNING("SerializeExportMapInvalidExport", PackageFilename);
+				ExportMap.Reset();
 				return false;
 			}
 		}
@@ -924,25 +955,41 @@ namespace UE::AssetRegistry
 			}
 		}
 	};
-	FString ReconstructFullClassPath(FArchive& BinaryArchive, const FString& PackageName, const FPackageFileSummary& PackageFileSummary, const FString& AssetClassName)
+	FString ReconstructFullClassPath(FArchive& BinaryArchive, const FString& PackageName, const FPackageFileSummary& PackageFileSummary, const FString& AssetClassName,
+		const TArray<FObjectImport>* InImports = nullptr, const TArray<FObjectExport>* InExports = nullptr)
 	{
-		FLinkerTables LinkerTables;
-		FNameMapAwareArchive NameMapArchive(BinaryArchive);
-		NameMapArchive.SerializeNameMap(PackageFileSummary);
 		FName ClassFName(*AssetClassName);
+		FLinkerTables LinkerTables;
+		if (!InImports || !InExports)
+		{
+			FNameMapAwareArchive NameMapArchive(BinaryArchive);
+			NameMapArchive.SerializeNameMap(PackageFileSummary);
 
-		// Load the linker tables
-		BinaryArchive.Seek(PackageFileSummary.ImportOffset);
-		for (int32 ImportMapIndex = 0; ImportMapIndex < PackageFileSummary.ImportCount; ++ImportMapIndex)
+			// Load the linker tables
+			if (!InImports)
+			{
+				BinaryArchive.Seek(PackageFileSummary.ImportOffset);
+				for (int32 ImportMapIndex = 0; ImportMapIndex < PackageFileSummary.ImportCount; ++ImportMapIndex)
+				{
+					NameMapArchive << LinkerTables.ImportMap.Emplace_GetRef();
+				}
+			}
+			if (!InExports)
+			{
+				BinaryArchive.Seek(PackageFileSummary.ExportOffset);
+				for (int32 ExportMapIndex = 0; ExportMapIndex < PackageFileSummary.ExportCount; ++ExportMapIndex)
+				{
+					NameMapArchive << LinkerTables.ExportMap.Emplace_GetRef();
+				}
+			}
+		}
+		if (InImports)
 		{
-			FObjectImport* Import = new(LinkerTables.ImportMap)FObjectImport;
-			NameMapArchive << *Import;
+			LinkerTables.ImportMap = *InImports;
 		}
-		BinaryArchive.Seek(PackageFileSummary.ExportOffset);
-		for (int32 ExportMapIndex = 0; ExportMapIndex < PackageFileSummary.ExportCount; ++ExportMapIndex)
+		if (InExports)
 		{
-			FObjectExport* Export = new(LinkerTables.ExportMap)FObjectExport;
-			NameMapArchive << *Export;
+			LinkerTables.ExportMap = *InExports;
 		}
 
 		FString ClassPathName;
@@ -979,7 +1026,8 @@ namespace UE::AssetRegistry
 	}
 
 	// See the corresponding WritePackageData defined in SavePackageUtilities.cpp in CoreUObject module
-	bool ReadPackageDataMain(FArchive& BinaryArchive, const FString& PackageName, const FPackageFileSummary& PackageFileSummary, int64& OutDependencyDataOffset, TArray<FAssetData*>& OutAssetDataList, EReadPackageDataMainErrorCode& OutError)
+	bool ReadPackageDataMain(FArchive& BinaryArchive, const FString& PackageName, const FPackageFileSummary& PackageFileSummary, int64& OutDependencyDataOffset,
+		TArray<FAssetData*>& OutAssetDataList, EReadPackageDataMainErrorCode& OutError, const TArray<FObjectImport>* InImports, const TArray<FObjectExport>* InExports)
 	{
 		OutError = EReadPackageDataMainErrorCode::Unknown;
 
@@ -1096,7 +1144,8 @@ namespace UE::AssetRegistry
 			if (!ObjectClassName.IsEmpty() && FPackageName::IsShortPackageName(ObjectClassName))
 			{
 				int64 CurrentPos = BinaryArchive.Tell();
-				ObjectClassName = ReconstructFullClassPath(BinaryArchive, PackageName, PackageFileSummary, ObjectClassName);
+				ObjectClassName = ReconstructFullClassPath(BinaryArchive, PackageName, PackageFileSummary,
+					ObjectClassName, InImports, InExports);
 				BinaryArchive.Seek(CurrentPos);
 			}
 			OutAssetDataList.Add(new FAssetData(PackageName, ObjectPath, FTopLevelAssetPath(ObjectClassName), MoveTemp(TagsAndValues), PackageFileSummary.ChunkIDs, PackageFileSummary.GetPackageFlags()));
diff --git a/Engine/Source/Runtime/AssetRegistry/Private/PackageReader.h b/Engine/Source/Runtime/AssetRegistry/Private/PackageReader.h
--- a/Engine/Source/Runtime/AssetRegistry/Private/PackageReader.h
+++ b/Engine/Source/Runtime/AssetRegistry/Private/PackageReader.h
@@ -85,13 +85,13 @@ private:
 
 	/** Serializers for different package maps */
 	bool SerializeNameMap();
-	bool SerializeImportMap(TArray<FObjectImport>& OutImportMap);
-	bool SerializeExportMap(TArray<FObjectExport>& OutExportMap);
-	bool SerializeImportedClasses(const TArray<FObjectImport>& ImportMap, TArray<FName>& OutClassNames);
+	bool SerializeImportMap();
+	bool SerializeExportMap();
+	bool SerializeImportedClasses(const TArray<FObjectImport>& InImportMap, TArray<FName>& OutClassNames);
 	bool SerializeSoftPackageReferenceList(TArray<FName>& OutSoftPackageReferenceList);
 	bool SerializeSearchableNamesMap(FLinkerTables& OutSearchableNames);
 	bool SerializeAssetRegistryDependencyData(TBitArray<>& OutImportUsedInGame, TBitArray<>& OutSoftPackageUsedInGame,
-		const TArray<FObjectImport>& ImportMap, const TArray<FName>& SoftPackageReferenceList);
+		const TArray<FObjectImport>& InImportMap, const TArray<FName>& SoftPackageReferenceList);
 	bool SerializePackageTrailer(FAssetPackageData& PackageData);
 
 	/** Returns flags the asset package was saved with */
@@ -103,6 +103,8 @@ private:
 	FArchive* Loader;
 	FPackageFileSummary PackageFileSummary;
 	TArray<FName> NameMap;
+	TArray<FObjectImport> ImportMap;
+	TArray<FObjectExport> ExportMap;
 	int64 PackageFileSize;
 	int64 AssetRegistryDependencyDataOffset;
 	bool bLoaderOwner;
diff --git a/Engine/Source/Runtime/AssetRegistry/Public/AssetRegistry/IAssetRegistry.h b/Engine/Source/Runtime/AssetRegistry/Public/AssetRegistry/IAssetRegistry.h
--- a/Engine/Source/Runtime/AssetRegistry/Public/AssetRegistry/IAssetRegistry.h
+++ b/Engine/Source/Runtime/AssetRegistry/Public/AssetRegistry/IAssetRegistry.h
@@ -25,6 +25,8 @@ struct FAssetRegistrySerializationOptions;
 class FAssetRegistryState;
 class FDependsNode;
 struct FPackageFileSummary;
+struct FObjectExport;
+struct FObjectImport;
 
 namespace EAssetAvailability
 {
@@ -693,7 +695,9 @@ namespace AssetRegistry
 	};
 	// Functions to read and write the data used by the AssetRegistry in each package; the format of this data is separate from the format of the data in the asset registry
 	// WritePackageData is declared in AssetRegistryInterface.h, in the CoreUObject module, because it is needed by SavePackage in CoreUObject
-	ASSETREGISTRY_API bool ReadPackageDataMain(FArchive& BinaryArchive, const FString& PackageName, const FPackageFileSummary& PackageFileSummary, int64& OutDependencyDataOffset, TArray<FAssetData*>& OutAssetDataList, EReadPackageDataMainErrorCode& OutError);
+	ASSETREGISTRY_API bool ReadPackageDataMain(FArchive& BinaryArchive, const FString& PackageName, const FPackageFileSummary& PackageFileSummary,
+		int64& OutDependencyDataOffset, TArray<FAssetData*>& OutAssetDataList, EReadPackageDataMainErrorCode& OutError,
+		const TArray<FObjectImport>* InImports = nullptr, const TArray<FObjectExport>* InExports = nullptr);
 	ASSETREGISTRY_API bool ReadPackageDataDependencies(FArchive& BinaryArchive, TBitArray<>& OutImportUsedInGame, TBitArray<>& OutSoftPackageUsedInGame);
 
 	/**
