AssetRegistry Blueprint functions: Expose GetDerivedClassNames and GetAncestorClassNames, and add the AssetRegistryHelpers function GetBlueprintAssets.
#rnx

diff --git a/Engine/Source/Runtime/AssetRegistry/Private/AssetRegistryHelpers.cpp b/Engine/Source/Runtime/AssetRegistry/Private/AssetRegistryHelpers.cpp
--- a/Engine/Source/Runtime/AssetRegistry/Private/AssetRegistryHelpers.cpp
+++ b/Engine/Source/Runtime/AssetRegistry/Private/AssetRegistryHelpers.cpp
@@ -1,9 +1,12 @@
 // Copyright Epic Games, Inc. All Rights Reserved.
 
 #include "AssetRegistry/AssetRegistryHelpers.h"
-#include "Modules/ModuleManager.h"
 #include "AssetRegistry/IAssetRegistry.h"
 #include "AssetRegistry.h"
+#include "Blueprint/BlueprintSupport.h"
+#include "Misc/PackageName.h"
+#include "Modules/ModuleManager.h"
+#include "UObject/LinkerLoad.h"
 
 TScriptInterface<IAssetRegistry> UAssetRegistryHelpers::GetAssetRegistry()
 {
@@ -83,6 +86,74 @@ FARFilter UAssetRegistryHelpers::SetFilterTagsAndValues(const FARFilter& InFilte
 	return FilterCopy;
 }
 
+void UAssetRegistryHelpers::GetBlueprintAssets(const FARFilter& InFilter, TArray<FAssetData>& OutAssetData)
+{
+	IAssetRegistry& AssetRegistry = IAssetRegistry::GetChecked();
+
+	FARFilter Filter(InFilter);
+	PRAGMA_DISABLE_DEPRECATION_WARNINGS;
+	UE_CLOG(!InFilter.ClassNames.IsEmpty(), LogCore, Error,
+		TEXT("ARFilter.ClassNames is not supported by UAssetRegistryHelpers::GetBlueprintAssets and will be ignored."));
+	Filter.ClassNames.Empty();
+	PRAGMA_ENABLE_DEPRECATION_WARNINGS;
+
+	// Expand list of classes to include derived classes
+	TArray<FTopLevelAssetPath> BlueprintParentClassPathRoots = MoveTemp(Filter.ClassPaths);
+	TSet<FTopLevelAssetPath> BlueprintParentClassPaths;
+	if (Filter.bRecursiveClasses)
+	{
+		AssetRegistry.GetDerivedClassNames(BlueprintParentClassPathRoots, TSet<FTopLevelAssetPath>(), BlueprintParentClassPaths);
+	}
+	else
+	{
+		BlueprintParentClassPaths.Append(BlueprintParentClassPathRoots);
+	}
+
+	// Search for all blueprints and then check BlueprintParentClassPaths in the results
+	Filter.ClassPaths.Reset(1);
+	Filter.ClassPaths.Add(FTopLevelAssetPath(FName(TEXT("/Script/Engine")), FName(TEXT("BlueprintCore"))));
+	Filter.bRecursiveClasses = true;
+
+	auto FilterLambda = [&OutAssetData, &BlueprintParentClassPaths](const FAssetData& AssetData)
+	{
+		// Verify blueprint class
+		if (BlueprintParentClassPaths.IsEmpty() || IsAssetDataBlueprintOfClassSet(AssetData, BlueprintParentClassPaths))
+		{
+			OutAssetData.Add(AssetData);
+		}
+		return true;
+	};
+	AssetRegistry.EnumerateAssets(Filter, FilterLambda);
+}
+
+bool UAssetRegistryHelpers::IsAssetDataBlueprintOfClassSet(const FAssetData& AssetData, const TSet<FTopLevelAssetPath>& ClassNameSet)
+{
+	const FString ParentClassFromData = AssetData.GetTagValueRef<FString>(FBlueprintTags::ParentClassPath);
+	if (!ParentClassFromData.IsEmpty())
+	{
+		const FTopLevelAssetPath ClassObjectPath(FPackageName::ExportTextPathToObjectPath(ParentClassFromData));
+		const FName ClassName = ClassObjectPath.GetAssetName();
+
+		TArray<FTopLevelAssetPath> ValidNames;
+		ValidNames.Add(ClassObjectPath);
+		// Check for redirected name
+		FTopLevelAssetPath RedirectedName = FTopLevelAssetPath(FLinkerLoad::FindNewPathNameForClass(ClassObjectPath.ToString(), false));
+		if (!RedirectedName.IsNull())
+		{
+			ValidNames.Add(RedirectedName);
+		}
+		for (const FTopLevelAssetPath& ValidName : ValidNames)
+		{
+			if (ClassNameSet.Contains(ValidName))
+			{
+				// Our parent class is in the class name set
+				return true;
+			}
+		}
+	}
+	return false;
+}
+
 UAssetRegistryHelpers::FTemporaryCachingModeScope::FTemporaryCachingModeScope(bool InTempCachingMode)
 {
 	PreviousCachingMode = UAssetRegistryHelpers::GetAssetRegistry()->GetTemporaryCachingMode();
diff --git a/Engine/Source/Runtime/AssetRegistry/Public/AssetRegistry/AssetRegistryHelpers.h b/Engine/Source/Runtime/AssetRegistry/Public/AssetRegistry/AssetRegistryHelpers.h
--- a/Engine/Source/Runtime/AssetRegistry/Public/AssetRegistry/AssetRegistryHelpers.h
+++ b/Engine/Source/Runtime/AssetRegistry/Public/AssetRegistry/AssetRegistryHelpers.h
@@ -6,6 +6,7 @@
 #include "AssetRegistry/AssetData.h"
 #include "AssetRegistry/IAssetRegistry.h"
 #include "Containers/Array.h"
+#include "Containers/Set.h"
 #include "Containers/UnrealString.h"
 #include "CoreMinimal.h"
 #include "UObject/NameTypes.h"
@@ -96,6 +97,10 @@ public:
 	UFUNCTION(BlueprintPure, Category = "Asset Registry")
 	static FARFilter SetFilterTagsAndValues(const FARFilter& InFilter, const TArray<FTagAndValue>& InTagsAndValues);
 
+	/** Gets asset data for all blueprint assets that match the filter. ClassPaths in the filter specify the blueprint's parent class. */
+	UFUNCTION(BlueprintPure, Category = "Asset Registry", meta=(ScriptMethod))
+	static void GetBlueprintAssets(const FARFilter& InFilter, TArray<FAssetData>& OutAssetData);
+
 	/** Enable/disable asset registry caching mode for the duration of the scope */
 	struct ASSETREGISTRY_API FTemporaryCachingModeScope
 	{
@@ -105,4 +110,8 @@ public:
 	private:
 		bool PreviousCachingMode;
 	};
+
+	/** Checks to see if the given asset data is a blueprint with a base class in the ClassNameSet. This checks the parent asset tag */
+	static bool ASSETREGISTRY_API IsAssetDataBlueprintOfClassSet(const FAssetData& AssetData, const TSet<FTopLevelAssetPath>& ClassNameSet);
+
 };
diff --git a/Engine/Source/Runtime/AssetRegistry/Public/AssetRegistry/IAssetRegistry.h b/Engine/Source/Runtime/AssetRegistry/Public/AssetRegistry/IAssetRegistry.h
--- a/Engine/Source/Runtime/AssetRegistry/Public/AssetRegistry/IAssetRegistry.h
+++ b/Engine/Source/Runtime/AssetRegistry/Public/AssetRegistry/IAssetRegistry.h
@@ -346,12 +346,14 @@ public:
 	virtual void StripAssetRegistryKeyForObject(FName ObjectPath, FName Key) {}
 
 	/** Returns true if the specified ClassName's ancestors could be found. If so, OutAncestorClassNames is a list of all its ancestors. This can be slow if temporary caching mode is not on */
+	UFUNCTION(BlueprintCallable, BlueprintPure = false, Category = "AssetRegistry", meta = (DisplayName = "GetAncestorClassNames", ScriptName = "GetAncestorClassNames"))
 	virtual bool GetAncestorClassNames(FTopLevelAssetPath ClassPathName, TArray<FTopLevelAssetPath>& OutAncestorClassNames) const = 0;
 
 	UE_DEPRECATED(5.1, "Class names are now represented by path names. Please use a version of this function that uses FTopLevelAssetPath.")
 	virtual bool GetAncestorClassNames(FName ClassName, TArray<FName>& OutAncestorClassNames) const = 0;
 
 	/** Returns the names of all classes derived by the supplied class names, excluding any classes matching the excluded class names. This can be slow if temporary caching mode is not on */
+	UFUNCTION(BlueprintCallable, BlueprintPure = false, Category = "AssetRegistry", meta = (DisplayName = "GetDerivedClassNames", ScriptName = "GetDerivedClassNames"))
 	virtual void GetDerivedClassNames(const TArray<FTopLevelAssetPath>& ClassNames, const TSet<FTopLevelAssetPath>& ExcludedClassNames, TSet<FTopLevelAssetPath>& OutDerivedClassNames) const = 0;
 	
 	UE_DEPRECATED(5.1, "Class names are now represented by path names. Please use a version of this function that uses FTopLevelAssetPath.")
diff --git a/Engine/Source/Runtime/Engine/Private/AssetManager.cpp b/Engine/Source/Runtime/Engine/Private/AssetManager.cpp
--- a/Engine/Source/Runtime/Engine/Private/AssetManager.cpp
+++ b/Engine/Source/Runtime/Engine/Private/AssetManager.cpp
@@ -5,6 +5,7 @@
 #include "Engine/PrimaryAssetLabel.h"
 #include "AssetRegistry/AssetData.h"
 #include "AssetRegistry/ARFilter.h"
+#include "AssetRegistry/AssetRegistryHelpers.h"
 #include "Containers/StringView.h"
 #include "Engine/Engine.h"
 #include "Engine/BlueprintGeneratedClass.h"
@@ -502,33 +503,7 @@ FPrimaryAssetId UAssetManager::DeterminePrimaryAssetIdForObject(const UObject* O
 
 bool UAssetManager::IsAssetDataBlueprintOfClassSet(const FAssetData& AssetData, const TSet<FTopLevelAssetPath>& ClassNameSet) const
 {
-	const FString ParentClassFromData = AssetData.GetTagValueRef<FString>(FBlueprintTags::ParentClassPath);
-	if (!ParentClassFromData.IsEmpty())
-	{
-		const FTopLevelAssetPath ClassObjectPath(FPackageName::ExportTextPathToObjectPath(ParentClassFromData));
-		const FName ClassName = ClassObjectPath.GetAssetName();
-
-		TArray<FTopLevelAssetPath> ValidNames;
-		ValidNames.Add(ClassObjectPath);
-#if WITH_EDITOR
-		// Check for redirected name
-		FTopLevelAssetPath RedirectedName = FTopLevelAssetPath(FLinkerLoad::FindNewPathNameForClass(ClassObjectPath.ToString(), false));
-		if (!RedirectedName.IsNull())
-		{
-			// @todo: we need a full path
-			ValidNames.Add(RedirectedName);
-		}
-#endif
-		for (const FTopLevelAssetPath& ValidName : ValidNames)
-		{
-			if (ClassNameSet.Contains(ValidName))
-			{
-				// Our parent class is in the class name set
-				return true;
-			}
-		}
-	}
-	return false;
+	return UAssetRegistryHelpers::IsAssetDataBlueprintOfClassSet(AssetData, ClassNameSet);
 }
 
 int32 UAssetManager::SearchAssetRegistryPaths(TArray<FAssetData>& OutAssetDataList, const FAssetManagerSearchRules& Rules) const
