Updated AssetRegistry to have UnsafeTypeCastWarningLevel set to WarningLevel.Error so it matches other modules. This speeds up compile times and removes the need to have a 338 MB Shared PCH be created specifically for AssetRegistry.

diff --git a/Engine/Source/Runtime/AssetRegistry/AssetRegistry.Build.cs b/Engine/Source/Runtime/AssetRegistry/AssetRegistry.Build.cs
--- a/Engine/Source/Runtime/AssetRegistry/AssetRegistry.Build.cs
+++ b/Engine/Source/Runtime/AssetRegistry/AssetRegistry.Build.cs
@@ -21,6 +21,6 @@ public class AssetRegistry : ModuleRules
 			DynamicallyLoadedModuleNames.AddRange(new string[] { "DirectoryWatcher" });
 		}
 
-		UnsafeTypeCastWarningLevel = WarningLevel.Warning;
+		UnsafeTypeCastWarningLevel = WarningLevel.Error;
 	}
 }
