Fix ensure when using Asset Manager to search for classes with core redirects where it tries to convert a short name to top level asset path
Change default behavior of FindPreviousPathNamesForClass to only return full paths as suggested by the name, and lower verbosity of asset registry warning that specifically handles the short names

diff --git a/Engine/Source/Runtime/AssetRegistry/Private/AssetRegistry.cpp b/Engine/Source/Runtime/AssetRegistry/Private/AssetRegistry.cpp
--- a/Engine/Source/Runtime/AssetRegistry/Private/AssetRegistry.cpp
+++ b/Engine/Source/Runtime/AssetRegistry/Private/AssetRegistry.cpp
@@ -5772,7 +5772,7 @@ void FAssetRegistryImpl::GetSubClasses_Recursive(Impl::FClassInheritanceContext&
 #if WITH_EDITOR
 		// Also add subclasses of any old names (from CoreRedirects) in case subclasses have not been resaved and have
 		// registered their ParentClass as the old name
-		TArray<FString> OldNames = FLinkerLoad::FindPreviousPathNamesForClass(InClassName.ToString(), false);
+		TArray<FString> OldNames = FLinkerLoad::FindPreviousPathNamesForClass(InClassName.ToString(), false, true);
 		for (const FString& OldNameString : OldNames)
 		{
 			if (FPackageName::IsShortPackageName(OldNameString))
@@ -5780,7 +5780,7 @@ void FAssetRegistryImpl::GetSubClasses_Recursive(Impl::FClassInheritanceContext&
 				// TODO: We do not have a way to handle these, because we don't allow searching subclasses for short path
 				// names. We should instead change the construction of CachedBPInheritanceMap and ReverseInheritanceMap 
 				// to modify the ParentClassName with CoreRedirects
-				UE_LOG(LogAssetRegistry, Display, TEXT("GetDerivedClassNames found CoreRedirect from ShortPathName %s; it will not be able to find old blueprint classes that were saved with that parent class before the Redirect was added."),
+				UE_LOG(LogAssetRegistry, Verbose, TEXT("GetDerivedClassNames found CoreRedirect from ShortPathName %s; it will not be able to find old blueprint classes that were saved with that parent class before the Redirect was added."),
 					*OldNameString);
 				continue;
 			}
diff --git a/Engine/Source/Runtime/CoreUObject/Private/UObject/LinkerLoad.cpp b/Engine/Source/Runtime/CoreUObject/Private/UObject/LinkerLoad.cpp
--- a/Engine/Source/Runtime/CoreUObject/Private/UObject/LinkerLoad.cpp
+++ b/Engine/Source/Runtime/CoreUObject/Private/UObject/LinkerLoad.cpp
@@ -5735,7 +5735,7 @@ TArray<FName> FLinkerLoad::FindPreviousNamesForClass(const FString& CurrentClass
 	return OldNames;
 }
 
-TArray<FString> FLinkerLoad::FindPreviousPathNamesForClass(const FString& CurrentClassPath, bool bIsInstance)
+TArray<FString> FLinkerLoad::FindPreviousPathNamesForClass(const FString& CurrentClassPath, bool bIsInstance, bool bIncludeShortNames)
 {
 	TArray<FString> OldNames;
 	TArray<FCoreRedirectObjectName> OldObjectNames;
@@ -5744,7 +5744,10 @@ TArray<FString> FLinkerLoad::FindPreviousPathNamesForClass(const FString& Curren
 	{
 		for (FCoreRedirectObjectName& OldObjectName : OldObjectNames)
 		{
-			OldNames.AddUnique(OldObjectName.ToString());
+			if (bIncludeShortNames || !OldObjectName.PackageName.IsNone())
+			{
+				OldNames.AddUnique(OldObjectName.ToString());
+			}
 		}
 	}
 
@@ -5755,7 +5758,10 @@ TArray<FString> FLinkerLoad::FindPreviousPathNamesForClass(const FString& Curren
 		{
 			for (FCoreRedirectObjectName& OldObjectName : OldObjectNames)
 			{
-				OldNames.AddUnique(OldObjectName.ToString());
+				if (bIncludeShortNames || !OldObjectName.PackageName.IsNone())
+				{
+					OldNames.AddUnique(OldObjectName.ToString());
+				}
 			}
 		}
 	}
diff --git a/Engine/Source/Runtime/CoreUObject/Public/UObject/LinkerLoad.h b/Engine/Source/Runtime/CoreUObject/Public/UObject/LinkerLoad.h
--- a/Engine/Source/Runtime/CoreUObject/Public/UObject/LinkerLoad.h
+++ b/Engine/Source/Runtime/CoreUObject/Public/UObject/LinkerLoad.h
@@ -298,9 +298,10 @@ public:
 	 * Utility functions to query the object name redirects list for previous names for a class
 	 * @param CurrentClassPath The current name of the class, with a full path
 	 * @param bIsInstance If true, we're an instance, so check instance only maps as well
-	 * @return Names without path of all classes that were redirected to this name. Empty if none found.
+	 * @param bIncludeShortNames If true, also include short names without a package
+	 * @return Full object paths of all classes that were redirected to this name. Empty if none found.
 	 */
-	COREUOBJECT_API static TArray<FString> FindPreviousPathNamesForClass(const FString& CurrentClassPath, bool bIsInstance);
+	COREUOBJECT_API static TArray<FString> FindPreviousPathNamesForClass(const FString& CurrentClassPath, bool bIsInstance, bool bIncludeShortNames = false);
 
 	/** 
 	 * Utility functions to query the object name redirects list for the current name for a class
