Removing FNames containing asset paths (/Path/To/Package.AssetName) or subobject paths (/Path/To/Package.AssetName:Inner).
These FNames can take 30mb+ of memory in a large project.
Refactor FSoftObjectPath to store a FTopLevelAssetPath instead of an FName for the top level asset part.
Refactor FAssetData to infer the object path from the PackageName and AssetName fields and an editor-only field called OptionalOuterPath for external actors which are non-top-level assets.
Refactor FAssetBundleData to only store top level asset references.
Increase EUnrealEngineObjectUE5Version for FSoftObjectPath invalidation.
Increase FAssetRegistryVersion to invalidate serialization of FAssetData and FAssetBundleData - asset registry serialization does not see EUnrealEngineObjectUE5Version.
Replace TMap in Asset Registry with TSet to avoid duplicating the key fields.
Remove object path FName fields from native blueprint make/break functions for FARFilter - those pins will be orphaned.
Refactor redirect collector to use soft object paths rather than FNames.

UE_DEPRECATED markers have been commented out in this change to allow separate submission of deprecation fixups across the engine.

diff --git a/Engine/Plugins/Editor/AssetManagerEditor/Source/AssetManagerEditor/Private/AssetManagerEditorModule.cpp b/Engine/Plugins/Editor/AssetManagerEditor/Source/AssetManagerEditor/Private/AssetManagerEditorModule.cpp
--- a/Engine/Plugins/Editor/AssetManagerEditor/Source/AssetManagerEditor/Private/AssetManagerEditorModule.cpp
+++ b/Engine/Plugins/Editor/AssetManagerEditor/Source/AssetManagerEditor/Private/AssetManagerEditorModule.cpp
@@ -1219,7 +1219,7 @@ bool FAssetManagerEditorModule::GetStringValueForCustomColumn(const FAssetData&
 		}
 		else
 		{
-			FAssetData PlatformData = CurrentRegistrySource->GetAssetByObjectPath(AssetData.ObjectPath);
+			FAssetData PlatformData = CurrentRegistrySource->GetAssetByObjectPath(AssetData.GetSoftObjectPath());
 			if (PlatformData.IsValid())
 			{
 				FoundChunks = MoveTemp(PlatformData.ChunkIDs);
@@ -1591,7 +1591,7 @@ void FAssetManagerEditorModule::SetCurrentRegistrySource(const FString& SourceNa
 
 				NewState->Serialize(SerializedAssetData, Options);
 
-				if (NewState->GetObjectPathToAssetDataMap().Num() > 0)
+				if (NewState->GetNumAssets() > 0)
 				{
 					bLoaded = true;
 					CurrentRegistrySource->SetRegistryState(NewState);
@@ -1622,11 +1622,9 @@ void FAssetManagerEditorModule::SetCurrentRegistrySource(const FString& SourceNa
 				// Iterate assets and look for chunks
 				const FAssetRegistryState* RegistryState = CurrentRegistrySource->GetOwnedRegistryState();
 				checkf(RegistryState, TEXT("Should be non-null because HasRegistry() && !bIsEditor"));
-				const TMap<FName, const FAssetData*>& AssetDataMap = RegistryState->GetObjectPathToAssetDataMap();
 
-				for (const TPair<FName, const FAssetData*>& Pair : AssetDataMap)
+				RegistryState->EnumerateAllAssets([&RegistryState, this](const FAssetData& AssetData)
 				{
-					const FAssetData& AssetData = *Pair.Value;
 					if (AssetData.ChunkIDs.Num() > 0)
 					{
 						TArray<FAssetIdentifier> ManagerAssets;
@@ -1671,7 +1669,7 @@ void FAssetManagerEditorModule::SetCurrentRegistrySource(const FString& SourceNa
 							}
 						}
 					}
-				}
+				});
 			}
 
 			CurrentRegistrySource->bManagementDataInitialized = true;
@@ -2293,7 +2291,7 @@ bool FAssetManagerEditorModule::WriteCollection(FName CollectionName, ECollectio
 	FText ResultsMessage;
 	bool bSuccess = false;
 
-	TSet<FName> ObjectPathsToAddToCollection;
+	TSet<FSoftObjectPath> ObjectPathsToAddToCollection;
 
 	FARFilter Filter;
 	Filter.PackageNames = PackageNames;
@@ -2302,7 +2300,7 @@ bool FAssetManagerEditorModule::WriteCollection(FName CollectionName, ECollectio
 	AssetRegistry->GetAssets(Filter, AssetsInPackages);
 	for (const FAssetData& AssetData : AssetsInPackages)
 	{
-		ObjectPathsToAddToCollection.Add(AssetData.ObjectPath);
+		ObjectPathsToAddToCollection.Add(AssetData.GetSoftObjectPath());
 	}
 
 	if (ObjectPathsToAddToCollection.Num() == 0)
@@ -2311,8 +2309,10 @@ bool FAssetManagerEditorModule::WriteCollection(FName CollectionName, ECollectio
 		ResultsMessage = FText::Format(LOCTEXT("NothingToAddToCollection", "Nothing to add to collection {0}"), FText::FromName(CollectionName));
 	}
 	else if (CreateOrEmptyCollection(CollectionName, ShareType))
-	{		
-		if (CollectionManager.AddToCollection(CollectionName, ECollectionShareType::CST_Local, ObjectPathsToAddToCollection.Array()))
+	{	
+PRAGMA_DISABLE_DEPRECATION_WARNINGS
+		if (CollectionManager.AddToCollection(CollectionName, ECollectionShareType::CST_Local, UE::SoftObjectPath::Private::ConvertSoftObjectPaths(ObjectPathsToAddToCollection.Array())))
+PRAGMA_ENABLE_DEPRECATION_WARNINGS
 		{
 			UE_LOG(LogAssetManagerEditor, Log, TEXT("Updated collection %s"), *CollectionName.ToString());
 			ResultsMessage = FText::Format(LOCTEXT("CreateCollectionSucceeded", "Updated collection {0}"), FText::FromName(CollectionName));
diff --git a/Engine/Plugins/Editor/AssetManagerEditor/Source/AssetManagerEditor/Public/AssetManagerEditorModule.h b/Engine/Plugins/Editor/AssetManagerEditor/Source/AssetManagerEditor/Public/AssetManagerEditorModule.h
--- a/Engine/Plugins/Editor/AssetManagerEditor/Source/AssetManagerEditor/Public/AssetManagerEditorModule.h
+++ b/Engine/Plugins/Editor/AssetManagerEditor/Source/AssetManagerEditor/Public/AssetManagerEditorModule.h
@@ -200,8 +200,7 @@ public:
 	void LoadRegistryTimestamp();
 
 	// Functions to forward on to the RegistryState, either the owned one or the editor's global State.
-
-	FAssetData GetAssetByObjectPath(FName ObjectPath) const
+	FAssetData GetAssetByObjectPath(const FSoftObjectPath& ObjectPath) const
 	{
 		if (bIsEditor)
 		{
@@ -218,6 +217,14 @@ public:
 		}
 	}
 
+	// UE_DEPRECATED(5.1, "Asset path FNames have been deprecated, use FSoftObjectPath instead.")
+	FAssetData GetAssetByObjectPath(FName ObjectPath) const
+	{
+PRAGMA_DISABLE_DEPRECATION_WARNINGS
+		return GetAssetByObjectPath(FSoftObjectPath(ObjectPath));
+PRAGMA_ENABLE_DEPRECATION_WARNINGS
+	}
+
 	bool GetDependencies(const FAssetIdentifier& AssetIdentifier, TArray<FAssetIdentifier>& OutDependencies,
 		UE::AssetRegistry::EDependencyCategory Category = UE::AssetRegistry::EDependencyCategory::All,
 		const UE::AssetRegistry::FDependencyQuery& Flags = UE::AssetRegistry::FDependencyQuery()) const
diff --git a/Engine/Source/Editor/UnrealEd/Classes/CookOnTheSide/CookOnTheFlyServer.h b/Engine/Source/Editor/UnrealEd/Classes/CookOnTheSide/CookOnTheFlyServer.h
--- a/Engine/Source/Editor/UnrealEd/Classes/CookOnTheSide/CookOnTheFlyServer.h
+++ b/Engine/Source/Editor/UnrealEd/Classes/CookOnTheSide/CookOnTheFlyServer.h
@@ -1099,7 +1099,7 @@ private:
 	 * @param RedirectedPaths map of original to redirected object paths
 	 * @return true if the Package contains a redirector false otherwise
 	 */
-	bool ContainsRedirector(const FName& PackageName, TMap<FName, FName>& RedirectedPaths) const;
+	bool ContainsRedirector(const FName& PackageName, TMap<FSoftObjectPath, FSoftObjectPath>& RedirectedPaths) const;
 	
 	/**
 	 * Prepares save by calling BeginCacheForCookedPlatformData on all UObjects in the package. 
diff --git a/Engine/Source/Editor/UnrealEd/Private/Commandlets/AssetRegistryGenerator.cpp b/Engine/Source/Editor/UnrealEd/Private/Commandlets/AssetRegistryGenerator.cpp
--- a/Engine/Source/Editor/UnrealEd/Private/Commandlets/AssetRegistryGenerator.cpp
+++ b/Engine/Source/Editor/UnrealEd/Private/Commandlets/AssetRegistryGenerator.cpp
@@ -985,7 +985,7 @@ void FAssetRegistryGenerator::InjectEncryptionData(FAssetRegistryState& TargetSt
 
 						if (AssetData->ChunkIDs.Num() > 1)
 						{
-							UE_LOG(LogAssetRegistryGenerator, Error, TEXT("Encrypted root asset '%s' exists in two chunks. Only secondary assets should be shared between chunks."), *AssetData->ObjectPath.ToString());
+							UE_LOG(LogAssetRegistryGenerator, Error, TEXT("Encrypted root asset '%s' exists in two chunks. Only secondary assets should be shared between chunks."), *AssetData->GetObjectPathString());
 						}
 						else if (AssetData->ChunkIDs.Num() == 1)
 						{
@@ -1120,7 +1120,7 @@ static void AddCookTagsToState(TMap<FSoftObjectPath, TArray<TPair<FName, FString
 	for (TPair<FSoftObjectPath, TArray<TPair<FName, FString>>>& ObjectToTags : InCookTags)
 	{
 		ObjectCount++;
-		const FAssetData* AssetData = InState.GetAssetByObjectPath(FName(WriteToString<256>(ObjectToTags.Key)));
+		const FAssetData* AssetData = InState.GetAssetByObjectPath(FSoftObjectPath(WriteToString<256>(ObjectToTags.Key)));
 
 		// Migrate to FAssetDataTagMap
 		FAssetDataTagMap NewTags;
@@ -1191,12 +1191,12 @@ void FAssetRegistryGenerator::UpdateCollectionAssetData()
 	// Collection tags are queried only by the existence of the key, the value is never used. But Tag Values are not allowed
 	// to be empty. Set the value for each tag to an arbitrary field, something short to avoid wasting memory. We use 1 (aka "true") for now.
 	FStringView CollectionValue(TEXTVIEW("1"));
-	for (const auto& AssetPathNameToCollectionTagsPair : AssetPathNamesToCollectionTags)
+	for (const TPair<FName, TArray<FName>>& AssetPathNameToCollectionTagsPair : AssetPathNamesToCollectionTags)
 	{
 		const FName AssetPathName = AssetPathNameToCollectionTagsPair.Key;
 		const TArray<FName>& CollectionTagsForAsset = AssetPathNameToCollectionTagsPair.Value;
 
-		const FAssetData* AssetData = State.GetAssetByObjectPath(AssetPathName);
+		const FAssetData* AssetData = State.GetAssetByObjectPath(FSoftObjectPath(WriteToString<FName::StringBufferSize>(AssetPathName)));
 		if (AssetData)
 		{
 			FAssetDataTagMap TagsAndValues = AssetData->TagsAndValues.CopyMap();
@@ -1375,9 +1375,8 @@ void FAssetRegistryGenerator::FinalizeChunkIDs(const TSet<FName>& InCookedPackag
 	}
 
 	// Copy ExplicitChunkIDs and other data from the AssetRegistry into the maps we use during finalization
-	for (const TPair<FName, const FAssetData*>& Pair : State.GetObjectPathToAssetDataMap())
+	State.EnumerateAllAssets([&](const FAssetData& AssetData)
 	{
-		const FAssetData& AssetData = *Pair.Value;
 		for (int32 ChunkID : AssetData.ChunkIDs)
 		{
 			if (ChunkID < 0)
@@ -1398,7 +1397,7 @@ void FAssetRegistryGenerator::FinalizeChunkIDs(const TSet<FName>& InCookedPackag
 		{
 			PackagesContainingMaps.Add(AssetData.PackageName);
 		}
-	}
+	});
 
 	// add all the packages to the unassigned package list
 	for (FName CookedPackage : CookedPackages)
@@ -1533,7 +1532,6 @@ bool FAssetRegistryGenerator::SaveAssetRegistry(const FString& SandboxPath, bool
 	UE_LOG(LogAssetRegistryGenerator, Display, TEXT("Saving asset registry v%d."), FAssetRegistryVersion::Type::LatestVersion);
 	FAutoScopedDurationTimer Timer;
 
-	const TMap<FName, const FAssetData*>& ObjectToDataMap = State.GetObjectPathToAssetDataMap();
 	
 	// Write development first, this will always write
 	FAssetRegistrySerializationOptions DevelopmentSaveOptions;
@@ -1682,7 +1680,9 @@ bool FAssetRegistryGenerator::SaveAssetRegistry(const FString& SandboxPath, bool
 			// Save the generated registry
 			FString PlatformSandboxPath = SandboxPath.Replace(TEXT("[Platform]"), *TargetPlatform->PlatformName());
 			FFileHelper::SaveArrayToFile(SerializedAssetRegistry, *PlatformSandboxPath);
-			UE_LOG(LogAssetRegistryGenerator, Display, TEXT("Generated asset registry num assets %d, size is %5.2fkb"), ObjectToDataMap.Num(), (float)SerializedAssetRegistry.Num() / 1024.f);
+
+			int32 NumAssets = State.GetNumAssets();
+			UE_LOG(LogAssetRegistryGenerator, Display, TEXT("Generated asset registry num assets %d, size is %5.2fkb"), NumAssets, (float)SerializedAssetRegistry.Num() / 1024.f);
 		}
 	}
 	
@@ -1752,19 +1752,17 @@ bool FAssetRegistryGenerator::WriteCookerOpenOrder(FSandboxPlatformFile& InSandb
 {
 	TSet<FName> PackageNameSet;
 	TSet<FName> MapList;
-	const TMap<FName, const FAssetData*>& ObjectToDataMap = State.GetObjectPathToAssetDataMap();
-	for (const TPair<FName, const FAssetData*>& Pair : ObjectToDataMap)
+	State.EnumerateAllAssets([this, &PackageNameSet, &MapList](const FAssetData& AssetData)
 	{
-		FAssetData* AssetData = const_cast<FAssetData*>(Pair.Value);
-		PackageNameSet.Add(AssetData->PackageName);
+		PackageNameSet.Add(AssetData.PackageName);
 
 		// REPLACE WITH PRIORITY
 
-		if (ContainsMap(AssetData->PackageName))
+		if (ContainsMap(AssetData.PackageName))
 		{
-			MapList.Add(AssetData->PackageName);
+			MapList.Add(AssetData.PackageName);
 		}
-	}
+	});
 
 	FString CookerFileOrderString;
 	{
@@ -1967,17 +1965,16 @@ bool FAssetRegistryGenerator::GenerateAssetChunkInformationCSV(const FString& Ou
 	FString TmpString, TmpStringChunks;
 	ANSICHAR HeaderText[] = "ChunkID, Package Name, Class Type, Hard or Soft Chunk, File Size, Other Chunks\n";
 
-	const TMap<FName, const FAssetData*>& ObjectToDataMap = State.GetObjectPathToAssetDataMap();
 	TArray<const FAssetData*> AssetDataList;
-	for (const TPair<FName, const FAssetData*>& Pair : ObjectToDataMap)
+	State.EnumerateAllAssets([&AssetDataList](const FAssetData& AssetData)
 	{
-		AssetDataList.Add(Pair.Value);
-	}
+		AssetDataList.Add(&AssetData);
+	});
 
 	// Sort list so it's consistent over time
 	AssetDataList.Sort([](const FAssetData& A, const FAssetData& B)
 	{
-		return A.ObjectPath.LexicalLess(B.ObjectPath);
+		return A.GetSoftObjectPath().LexicalLess(B.GetSoftObjectPath());
 	});
 
 	// Create file for all chunks
@@ -2506,7 +2503,7 @@ FAssetRegistryGenerator::FCreateOrFindArray FAssetRegistryGenerator::CreateOrFin
 
 const FAssetData* FAssetRegistryGenerator::CreateOrFindAssetData(UObject& Object)
 {
-	const FAssetData* const AssetData = State.GetAssetByObjectPath(FName(Object.GetPathName()));
+	const FAssetData* const AssetData = State.GetAssetByObjectPath(FSoftObjectPath(&Object));
 	if (!AssetData)
 	{
 		FAssetData* const NewAssetData = new FAssetData(&Object, true /* bAllowBlueprintClass */);
diff --git a/Engine/Source/Editor/UnrealEd/Private/CookOnTheFlyServer.cpp b/Engine/Source/Editor/UnrealEd/Private/CookOnTheFlyServer.cpp
--- a/Engine/Source/Editor/UnrealEd/Private/CookOnTheFlyServer.cpp
+++ b/Engine/Source/Editor/UnrealEd/Private/CookOnTheFlyServer.cpp
@@ -989,7 +989,7 @@ bool UCookOnTheFlyServer::ContainsMap(const FName& PackageName) const
 	return false;
 }
 
-bool UCookOnTheFlyServer::ContainsRedirector(const FName& PackageName, TMap<FName, FName>& RedirectedPaths) const
+bool UCookOnTheFlyServer::ContainsRedirector(const FName& PackageName, TMap<FSoftObjectPath, FSoftObjectPath>& RedirectedPaths) const
 {
 	bool bFoundRedirector = false;
 	TArray<FAssetData> Assets;
@@ -999,14 +999,14 @@ bool UCookOnTheFlyServer::ContainsRedirector(const FName& PackageName, TMap<FNam
 	{
 		if (Asset.IsRedirector())
 		{
-			FName RedirectedPath;
+			FSoftObjectPath RedirectedPath;
 			FString RedirectedPathString;
 			if (Asset.GetTagValue("DestinationObject", RedirectedPathString))
 			{
 				ConstructorHelpers::StripObjectClass(RedirectedPathString);
-				RedirectedPath = FName(*RedirectedPathString);
+				RedirectedPath = FSoftObjectPath(RedirectedPathString);
 				FAssetData DestinationData = AssetRegistry->GetAssetByObjectPath(RedirectedPath, true);
-				TSet<FName> SeenPaths;
+				TSet<FSoftObjectPath> SeenPaths;
 
 				SeenPaths.Add(RedirectedPath);
 
@@ -1016,7 +1016,7 @@ bool UCookOnTheFlyServer::ContainsRedirector(const FName& PackageName, TMap<FNam
 					if (DestinationData.GetTagValue("DestinationObject", RedirectedPathString))
 					{
 						ConstructorHelpers::StripObjectClass(RedirectedPathString);
-						RedirectedPath = FName(*RedirectedPathString);
+						RedirectedPath = FSoftObjectPath(RedirectedPathString);
 
 						if (SeenPaths.Contains(RedirectedPath))
 						{
@@ -1041,7 +1041,7 @@ bool UCookOnTheFlyServer::ContainsRedirector(const FName& PackageName, TMap<FNam
 
 				if (!bDestinationValid)
 				{
-					if (RedirectedPath != NAME_None)
+					if (RedirectedPath.IsValid())
 					{
 						FName StandardPackageName = PackageDatas->GetFileNameByPackageName(FName(*FPackageName::ObjectPathToPackageName(RedirectedPathString)));
 						if (!StandardPackageName.IsNone())
@@ -1053,11 +1053,11 @@ bool UCookOnTheFlyServer::ContainsRedirector(const FName& PackageName, TMap<FNam
 
 				if (bDestinationValid)
 				{
-					RedirectedPaths.Add(Asset.ObjectPath, RedirectedPath);
+					RedirectedPaths.Add(Asset.GetSoftObjectPath(), RedirectedPath);
 				}
 				else
 				{
-					RedirectedPaths.Add(Asset.ObjectPath, NAME_None);
+					RedirectedPaths.Add(Asset.GetSoftObjectPath(), FSoftObjectPath{});
 					UE_LOG(LogCook, Log, TEXT("Found redirector in package %s pointing to deleted object %s"), *PackageName.ToString(), *RedirectedPathString);
 				}
 
@@ -5607,12 +5607,12 @@ void FSaveCookedPackageContext::FinishPackage()
 		GRedirectCollector.ProcessSoftObjectPathPackageList(PackageFName, false, SoftObjectPackages);
 		for (FName SoftObjectPackage : SoftObjectPackages)
 		{
-			TMap<FName, FName> RedirectedPaths;
+			TMap<FSoftObjectPath, FSoftObjectPath> RedirectedPaths;
 
 			// If this is a redirector, extract destination from asset registry
 			if (COTFS.ContainsRedirector(SoftObjectPackage, RedirectedPaths))
 			{
-				for (TPair<FName, FName>& RedirectedPath : RedirectedPaths)
+				for (TPair<FSoftObjectPath, FSoftObjectPath>& RedirectedPath : RedirectedPaths)
 				{
 					GRedirectCollector.AddAssetPathRedirection(RedirectedPath.Key, RedirectedPath.Value);
 				}
@@ -9373,12 +9373,12 @@ void UCookOnTheFlyServer::GenerateInitialRequests(FBeginCookContext& BeginContex
 	}
 	for (FName SoftObjectPackage : StartupSoftObjectPackages)
 	{
-		TMap<FName, FName> RedirectedPaths;
+		TMap<FSoftObjectPath, FSoftObjectPath> RedirectedPaths;
 
 		// If this is a redirector, extract destination from asset registry
 		if (ContainsRedirector(SoftObjectPackage, RedirectedPaths))
 		{
-			for (TPair<FName, FName>& RedirectedPath : RedirectedPaths)
+			for (TPair<FSoftObjectPath, FSoftObjectPath>& RedirectedPath : RedirectedPaths)
 			{
 				GRedirectCollector.AddAssetPathRedirection(RedirectedPath.Key, RedirectedPath.Value);
 			}
@@ -9955,50 +9955,39 @@ bool UCookOnTheFlyServer::GetAllPackageFilenamesFromAssetRegistry(const FString&
 			GPreloadedARState.Serialize(*Reader.Get(), FAssetRegistrySerializationOptions());
 		}
 
-		// possibly use the preloaded state
-		const TMap<FName, const FAssetData*>& RegistryDataMapLoaded = GPreloadedARState.GetObjectPathToAssetDataMap();
-		TMap<FName, const FAssetData*> RegistryDataMapTrimmed;
 		
-		// if we want to reevaluate uncooked pacakges, then remove the uncooked packages from the set of known packages
-		// removing them from the map up here simplifies the logic below 
-		if (bReevaluateUncookedPackages)
-		{
-			for (TPair<FName, const FAssetData*> Pair : RegistryDataMapLoaded)
-			{
-				if (Pair.Value->PackageFlags != 0)
-				{
-					RegistryDataMapTrimmed.Add(Pair.Key, Pair.Value);
-				}
-			}
-		}
-		const TMap<FName, const FAssetData*>& RegistryDataMap = bReevaluateUncookedPackages ? RegistryDataMapTrimmed : RegistryDataMapLoaded;
-
 		check(OutPackageDatas.Num() == 0);
 
-		int32 NumPackages = RegistryDataMap.Num();
+		int32 NumPackages = GPreloadedARState.GetNumAssets();
 		TArray<const FAssetData*> AssetDatas;
 		TSet<FName> PackageNames;
 		AssetDatas.Reserve(NumPackages);
 		OutPackageDatas.Reserve(NumPackages);
 
 		// Convert the Map of RegistryData into an Array of FAssetData and populate PackageNames in the output array
-		for (const TPair<FName, const FAssetData*>& RegistryData : RegistryDataMap)
+		GPreloadedARState.EnumerateAllAssets([&](const FAssetData& RegistryData)
 		{
-			FName PackageName = RegistryData.Value->PackageName;
+			// If we want to reevaluate (try cooking again) the uncooked packages (packages that were found to be empty when we cooked them before),
+			// then remove the uncooked packages from the set of known packages. Uncooked packages are identified by PackageFlags == 0.
+			if (bReevaluateUncookedPackages && (RegistryData.PackageFlags == 0))
+			{
+				return;
+			}
+			FName PackageName = RegistryData.PackageName;
 			bool bPackageAlreadyAdded;
 			PackageNames.Add(PackageName, &bPackageAlreadyAdded);
 			if (bPackageAlreadyAdded)
 			{
-				continue;
+				return;
 			}
 			if (FPackageName::GetPackageMountPoint(PackageName.ToString()).IsNone())
 			{
 				// Skip any packages that are not currently mounted; if we tried to find their FileNames below
 				// we would get log spam
-				continue;
+				return;
 			}
 
-			AssetDatas.Add(RegistryData.Value);
+			AssetDatas.Add(&RegistryData);
 			FConstructPackageData& PackageData = OutPackageDatas.Emplace_GetRef();
 			PackageData.PackageName = PackageName;
 
@@ -10009,7 +9998,7 @@ bool UCookOnTheFlyServer::GetAllPackageFilenamesFromAssetRegistry(const FString&
 			{
 				PackageData.NormalizedFileName = ExistingPackageData->GetFileName();
 			}
-		}
+		});
 		NumPackages = AssetDatas.Num();
 
 		ParallelFor(NumPackages,
@@ -10305,12 +10294,12 @@ uint32 UCookOnTheFlyServer::CookFullLoadAndSave()
 
 					for (FName StringAssetPackage : StringAssetPackages)
 					{
-						TMap<FName, FName> RedirectedPaths;
+						TMap<FSoftObjectPath, FSoftObjectPath> RedirectedPaths;
 
 						// If this is a redirector, extract destination from asset registry
 						if (ContainsRedirector(StringAssetPackage, RedirectedPaths))
 						{
-							for (TPair<FName, FName>& RedirectedPath : RedirectedPaths)
+							for (TPair<FSoftObjectPath, FSoftObjectPath>& RedirectedPath : RedirectedPaths)
 							{
 								GRedirectCollector.AddAssetPathRedirection(RedirectedPath.Key, RedirectedPath.Value);
 								PackagesToLoad.Add(FPackageName::ObjectPathToPackageName(RedirectedPath.Value.ToString()));
diff --git a/Engine/Source/Runtime/AssetRegistry/Private/AssetRegistry.cpp b/Engine/Source/Runtime/AssetRegistry/Private/AssetRegistry.cpp
--- a/Engine/Source/Runtime/AssetRegistry/Private/AssetRegistry.cpp
+++ b/Engine/Source/Runtime/AssetRegistry/Private/AssetRegistry.cpp
@@ -1884,12 +1884,11 @@ void EnumerateMemoryAssetsHelper(const FARCompiledFilter& InFilter, TSet<FName>&
 				return;
 			}
 
-			// Object Path
-			const FString ObjectPathStr = Obj->GetPathName();
-			if (InFilter.ObjectPaths.Num() > 0)
+			// Asset Path
+			const FSoftObjectPath ObjectPath = FSoftObjectPath(Obj);
+			if (InFilter.SoftObjectPaths.Num() > 0)
 			{
-				const FName ObjectPath = FName(*ObjectPathStr, FNAME_Find);
-				if (!InFilter.ObjectPaths.Contains(ObjectPath))
+				if (!InFilter.SoftObjectPaths.Contains(ObjectPath))
 				{
 					return;
 				}
@@ -1908,7 +1907,7 @@ void EnumerateMemoryAssetsHelper(const FARCompiledFilter& InFilter, TSet<FName>&
 
 			// Could perhaps save some FName -> String conversions by creating this a bit earlier using the UObject constructor
 			// to get package name and path.
-			FAssetData PartialAssetData(PackageNameStr, ObjectPathStr, Obj->GetClass()->GetClassPathName(), FAssetDataTagMap(),
+			FAssetData PartialAssetData(PackageNameStr, ObjectPath.ToString(), Obj->GetClass()->GetClassPathName(), FAssetDataTagMap(),
 				InMemoryPackage->GetChunkIDs(), InMemoryPackage->GetPackageFlags());
 
 			// All filters passed, except for AssetRegistry filter; caller must check that one
@@ -2005,12 +2004,13 @@ void FAssetRegistryImpl::EnumerateDiskAssets(const FARCompiledFilter& InFilter,
 
 }
 
-FAssetData UAssetRegistryImpl::GetAssetByObjectPath(const FName ObjectPath, bool bIncludeOnlyOnDiskAssets) const
+FAssetData UAssetRegistryImpl::GetAssetByObjectPath(const FSoftObjectPath& ObjectPath, bool bIncludeOnlyOnDiskAssets) const
 {
 	if (!bIncludeOnlyOnDiskAssets)
 	{
-		const FNameBuilder ObjectPathStr(ObjectPath);
-		UObject* Asset = FindObject<UObject>(nullptr, *ObjectPathStr);
+		TStringBuilder<FName::StringBufferSize> Builder;
+		ObjectPath.ToString(Builder);
+		UObject* Asset = FindObject<UObject>(nullptr, *Builder);
 
 		if (Asset)
 		{
@@ -2033,6 +2033,13 @@ FAssetData UAssetRegistryImpl::GetAssetByObjectPath(const FName ObjectPath, bool
 	}
 }
 
+FAssetData UAssetRegistryImpl::GetAssetByObjectPath(const FName ObjectPath, bool bIncludeOnlyOnDiskAssets) const
+{
+PRAGMA_DISABLE_DEPRECATION_WARNINGS;
+	return GetAssetByObjectPath(FSoftObjectPath(ObjectPath.ToString()));
+PRAGMA_ENABLE_DEPRECATION_WARNINGS;
+}
+
 bool UAssetRegistryImpl::GetAllAssets(TArray<FAssetData>& OutAssetData, bool bIncludeOnlyOnDiskAssets) const
 {
 	const double GetAllAssetsStartTime = FPlatformTime::Seconds();
@@ -2490,48 +2497,57 @@ bool UAssetRegistryImpl::DoesPackageExistOnDisk(FName PackageName, FString* OutC
 	}
 }
 
+FSoftObjectPath UAssetRegistryImpl::GetRedirectedObjectPath(const FSoftObjectPath& ObjectPath) const
+{
+	FReadScopeLock InterfaceScopeLock(InterfaceLock);
+	return GuardedData.GetRedirectedObjectPath(ObjectPath);
+}
 
 FName UAssetRegistryImpl::GetRedirectedObjectPath(const FName ObjectPath) const
 {
 	FReadScopeLock InterfaceScopeLock(InterfaceLock);
+PRAGMA_DISABLE_DEPRECATION_WARNINGS
 	return GuardedData.GetRedirectedObjectPath(ObjectPath);
+PRAGMA_ENABLE_DEPRECATION_WARNINGS
 }
 
 namespace UE::AssetRegistry
 {
 
-FName FAssetRegistryImpl::GetRedirectedObjectPath(const FName ObjectPath) const
+FSoftObjectPath FAssetRegistryImpl::GetRedirectedObjectPath(const FSoftObjectPath& ObjectPath) const
 {
-	FString RedirectedPath = ObjectPath.ToString();
+	FSoftObjectPath RedirectedPath = ObjectPath;
 
 	// For legacy behavior, for the first object pointed to, we look up the object in memory
 	// before checking the on-disk assets
-	UObject* Asset = FindObject<UObject>(nullptr, *ObjectPath.ToString());
+	UObject* Asset = ObjectPath.ResolveObject();
 	const FAssetData* AssetData = nullptr;
 	if (!Asset)
 	{
 		AssetData = State.GetAssetByObjectPath(ObjectPath);
 	}
 
-	TSet<FString> SeenPaths;
+	TSet<FSoftObjectPath> SeenPaths;
 	SeenPaths.Add(RedirectedPath);
 
-	auto TryGetRedirectedPath = [](UObject* InAsset, const FAssetData* InAssetData, FString& OutRedirectedPath)
+	auto TryGetRedirectedPath = [](UObject* InAsset, const FAssetData* InAssetData, FSoftObjectPath& OutRedirectedPath)
 	{
 		if (InAsset)
 		{
 			UObjectRedirector* Redirector = Cast<UObjectRedirector>(InAsset);
 			if (Redirector && Redirector->DestinationObject)
 			{
-				OutRedirectedPath = Redirector->DestinationObject->GetPathName();
+				OutRedirectedPath = FSoftObjectPath(Redirector->DestinationObject);
 				return true;
 			}
 		}
 		else if (InAssetData && InAssetData->IsRedirector())
 		{
-			if (InAssetData->GetTagValue("DestinationObject", OutRedirectedPath))
+			FString Dest;
+			if (InAssetData->GetTagValue("DestinationObject", Dest))
 			{
-				ConstructorHelpers::StripObjectClass(OutRedirectedPath);
+				ConstructorHelpers::StripObjectClass(Dest);
+				OutRedirectedPath = Dest;
 				return true;
 			}
 		}
@@ -2551,11 +2567,17 @@ FName FAssetRegistryImpl::GetRedirectedObjectPath(const FName ObjectPath) const
 			SeenPaths.Add(RedirectedPath);
 			// For legacy behavior, for all redirects after the initial request, we only check on-disk assets
 			Asset = nullptr;
-			AssetData = State.GetAssetByObjectPath(FName(*RedirectedPath));
+			AssetData = State.GetAssetByObjectPath(RedirectedPath);
 		}
 	}
 
-	return FName(*RedirectedPath);
+	return RedirectedPath;
+}
+
+UE_DEPRECATED(5.1, "Asset path FNames have been deprecated, use FSoftObjectPath instead.")
+FName FAssetRegistryImpl::GetRedirectedObjectPath(const FName ObjectPath) const
+{
+	return FName(*GetRedirectedObjectPath(FSoftObjectPath(ObjectPath.ToString())).ToString());
 }
 
 }
@@ -2833,9 +2855,9 @@ bool RunAssetThroughFilter_Unchecked(const FAssetData& AssetData, const FARCompi
 	}
 
 	// ObjectPaths
-	if (Filter.ObjectPaths.Num() > 0)
+	if (Filter.SoftObjectPaths.Num() > 0)
 	{
-		const bool bPassesObjectPaths = Filter.ObjectPaths.Contains(AssetData.ObjectPath);
+		const bool bPassesObjectPaths = Filter.SoftObjectPaths.Contains(AssetData.GetSoftObjectPath());
 		if (bPassesObjectPaths != bPassFilterValue)
 		{
 			return !bPassFilterValue;
@@ -2911,7 +2933,7 @@ void UAssetRegistryImpl::ExpandRecursiveFilter(const FARFilter& InFilter, FARFil
 	ExpandedFilter.Clear();
 	ExpandedFilter.PackageNames = CompiledFilter.PackageNames.Array();
 	ExpandedFilter.PackagePaths = CompiledFilter.PackagePaths.Array();
-	ExpandedFilter.ObjectPaths = CompiledFilter.ObjectPaths.Array();
+	ExpandedFilter.SoftObjectPaths = CompiledFilter.SoftObjectPaths.Array();
 	ExpandedFilter.ClassPaths = CompiledFilter.ClassPaths.Array();
 	ExpandedFilter.TagsAndValues = CompiledFilter.TagsAndValues;
 	ExpandedFilter.bIncludeOnlyOnDiskAssets = CompiledFilter.bIncludeOnlyOnDiskAssets;
@@ -2950,7 +2972,10 @@ void FAssetRegistryImpl::CompileFilter(Impl::FClassInheritanceContext& Inheritan
 	{
 		OutCompiledFilter.PackagePaths.Add(FPathTree::NormalizePackagePath(PackagePath));
 	}
-	OutCompiledFilter.ObjectPaths.Append(InFilter.ObjectPaths);
+	OutCompiledFilter.SoftObjectPaths.Append(InFilter.SoftObjectPaths);
+PRAGMA_DISABLE_DEPRECATION_WARNINGS
+	OutCompiledFilter.SoftObjectPaths.Append(UE::SoftObjectPath::Private::ConvertObjectPathNames(InFilter.ObjectPaths));
+PRAGMA_ENABLE_DEPRECATION_WARNINGS
 PRAGMA_DISABLE_DEPRECATION_WARNINGS
 	if (!ensureAlwaysMsgf(InFilter.ClassNames.Num() == 0, TEXT("Asset Registry Filter using ClassNames instead of ClassPaths. First class name: \"%s\""), *InFilter.ClassNames[0].ToString()))
 	{
@@ -3375,7 +3400,7 @@ void UAssetRegistryImpl::AssetDeleted(UObject* DeletedAsset)
 
 		{
 			// Need to remove from GRedirectCollector
-			GRedirectCollector.RemoveAssetPathRedirection(FName(*DeletedAsset->GetPathName()));
+			GRedirectCollector.RemoveAssetPathRedirection(FSoftObjectPath(DeletedAsset));
 		}
 #endif
 
@@ -3757,7 +3782,7 @@ void FAssetRegistryImpl::TickGatherPackage(Impl::FEventContext& EventContext, co
 }
 
 #if WITH_EDITOR
-void FAssetRegistryImpl::LoadCalculatedDependencies(TArray<FName>* AssetsToCalculate, double TickStartTime, bool& bOutInterrupted)
+void FAssetRegistryImpl::LoadCalculatedDependencies(TArray<FName>* AssetPackageNamesToCalculate, double TickStartTime, bool& bOutInterrupted)
 {
 	auto CheckForTimeUp = [&TickStartTime, &bOutInterrupted](bool bHadActivity)
 	{
@@ -3774,9 +3799,9 @@ void FAssetRegistryImpl::LoadCalculatedDependencies(TArray<FName>* AssetsToCalcu
 		return false;
 	};
 
-	if (AssetsToCalculate)
+	if (AssetPackageNamesToCalculate)
 	{
-		for (FName PackageName : *AssetsToCalculate)
+		for (FName PackageName : *AssetPackageNamesToCalculate)
 		{
 			// We do not remove the package from PackagesNeedingDependencyCalculation, because
 			// we are only calculating an interim result when AssetsToCalculate is non-null
@@ -3963,10 +3988,8 @@ void FAssetRegistryImpl::CachePathsFromState(Impl::FEventContext& EventContext,
 	CollectCodeGeneratorClasses();
 
 	// Add paths to cache
-	for (const TPair<FName, FAssetData*>& AssetDataPair : InState.CachedAssetsByObjectPath)
+	for (const FAssetData* AssetData : InState.CachedAssets)
 	{
-		const FAssetData* AssetData = AssetDataPair.Value;
-
 		if (AssetData != nullptr)
 		{
 			AddAssetPath(EventContext, AssetData->PackagePath);
@@ -4111,13 +4134,13 @@ bool ReadAssetFile(FPackageReader& PackageReader, IAssetRegistry::FLoadPackageRe
 
 }
 
-void UAssetRegistryImpl::InitializeTemporaryAssetRegistryState(FAssetRegistryState& OutState, const FAssetRegistrySerializationOptions& Options, bool bRefreshExisting, const TMap<FName, FAssetData*>& OverrideData) const
+void UAssetRegistryImpl::InitializeTemporaryAssetRegistryState(FAssetRegistryState& OutState, const FAssetRegistrySerializationOptions& Options, bool bRefreshExisting) const
 {
+	using FAssetDataMap = UE::AssetRegistry::Private::FAssetDataMap;
+
 	FReadScopeLock InterfaceScopeLock(InterfaceLock);
 	const FAssetRegistryState& State = GuardedData.GetState();
-	const TMap<FName, FAssetData*>& DataToUse = OverrideData.Num() > 0 ? OverrideData : State.CachedAssetsByObjectPath;
-
-	OutState.InitializeFromExisting(DataToUse, State.CachedDependsNodes, State.CachedPackageData, Options, bRefreshExisting ? FAssetRegistryState::EInitializationMode::OnlyUpdateExisting : FAssetRegistryState::EInitializationMode::Rebuild);
+	OutState.InitializeFromExisting(State.CachedAssets, State.CachedDependsNodes, State.CachedPackageData, Options, bRefreshExisting ? FAssetRegistryState::EInitializationMode::OnlyUpdateExisting : FAssetRegistryState::EInitializationMode::Rebuild);
 }
 
 #if ASSET_REGISTRY_STATE_DUMPING_ENABLED
@@ -4163,7 +4186,7 @@ namespace Impl
 {
 
 FScanPathContext::FScanPathContext(FEventContext& InEventContext, const TArray<FString>& InDirs, const TArray<FString>& InFiles,
-	bool bInForceRescan, bool bInIgnoreDenyListScanFilters, TArray<FName>* FoundAssets)
+	bool bInForceRescan, bool bInIgnoreDenyListScanFilters, TArray<FSoftObjectPath>* FoundAssets)
 	: EventContext(InEventContext)
 	, OutFoundAssets(FoundAssets)
 	, bForceRescan(bInForceRescan)
@@ -4266,19 +4289,14 @@ void FAssetRegistryImpl::ScanPathsSynchronous(Impl::FScanPathContext& Context)
 	}
 
 	Gatherer.ScanPathsSynchronous(Context.LocalPaths, Context.bForceRescan, Context.bIgnoreDenyListScanFilters, CacheFilename, Context.PackageDirs);
-	TArray<FName> BufferFoundAssets;
-	TArray<FName>* FoundAssets = &BufferFoundAssets;
-	if (Context.OutFoundAssets)
-	{
-		FoundAssets = Context.OutFoundAssets;
-	}
+	TArray<FName> FoundAssetPackageNames;
 
-	auto AssetsFoundCallback = [&Context, FoundAssets, this](const TRingBuffer<FAssetData*>& InFoundAssets)
+	auto AssetsFoundCallback = [&Context, &FoundAssetPackageNames, this](const TRingBuffer<FAssetData*>& InFoundAssets)
 	{
 		Context.NumFoundAssets = InFoundAssets.Num();
 
-		FoundAssets->Reset();
-		FoundAssets->Reserve(Context.NumFoundAssets);
+		FoundAssetPackageNames.Reset();
+		FoundAssetPackageNames.Reserve(Context.NumFoundAssets);
 
 		// The gatherer may have added other assets that were scanned as part of the ongoing background scan; remove any assets that were not in the requested paths
 		for (FAssetData* AssetData : InFoundAssets)
@@ -4313,8 +4331,12 @@ void FAssetRegistryImpl::ScanPathsSynchronous(Impl::FScanPathContext& Context)
 			if (bIsInRequestedPaths)
 			{
 				UE_LOG(LogAssetRegistry, VeryVerbose, TEXT("FAssetRegistryImpl::ScanPathsSynchronous: Found Asset: %s"),
-					*AssetData->ObjectPath.ToString());
-				FoundAssets->Add(AssetData->ObjectPath);
+					*AssetData->GetObjectPathString());
+				if (Context.OutFoundAssets)
+				{
+					Context.OutFoundAssets->Add(AssetData->GetSoftObjectPath());
+				}
+				FoundAssetPackageNames.Add(AssetData->PackageName);
 			}
 		}
 	};
@@ -4322,7 +4344,7 @@ void FAssetRegistryImpl::ScanPathsSynchronous(Impl::FScanPathContext& Context)
 	bool bUnusedInterrupted;
 	Context.Status = TickGatherer(Context.EventContext, -1., bUnusedInterrupted, FAssetsFoundCallback(AssetsFoundCallback));
 #if WITH_EDITOR
-	LoadCalculatedDependencies(FoundAssets, -1., bUnusedInterrupted);
+	LoadCalculatedDependencies(&FoundAssetPackageNames, -1., bUnusedInterrupted);
 #endif
 }
 
@@ -4424,7 +4446,9 @@ void FAssetRegistryImpl::AssetSearchDataGathered(Impl::FEventContext& EventConte
 		CA_ASSUME(BackgroundResult.Get() != nullptr);
 
 		// Try to update any asset data that may already exist
-		FAssetData* ExistingAssetData = State.CachedAssetsByObjectPath.FindRef(BackgroundResult->ObjectPath);
+		FCachedAssetKey Key(*BackgroundResult);
+		FAssetData* const* FoundData = State.CachedAssets.Find(Key);
+		FAssetData* ExistingAssetData = FoundData ? *FoundData : nullptr;
 
 		const FName PackagePath = BackgroundResult->PackagePath;
 
@@ -4446,7 +4470,7 @@ void FAssetRegistryImpl::AssetSearchDataGathered(Impl::FEventContext& EventConte
 			{
 				// If the current AssetData came from a loaded asset, don't overwrite it with the new one from disk; loaded asset is more authoritative because it has run the postload steps
 #if WITH_EDITOR
-				if (!AssetDataObjectPathsUpdatedOnLoad.Contains(BackgroundResult->ObjectPath))
+				if (!AssetDataObjectPathsUpdatedOnLoad.Contains(BackgroundResult->GetSoftObjectPath()))
 #endif
 				{
 #if WITH_EDITOR
@@ -4476,7 +4500,7 @@ void FAssetRegistryImpl::AssetSearchDataGathered(Impl::FEventContext& EventConte
 		}
 		else
 		{
-			UE_LOG(LogAssetRegistry, Warning, TEXT("AssetRegistry: An asset has been loaded with an invalid mount point: '%s', Mount Point: '%s'"), *BackgroundResult->ObjectPath.ToString(), *PackagePathString)
+			UE_LOG(LogAssetRegistry, Warning, TEXT("AssetRegistry: An asset has been loaded with an invalid mount point: '%s', Mount Point: '%s'"), *BackgroundResult->GetObjectPathString(), *PackagePathString)
 		}
 
 		// Check to see if we have run out of time in this tick
@@ -5183,9 +5207,10 @@ void FAssetRegistryImpl::PushProcessLoadedAssetsBatch(Impl::FEventContext& Event
 	// Add or update existing for all of the AssetDatas created by the batch
 	for (FAssetData& NewAssetData : LoadedAssetDatas)
 	{
-		FAssetData** DataFromGather = State.CachedAssetsByObjectPath.Find(NewAssetData.ObjectPath);
+		FCachedAssetKey Key(NewAssetData);
+		FAssetData** DataFromGather = State.CachedAssets.Find(Key);
 
-		AssetDataObjectPathsUpdatedOnLoad.Add(NewAssetData.ObjectPath);
+		AssetDataObjectPathsUpdatedOnLoad.Add(NewAssetData.GetSoftObjectPath());
 
 		if (!DataFromGather)
 		{
@@ -5213,11 +5238,11 @@ void FAssetRegistryImpl::UpdateRedirectCollector()
 
 	for (const FAssetData* AssetData : RedirectorAssets)
 	{
-		FName Destination = GetRedirectedObjectPath(AssetData->ObjectPath);
+		FSoftObjectPath Destination = GetRedirectedObjectPath(AssetData->GetSoftObjectPath());
 
-		if (Destination != AssetData->ObjectPath)
+		if (Destination != AssetData->GetSoftObjectPath())
 		{
-			GRedirectCollector.AddAssetPathRedirection(AssetData->ObjectPath, Destination);
+			GRedirectCollector.AddAssetPathRedirection(AssetData->GetSoftObjectPath(), Destination);
 		}
 	}
 }
@@ -5277,7 +5302,7 @@ void FAssetRegistryImpl::ScanModifiedAssetFiles(Impl::FEventContext& EventContex
 		}
 
 		// Re-scan and update the asset registry with the new asset data
-		TArray<FName> FoundAssets;
+		TArray<FSoftObjectPath> FoundAssets;
 		Impl::FScanPathContext Context(EventContext, TArray<FString>(), InFilePaths,
 			true /* bForceRescan */, false /* bIgnoreDenyListScanFilters */, &FoundAssets);
 		ScanPathsSynchronous(Context);
@@ -5287,7 +5312,7 @@ void FAssetRegistryImpl::ScanModifiedAssetFiles(Impl::FEventContext& EventContex
 		{
 			for (FAssetData* OldPackageAsset : OldPackageAssets)
 			{
-				if (!FoundAssets.Contains(OldPackageAsset->ObjectPath))
+				if (!FoundAssets.Contains(OldPackageAsset->GetSoftObjectPath()))
 				{
 					RemoveAssetData(EventContext, OldPackageAsset);
 				}
@@ -5295,9 +5320,9 @@ void FAssetRegistryImpl::ScanModifiedAssetFiles(Impl::FEventContext& EventContex
 		}
 
 		// Send ModifiedOnDisk event for every Asset that was modified
-		for (FName FoundAsset : FoundAssets)
+		for (const FSoftObjectPath& FoundAsset : FoundAssets)
 		{
-			FAssetData** AssetData = State.CachedAssetsByObjectPath.Find(FoundAsset);
+			FAssetData** AssetData = State.CachedAssets.Find(FCachedAssetKey(FoundAsset));
 			if (AssetData)
 			{
 				EventContext.AssetEvents.Emplace(**AssetData, Impl::FEventContext::EEvent::UpdatedOnDisk);
@@ -5657,7 +5682,7 @@ void UAssetRegistryImpl::OnGetExtraObjectTags(const UObject* Object, TArray<UObj
 	{
 		TSet<FName>& MetaDataTags = UObject::GetMetaDataTagsForAssetRegistry();
 		// It is critical that bIncludeOnlyOnDiskAssets=true otherwise this will cause an infinite loop
-		const FAssetData AssetData = GetAssetByObjectPath(*Object->GetPathName(), /*bIncludeOnlyOnDiskAssets=*/true);
+		const FAssetData AssetData = GetAssetByObjectPath(FSoftObjectPath(Object), /*bIncludeOnlyOnDiskAssets=*/true);
 		for (const FName MetaDataTag : MetaDataTags)
 		{
 			auto OutTagsContainsTagPredicate = [MetaDataTag](const UObject::FAssetRegistryTag& Tag) { return Tag.Name == MetaDataTag; };
@@ -6014,7 +6039,7 @@ void FAssetRegistryImpl::SetManageReferences(const TMultiMap<FAssetIdentifier, F
 
 }
 
-bool UAssetRegistryImpl::SetPrimaryAssetIdForObjectPath(const FName ObjectPath, FPrimaryAssetId PrimaryAssetId)
+bool UAssetRegistryImpl::SetPrimaryAssetIdForObjectPath(const FSoftObjectPath& ObjectPath, FPrimaryAssetId PrimaryAssetId)
 {
 	UE::AssetRegistry::Impl::FEventContext EventContext;
 	bool bResult;
@@ -6029,9 +6054,9 @@ bool UAssetRegistryImpl::SetPrimaryAssetIdForObjectPath(const FName ObjectPath,
 namespace UE::AssetRegistry
 {
 
-bool FAssetRegistryImpl::SetPrimaryAssetIdForObjectPath(Impl::FEventContext& EventContext, const FName ObjectPath, FPrimaryAssetId PrimaryAssetId)
+bool FAssetRegistryImpl::SetPrimaryAssetIdForObjectPath(Impl::FEventContext& EventContext, const FSoftObjectPath& ObjectPath, FPrimaryAssetId PrimaryAssetId)
 {
-	FAssetData** FoundAssetData = State.CachedAssetsByObjectPath.Find(ObjectPath);
+	FAssetData** FoundAssetData = State.CachedAssets.Find(FCachedAssetKey(ObjectPath));
 
 	if (!FoundAssetData)
 	{
diff --git a/Engine/Source/Runtime/AssetRegistry/Private/AssetRegistry.h b/Engine/Source/Runtime/AssetRegistry/Private/AssetRegistry.h
--- a/Engine/Source/Runtime/AssetRegistry/Private/AssetRegistry.h
+++ b/Engine/Source/Runtime/AssetRegistry/Private/AssetRegistry.h
@@ -42,7 +42,9 @@ public:
 	virtual bool GetAssets(const FARFilter& Filter, TArray<FAssetData>& OutAssetData, bool bSkipARFilteredAssets = true) const override;
 	virtual bool EnumerateAssets(const FARFilter& Filter, TFunctionRef<bool(const FAssetData&)> Callback, bool bSkipARFilteredAssets=true) const override;
 	virtual bool EnumerateAssets(const FARCompiledFilter& Filter, TFunctionRef<bool(const FAssetData&)> Callback, bool bSkipARFilteredAssets = true) const override;
+	// UE_DEPRECATED(5.1, "Asset path FNames have been deprecated, use FSoftObjectPath instead.")
 	virtual FAssetData GetAssetByObjectPath( const FName ObjectPath, bool bIncludeOnlyOnDiskAssets = false ) const override;
+	virtual FAssetData GetAssetByObjectPath(const FSoftObjectPath& ObjectPath, bool bIncludeOnlyOnDiskAssets = false) const override;
 	virtual bool GetAllAssets(TArray<FAssetData>& OutAssetData, bool bIncludeOnlyOnDiskAssets = false) const override;
 	virtual bool EnumerateAllAssets(TFunctionRef<bool(const FAssetData&)> Callback, bool bIncludeOnlyOnDiskAssets = false) const override;
 	virtual void GetPackagesByName(FStringView PackageName, TArray<FName>& OutPackageNames) const override;
@@ -65,6 +67,8 @@ public:
 	virtual TOptional<FAssetPackageData> GetAssetPackageDataCopy(FName PackageName) const override;
 	virtual void EnumerateAllPackages(TFunctionRef<void(FName PackageName, const FAssetPackageData& PackageData)> Callback) const override;
 	virtual bool DoesPackageExistOnDisk(FName PackageName, FString* OutCorrectCasePackageName = nullptr, FString* OutExtension = nullptr) const override;
+	virtual FSoftObjectPath GetRedirectedObjectPath(const FSoftObjectPath& ObjectPath) const override;
+	// UE_DEPRECATED(5.1, "Asset path FNames have been deprecated, use FSoftObjectPath instead.")
 	virtual FName GetRedirectedObjectPath(const FName ObjectPath) const override;
 	UE_DEPRECATED(5.1, "Class names are now represented by path names. Please use a version of this function that uses FTopLevelAssetPath.")
 	virtual bool GetAncestorClassNames(FName ClassName, TArray<FName>& OutAncestorClassNames) const override;
@@ -109,7 +113,7 @@ public:
 	virtual SIZE_T GetAllocatedSize(bool bLogDetailed = false) const override;
 	virtual void LoadPackageRegistryData(FArchive& Ar, FLoadPackageRegistryData& InOutData) const override;
 	virtual void LoadPackageRegistryData(const FString& PackageFilename, FLoadPackageRegistryData& InOutData) const override;
-	virtual void InitializeTemporaryAssetRegistryState(FAssetRegistryState& OutState, const FAssetRegistrySerializationOptions& Options, bool bRefreshExisting = false, const TMap<FName, FAssetData*>& OverrideData = TMap<FName, FAssetData*>()) const override;
+	virtual void InitializeTemporaryAssetRegistryState(FAssetRegistryState& OutState, const FAssetRegistrySerializationOptions& Options, bool bRefreshExisting = false) const override;
 #if ASSET_REGISTRY_STATE_DUMPING_ENABLED
 	virtual void DumpState(const TArray<FString>& Arguments, TArray<FString>& OutPages, int32 LinesPerPage = 1) const override;
 #endif
@@ -169,7 +173,13 @@ public:
 
 protected:
 	virtual void SetManageReferences(const TMultiMap<FAssetIdentifier, FAssetIdentifier>& ManagerMap, bool bClearExisting, UE::AssetRegistry::EDependencyCategory RecurseType, TSet<FDependsNode*>& ExistingManagedNodes, ShouldSetManagerPredicate ShouldSetManager = nullptr) override;
-	virtual bool SetPrimaryAssetIdForObjectPath(const FName ObjectPath, FPrimaryAssetId PrimaryAssetId) override;
+	virtual bool SetPrimaryAssetIdForObjectPath(const FName ObjectPath, FPrimaryAssetId PrimaryAssetId) override
+	{
+PRAGMA_DISABLE_DEPRECATION_WARNINGS
+		return SetPrimaryAssetIdForObjectPath(FSoftObjectPath(ObjectPath), PrimaryAssetId);
+PRAGMA_ENABLE_DEPRECATION_WARNINGS
+	}
+	virtual bool SetPrimaryAssetIdForObjectPath(const FSoftObjectPath& ObjectPath, FPrimaryAssetId PrimaryAssetId) override;
 
 private:
 	void OnEnginePreExit();
diff --git a/Engine/Source/Runtime/AssetRegistry/Private/AssetRegistryArchive.cpp b/Engine/Source/Runtime/AssetRegistry/Private/AssetRegistryArchive.cpp
--- a/Engine/Source/Runtime/AssetRegistry/Private/AssetRegistryArchive.cpp
+++ b/Engine/Source/Runtime/AssetRegistry/Private/AssetRegistryArchive.cpp
@@ -13,22 +13,26 @@ static void SaveBundleEntries(FArchive& Ar, TArray<FAssetBundleEntry*>& Entries)
 {
 	for (FAssetBundleEntry* EntryPtr : Entries)
 	{
-		FAssetBundleEntry& Entry = *EntryPtr;
+		FAssetBundleEntry& Entry = *EntryPtr;	
 		Ar << Entry.BundleName;
 
-		int32 Num = Entry.BundleAssets.Num();
+		int32 Num = Entry.AssetPaths.Num();
 		Ar << Num;
 
-		TArray<FSoftObjectPath*> SortedPaths;
+		TArray<FTopLevelAssetPath*> SortedPaths;
 		SortedPaths.Reserve(Num);
-		for (FSoftObjectPath& Path : Entry.BundleAssets)
+		for (FTopLevelAssetPath& Path : Entry.AssetPaths)
 		{
 			SortedPaths.Add(&Path);
 		}
-		Algo::Sort(SortedPaths, [](FSoftObjectPath* A, FSoftObjectPath* B) { return A->LexicalLess(*B); });
-		for (FSoftObjectPath* Path : SortedPaths)
+		Algo::Sort(SortedPaths, [](FTopLevelAssetPath* A, FTopLevelAssetPath* B) { return A->Compare(*B) < 0; });
+		for (FTopLevelAssetPath* Path : SortedPaths)
 		{
-			Path->SerializePath(Ar);
+			// Serialize using FSoftObjectPath for backwards compatibility with FRedirectCollector code.
+			// We can investigate if any of this can be bypassed in future.
+
+			FSoftObjectPath TmpPath(*Path, {});
+			TmpPath.SerializePath(Ar);
 		}
 	}
 }
@@ -41,12 +45,68 @@ static void LoadBundleEntries(FArchive& Ar, TArray<FAssetBundleEntry>& Entries)
 
 		int32 Num = 0;
 		Ar << Num;
-		Entry.BundleAssets.SetNum(Num);
+		Entry.AssetPaths.SetNum(Num);
+
+		for (FTopLevelAssetPath& Path : Entry.AssetPaths)
+		{
+			FSoftObjectPath TmpPath;
+			TmpPath.SerializePath(Ar);
+			Path = TmpPath.GetAssetPath();
+		}
+
+PRAGMA_DISABLE_DEPRECATION_WARNINGS
+		Entry.BundleAssets.Reserve(Entry.AssetPaths.Num());
+		for (const FTopLevelAssetPath& Path : Entry.AssetPaths)
+		{
+			Entry.BundleAssets.Add(FSoftObjectPath(Path, {}));
+		}
+PRAGMA_ENABLE_DEPRECATION_WARNINGS
+	}
+}
 
-		for (FSoftObjectPath& Path : Entry.BundleAssets)
+static void LoadBundleEntriesOldVersion(FArchive& Ar, TArray<FAssetBundleEntry>& Entries, FAssetRegistryVersion::Type Version)
+{
+	for (FAssetBundleEntry& Entry : Entries)
+	{
+		Ar << Entry.BundleName;
+
+		int32 Num = 0;
+		Ar << Num;
+		Entry.AssetPaths.SetNum(Num);
+
+		if (Version < FAssetRegistryVersion::RemoveAssetPathFNames)
+		{
+			for (FTopLevelAssetPath& Path : Entry.AssetPaths)
+			{
+				// This change is synchronized with a change to the format of FSoftObjectPath in EUnrealEngineObjectUE5Version::FSOFTOBJECTPATH_REMOVE_ASSET_PATH_FNAMES
+				// We have to manually deserialize the old format of FSoftObjectPath
+				FName AssetPathName;
+				Ar << AssetPathName;
+				FString SubPathString;
+				Ar << SubPathString;
+			
+PRAGMA_DISABLE_DEPRECATION_WARNINGS
+				Path = FTopLevelAssetPath(AssetPathName);
+PRAGMA_ENABLE_DEPRECATION_WARNINGS
+			}
+		}
+		else
 		{
-			Path.SerializePath(Ar);
+			for (FTopLevelAssetPath& Path : Entry.AssetPaths)
+			{
+				FSoftObjectPath TmpPath;
+				TmpPath.SerializePath(Ar);
+				Path = TmpPath.GetAssetPath();
+			}
 		}
+		
+PRAGMA_DISABLE_DEPRECATION_WARNINGS
+		Entry.BundleAssets.Reserve(Entry.AssetPaths.Num());
+		for (const FTopLevelAssetPath& Path : Entry.AssetPaths)
+		{
+			Entry.BundleAssets.Add(FSoftObjectPath(Path, {}));
+		}
+PRAGMA_ENABLE_DEPRECATION_WARNINGS
 	}
 }
 
@@ -70,7 +130,7 @@ static void SaveBundles(FArchive& Ar, const TSharedPtr<FAssetBundleData, ESPMode
 	SaveBundleEntries(Ar, SortedEntries);
 }
 
-static TSharedPtr<FAssetBundleData, ESPMode::ThreadSafe> LoadBundles(FArchive& Ar)
+static FORCEINLINE TSharedPtr<FAssetBundleData, ESPMode::ThreadSafe> LoadBundlesInternal(FArchive& Ar, FAssetRegistryVersion::Type Version)
 {
 	int32 Num;
 	Ar << Num;
@@ -79,7 +139,14 @@ static TSharedPtr<FAssetBundleData, ESPMode::ThreadSafe> LoadBundles(FArchive& A
 	{
 		FAssetBundleData Temp;
 		Temp.Bundles.SetNum(Num);
-		LoadBundleEntries(Ar, Temp.Bundles);
+		if (Version == FAssetRegistryVersion::LatestVersion)
+		{
+			LoadBundleEntries(Ar, Temp.Bundles);
+		}
+		else
+		{
+			LoadBundleEntriesOldVersion(Ar, Temp.Bundles, Version);
+		}
 
 		return MakeShared<FAssetBundleData, ESPMode::ThreadSafe>(MoveTemp(Temp));
 	}
@@ -87,6 +154,16 @@ static TSharedPtr<FAssetBundleData, ESPMode::ThreadSafe> LoadBundles(FArchive& A
 	return TSharedPtr<FAssetBundleData, ESPMode::ThreadSafe>();
 }
 
+static TSharedPtr<FAssetBundleData, ESPMode::ThreadSafe> LoadBundles(FArchive& Ar)
+{
+	return LoadBundlesInternal(Ar, FAssetRegistryVersion::LatestVersion);
+}
+
+static TSharedPtr<FAssetBundleData, ESPMode::ThreadSafe> LoadBundlesOldVersion(FArchive& Ar, FAssetRegistryVersion::Type Version)
+{
+	return LoadBundlesInternal(Ar, Version);
+}
+
 #if ALLOW_NAME_BATCH_SAVING
 
 FAssetRegistryWriterOptions::FAssetRegistryWriterOptions(const FAssetRegistrySerializationOptions& Options)
@@ -223,6 +300,12 @@ void FAssetRegistryReader::SerializeTagsAndBundles(FAssetData& Out)
 	Out.TaggedAssetBundles = LoadBundles(*this);
 }
 
+void FAssetRegistryReader::SerializeTagsAndBundlesOldVersion(FAssetData& Out, FAssetRegistryVersion::Type Version)
+{
+	Out.TagsAndValues = LoadTags(*this);
+	Out.TaggedAssetBundles = LoadBundlesOldVersion(*this, Version);
+}
+
 ////////////////////////////////////////////////////////////////////////////
 
 #if WITH_DEV_AUTOMATION_TESTS 
diff --git a/Engine/Source/Runtime/AssetRegistry/Private/AssetRegistryArchive.h b/Engine/Source/Runtime/AssetRegistry/Private/AssetRegistryArchive.h
--- a/Engine/Source/Runtime/AssetRegistry/Private/AssetRegistryArchive.h
+++ b/Engine/Source/Runtime/AssetRegistry/Private/AssetRegistryArchive.h
@@ -20,6 +20,7 @@ public:
 	virtual FArchive& operator<<(FName& Value) override;
 
 	void SerializeTagsAndBundles(FAssetData& Out);
+	void SerializeTagsAndBundlesOldVersion(FAssetData& Out, FAssetRegistryVersion::Type Version);
 
 	void WaitForTasks();
 
diff --git a/Engine/Source/Runtime/AssetRegistry/Private/AssetRegistryConsoleCommands.h b/Engine/Source/Runtime/AssetRegistry/Private/AssetRegistryConsoleCommands.h
--- a/Engine/Source/Runtime/AssetRegistry/Private/AssetRegistryConsoleCommands.h
+++ b/Engine/Source/Runtime/AssetRegistry/Private/AssetRegistryConsoleCommands.h
@@ -23,6 +23,8 @@ public:
 	FAutoConsoleCommand GetReferencersCommand;
 	FAutoConsoleCommand FindInvalidUAssetsCommand;
 	FAutoConsoleCommand ScanPathCommand;
+	FAutoConsoleCommand DumpAllocatedSizeCommand;
+	FAutoConsoleCommand DumpStateCommand;
 
 	FAssetRegistryConsoleCommands()
 		: GetByNameCommand(
@@ -57,6 +59,14 @@ public:
 		TEXT("AssetRegistry.ScanPath"),
 		*LOCTEXT("CommandText_ScanPath", "<PathToScan> //Scan the given filename or directoryname for package files and load them into the assetregistry. Extra string parameters: -forcerescan, -ignoreDenyLists, -asfile, -asdir").ToString(),
 		FConsoleCommandWithArgsDelegate::CreateRaw(this, &FAssetRegistryConsoleCommands::ScanPath ) )
+	, DumpAllocatedSizeCommand(
+		TEXT("AssetRegistry.DumpAllocatedSize"),
+		*LOCTEXT("CommandText_DumpAllocatedSize", "Dump the allocations of the asset registry state to the log").ToString(),
+		FConsoleCommandWithArgsDelegate::CreateRaw(this, &FAssetRegistryConsoleCommands::DumpAllocatedSize ) )
+	, DumpStateCommand(
+		TEXT("AssetRegistry.DumpState"),
+		*LOCTEXT("CommandText_DumpState", "Dump the state of the asset registry to a file. Pass -log to dump to the log as well. Extra string parameters: All, ObjectPath, PackageName, Path, Class, Tag, Dependencies, DependencyDetails, PackageData, AssetBundles, AssetTags").ToString(),
+		FConsoleCommandWithArgsDelegate::CreateRaw(this, &FAssetRegistryConsoleCommands::DumpState ) )
 	{}
 
 	void GetByName(const TArray<FString>& Args)
@@ -207,7 +217,7 @@ public:
 				if ( FPaths::GetExtension(PackageFilename, true) == FPackageName::GetAssetPackageExtension() && !AssetData.IsUAsset())
 				{
 					// This asset was in a package with a uasset extension but did not share the name of the package
-					UE_LOG(LogAssetRegistry, Log, TEXT("%s"), *AssetData.ObjectPath.ToString());
+					UE_LOG(LogAssetRegistry, Log, TEXT("%s"), *AssetData.GetObjectPathString());
 				}
 			}
 		}
@@ -273,6 +283,48 @@ public:
 		}
 	}
 
+	void DumpAllocatedSize(const TArray<FString>& Args)
+	{
+		SIZE_T Size = IAssetRegistry::GetChecked().GetAllocatedSize(true);
+
+		UE_LOG(LogAssetRegistry, Log, TEXT("Total %2.f mb"), double(Size) / 1024.0 / 1024.0);
+	}
+
+	void DumpState(const TArray<FString>& Args)	
+	{
+		const bool bLog = Args.Contains(TEXT("log"));
+		const bool bDashLog = Args.Contains(TEXT("-log"));
+		if (Args.Num() == 0 + int32(bLog) + int32(bDashLog))
+		{
+			UE_LOG(LogAssetRegistry, Error, TEXT("No arguments for asset registry dump."));
+			return;
+		}
+
+		const bool bDoLog = bLog || bDashLog || (!ALLOW_DEBUG_FILES);
+
+#if ASSET_REGISTRY_STATE_DUMPING_ENABLED
+		FString Path = FPaths::ProfilingDir() / FString::Printf(TEXT("AssetRegistryState_%s.txt"), *FDateTime::Now().ToString());
+		TUniquePtr<FArchive> Ar{ IFileManager::Get().CreateDebugFileWriter(*Path) };
+
+		TArray<FString> Pages;
+		IAssetRegistry::GetChecked().DumpState(Args, Pages, 1000);
+		for (const FString& Page : Pages)
+		{
+			if( bDoLog )
+			{
+				UE_LOG(LogAssetRegistry, Log, TEXT("%s"), *Page);
+			}
+#if ALLOW_DEBUG_FILES
+			Ar->Logf(TEXT("%s"), *Page);
+#endif
+		}
+
+		UE_LOG(LogAssetRegistry, Display, TEXT("Dumped asset registry state to %s."), *Path);
+#else
+		UE_LOG(LogAssetRegistry, Error, TEXT("Asset registry dumping is disabled by compilation flags."));
+#endif
+	}
+
 };
 
 #undef LOCTEXT_NAMESPACE
diff --git a/Engine/Source/Runtime/AssetRegistry/Private/AssetRegistryImpl.h b/Engine/Source/Runtime/AssetRegistry/Private/AssetRegistryImpl.h
--- a/Engine/Source/Runtime/AssetRegistry/Private/AssetRegistryImpl.h
+++ b/Engine/Source/Runtime/AssetRegistry/Private/AssetRegistryImpl.h
@@ -110,6 +110,9 @@ namespace Impl
  */
 class FAssetRegistryImpl
 {
+	using FAssetDataMap = UE::AssetRegistry::Private::FAssetDataMap;
+	using FCachedAssetKey = UE::AssetRegistry::Private::FCachedAssetKey;
+
 public:
 	/** Constructor initializes as little as possible; meaningful operations are first done during Initialize. */
 	FAssetRegistryImpl();
@@ -121,6 +124,8 @@ public:
 
 
 	bool HasAssets(const FName PackagePath, const bool bRecursive) const;
+	FSoftObjectPath GetRedirectedObjectPath(const FSoftObjectPath& ObjectPath) const;
+	// UE_DEPRECATED(5.1, "Asset path FNames have been deprecated, use FSoftObjectPath instead.")
 	FName GetRedirectedObjectPath(const FName ObjectPath) const;
 	bool GetAncestorClassNames(Impl::FClassInheritanceContext& InheritanceContext, FTopLevelAssetPath ClassName,
 		TArray<FTopLevelAssetPath>& OutAncestorClassNames) const;
@@ -140,7 +145,7 @@ public:
 	void SetManageReferences(const TMultiMap<FAssetIdentifier, FAssetIdentifier>& ManagerMap,
 		bool bClearExisting, UE::AssetRegistry::EDependencyCategory RecurseType,
 		TSet<FDependsNode*>& ExistingManagedNodes, IAssetRegistry::ShouldSetManagerPredicate ShouldSetManager);
-	bool SetPrimaryAssetIdForObjectPath(Impl::FEventContext& EventContext, const FName ObjectPath, FPrimaryAssetId PrimaryAssetId);
+	bool SetPrimaryAssetIdForObjectPath(Impl::FEventContext& EventContext, const FSoftObjectPath& ObjectPath, FPrimaryAssetId PrimaryAssetId);
 	bool ResolveRedirect(const FString& InPackageName, FString& OutPackageName) const;
 #if WITH_EDITOR
 	void OnDirectoryChanged(Impl::FEventContext& EventContext, TArray<FFileChangeData>& FileChangesProcessed);
@@ -181,7 +186,7 @@ public:
 	void PushProcessLoadedAssetsBatch(Impl::FEventContext& EventContext,
 		TArrayView<FAssetData> LoadedAssetDatas, TArrayView<const UObject*> UnprocessedFromBatch);
 	/** Call LoadCalculatedDependencies on each Package updated after the last LoadCalculatedDependencies. */
-	void LoadCalculatedDependencies(TArray<FName>* AssetsToCalculate, double TickStartTime, bool& bOutInterrupted);
+	void LoadCalculatedDependencies(TArray<FName>* AssetPackageNamessToCalculate, double TickStartTime, bool& bOutInterrupted);
 	/**
 	 * Look for a CalculatedDependencies function registered for the asset(s) in the given package
 	 * and call that function to add calculated dependencies. Calculated dependencies are added only after
@@ -397,7 +402,7 @@ private:
 	TRingBuffer<TWeakObjectPtr<const UObject>> LoadedAssetsToProcess;
 
 	/** The set of object paths that have had their disk cache updated from the in memory version */
-	TSet<FName> AssetDataObjectPathsUpdatedOnLoad;
+	TSet<FSoftObjectPath> AssetDataObjectPathsUpdatedOnLoad;
 
 	/**
 	 * The set of object paths that have had their dependencies gathered since the last idle,
@@ -482,7 +487,7 @@ struct FClassInheritanceContext
 struct FScanPathContext
 {
 	FScanPathContext(FEventContext& InEventContext, const TArray<FString>& InDirs, const TArray<FString>& InFiles,
-		bool bInForceRescan = false, bool bInIgnoreDenyListScanFilters = false, TArray<FName>* FoundAssets = nullptr);
+		bool bInForceRescan = false, bool bInIgnoreDenyListScanFilters = false, TArray<FSoftObjectPath>* FoundAssets = nullptr);
 
 	TArray<FString> PackageDirs;
 	TArray<FString> LocalDirs;
@@ -490,7 +495,7 @@ struct FScanPathContext
 	TArray<FString> LocalFiles;
 	TArray<FString> LocalPaths;
 	FEventContext& EventContext;
-	TArray<FName>* OutFoundAssets = nullptr;
+	TArray<FSoftObjectPath>* OutFoundAssets = nullptr;
 	int32 NumFoundAssets = 0;
 	bool bForceRescan = false;
 	bool bIgnoreDenyListScanFilters = false;
diff --git a/Engine/Source/Runtime/AssetRegistry/Private/AssetRegistryState.cpp b/Engine/Source/Runtime/AssetRegistry/Private/AssetRegistryState.cpp
--- a/Engine/Source/Runtime/AssetRegistry/Private/AssetRegistryState.cpp
+++ b/Engine/Source/Runtime/AssetRegistry/Private/AssetRegistryState.cpp
@@ -30,7 +30,7 @@ FAssetRegistryState& FAssetRegistryState::operator=(FAssetRegistryState&& Rhs)
 {
 	Reset();
 
-	CachedAssetsByObjectPath			= MoveTemp(Rhs.CachedAssetsByObjectPath);
+	CachedAssets						= MoveTemp(Rhs.CachedAssets);
 	CachedAssetsByPackageName			= MoveTemp(Rhs.CachedAssetsByPackageName);
 	CachedAssetsByPath					= MoveTemp(Rhs.CachedAssetsByPath);
 	CachedAssetsByClass					= MoveTemp(Rhs.CachedAssetsByClass);
@@ -68,11 +68,11 @@ void FAssetRegistryState::Reset()
 	else
 	{
 		// Delete all assets in the cache
-		for (TMap<FName, FAssetData*>::TConstIterator AssetDataIt(CachedAssetsByObjectPath); AssetDataIt; ++AssetDataIt)
+		for (auto AssetDataIt = CachedAssets.CreateConstIterator(); AssetDataIt; ++AssetDataIt)
 		{
-			if (AssetDataIt.Value())
+			if (*AssetDataIt)
 			{
-				delete AssetDataIt.Value();
+				delete *AssetDataIt;
 				NumAssets--;
 			}
 		}
@@ -83,7 +83,7 @@ void FAssetRegistryState::Reset()
 	// TODO: Restore the ensure
 	// ensure(NumAssets == 0);
 	UE_CLOG(NumAssets != 0, LogAssetRegistry, Display,
-		TEXT("AssetRegistryState::Reset: NumAssets does not match the number of CachedAssetsByObjectPaths entries. Leaking some allocations."));
+		TEXT("AssetRegistryState::Reset: NumAssets does not match the number of CachedAssets entries. Leaking some allocations."));
 
 	if (PreallocatedDependsNodeDataBuffers.Num())
 	{
@@ -136,7 +136,7 @@ void FAssetRegistryState::Reset()
 	ensure(NumPackageData == 0);
 
 	// Clear cache
-	CachedAssetsByObjectPath.Empty();
+	CachedAssets.Empty();
 	CachedAssetsByPackageName.Empty();
 	CachedAssetsByPath.Empty();
 	CachedAssetsByClass.Empty();
@@ -186,10 +186,8 @@ void FAssetRegistryState::InitializeFromExistingAndPrune(const FAssetRegistrySta
 	TSet<FName> RequiredDependNodePackages;
 
 	// Duplicate asset data entries
-	for (const TPair<FName, FAssetData*>& AssetPair : ExistingState.CachedAssetsByObjectPath)
+	for (FAssetData* AssetData : ExistingState.CachedAssets)
 	{
-		const FAssetData* AssetData = AssetPair.Value;
-
 		bool bRemoveAssetData = false;
 		bool bRemoveDependencyData = true;
 
@@ -320,7 +318,7 @@ void FAssetRegistryState::InitializeFromExistingAndPrune(const FAssetRegistrySta
 	}
 }
 
-void FAssetRegistryState::InitializeFromExisting(const TMap<FName, FAssetData*>& AssetDataMap, const TMap<FAssetIdentifier, FDependsNode*>& DependsNodeMap,
+void FAssetRegistryState::InitializeFromExisting(const FAssetDataMap& AssetDataMap, const TMap<FAssetIdentifier, FDependsNode*>& DependsNodeMap, 
 	const TMap<FName, FAssetPackageData*>& AssetPackageDataMap, const FAssetRegistrySerializationOptions& Options, EInitializationMode InInitializationMode)
 {
 	LLM_SCOPE(ELLMTag::AssetRegistry);
@@ -329,18 +327,23 @@ void FAssetRegistryState::InitializeFromExisting(const TMap<FName, FAssetData*>&
 		Reset();
 	}
 
-	for (const TPair<FName, FAssetData*>& Pair : AssetDataMap)
+	for (const FAssetData* AssetDataPtr : AssetDataMap)
 	{
-		if (Pair.Value == nullptr)
+		if (AssetDataPtr == nullptr)
 		{
 			// don't do anything 
 			continue;
 		}
+		
+		const FAssetData& AssetData = *AssetDataPtr;
 
 		FAssetData* ExistingData = nullptr;
 		if (InInitializationMode != EInitializationMode::Rebuild) // minor optimization to avoid lookup in rebuild mode
 		{
-			ExistingData = CachedAssetsByObjectPath.FindRef(Pair.Key);
+			if (FAssetData** Ptr = CachedAssets.Find(FCachedAssetKey(AssetData)))
+			{
+				ExistingData = *Ptr;
+			}
 		}
 		if (InInitializationMode == EInitializationMode::OnlyUpdateExisting && ExistingData == nullptr)
 		{
@@ -352,8 +355,6 @@ void FAssetRegistryState::InitializeFromExisting(const TMap<FName, FAssetData*>&
 		}
 
 		// Filter asset registry tags now
-		const FAssetData& AssetData = *Pair.Value;
-
 		FAssetDataTagMap LocalTagsAndValues;
 		FAssetRegistryState::FilterTags(AssetData.TagsAndValues, LocalTagsAndValues, Options.CookFilterlistTagsByClass.Find(AssetData.AssetClassPath), Options);
 		
@@ -475,8 +476,7 @@ void FAssetRegistryState::PruneAssetData(const TSet<FName>& RequiredPackages, co
 	TSet<FName> RequiredDependNodePackages;
 
 	// Generate list up front as the maps will get cleaned up
-	TArray<FAssetData*> AllAssetData;
-	CachedAssetsByObjectPath.GenerateValueArray(AllAssetData);
+	TArray<FAssetData*> AllAssetData = CachedAssets.Array();
 	TSet<FDependsNode*> RemoveDependsNodes;
 
 	// Remove assets and mark-for-removal any dependencynodes for assets removed due to having no tags
@@ -678,16 +678,16 @@ bool FAssetRegistryState::EnumerateAssets(const FARCompiledFilter& Filter, const
 	}
 
 	// On disk object paths
-	if (Filter.ObjectPaths.Num() > 0)
+	if (Filter.SoftObjectPaths.Num() > 0)
 	{
 		TArray<FAssetData*>& ObjectPathsFilter = FilterResults.Emplace_GetRef();
-		ObjectPathsFilter.Reserve(Filter.ObjectPaths.Num());
+		ObjectPathsFilter.Reserve(Filter.SoftObjectPaths.Num());
 
-		for (FName ObjectPath : Filter.ObjectPaths)
+		for (const FSoftObjectPath& ObjectPath : Filter.SoftObjectPaths)
 		{
-			if (FAssetData* AssetDataPtr = CachedAssetsByObjectPath.FindRef(ObjectPath))
+			if (FAssetData* const* AssetDataPtr = CachedAssets.Find(FCachedAssetKey(ObjectPath)))
 			{
-				ObjectPathsFilter.Add(AssetDataPtr);
+				ObjectPathsFilter.Add(*AssetDataPtr);
 			}
 		}
 	}
@@ -812,10 +812,8 @@ bool FAssetRegistryState::GetAllAssets(const TSet<FName>& PackageNamesToSkip, TA
 bool FAssetRegistryState::EnumerateAllAssets(const TSet<FName>& PackageNamesToSkip, TFunctionRef<bool(const FAssetData&)> Callback, bool bSkipARFilteredAssets) const
 {
 	// All unloaded disk assets
-	for (const TPair<FName, FAssetData*>& AssetDataPair : CachedAssetsByObjectPath)
+	for (const FAssetData* AssetData : CachedAssets)
 	{
-		const FAssetData* AssetData = AssetDataPair.Value;
-
 		if (AssetData &&
 			!PackageNamesToSkip.Contains(AssetData->PackageName) &&
 			(!bSkipARFilteredAssets || !UE::AssetRegistry::FFiltering::ShouldSkipAsset(AssetData->AssetClassPath, AssetData->PackageFlags)))
@@ -829,6 +827,17 @@ bool FAssetRegistryState::EnumerateAllAssets(const TSet<FName>& PackageNamesToSk
 	return true;
 }
 
+void FAssetRegistryState::EnumerateAllAssets(TFunctionRef<void(const FAssetData&)> Callback) const
+{
+	for (const FAssetData* AssetData : CachedAssets)
+	{
+		if (AssetData)
+		{
+			Callback(*AssetData);
+		}
+	}
+}
+
 void FAssetRegistryState::GetPackagesByName(FStringView PackageName, TArray<FName>& OutPackageNames) const
 {
 	// Note that we use CachedAssetsByPackageName rather than CachedPackageData because CachedPackageData
@@ -1032,7 +1041,7 @@ bool FAssetRegistryState::Save(FArchive& OriginalAr, const FAssetRegistrySeriali
 #if !ALLOW_NAME_BATCH_SAVING
 	checkf(false, TEXT("Cannot save cooked AssetRegistryState in this configuration"));
 #else
-	check(CachedAssetsByObjectPath.Num() == NumAssets);
+	check(CachedAssets.Num() == NumAssets);
 
 	FAssetRegistryVersion::Type Version = FAssetRegistryVersion::LatestVersion;
 	FAssetRegistryVersion::SerializeVersion(OriginalAr, Version);
@@ -1041,16 +1050,16 @@ bool FAssetRegistryState::Save(FArchive& OriginalAr, const FAssetRegistrySeriali
 	FAssetRegistryWriter Ar(FAssetRegistryWriterOptions(Options), OriginalAr);
 
 	// serialize number of objects
-	int32 AssetCount = CachedAssetsByObjectPath.Num();
+	int32 AssetCount = CachedAssets.Num();
 	Ar << AssetCount;
 
 	// Write asset data first
-	TArray<TPair<FName, FAssetData*>> SortedAssetsByObjectPath = CachedAssetsByObjectPath.Array();
-	Algo::Sort(SortedAssetsByObjectPath, [](TPair<FName, FAssetData*>& A, TPair<FName, FAssetData*>& B) { return A.Key.LexicalLess(B.Key); });
-	for (const TPair<FName, FAssetData*>& Pair : SortedAssetsByObjectPath)
+	TArray<FAssetData*> SortedAssetsByObjectPath = CachedAssets.Array();
+	Algo::Sort(SortedAssetsByObjectPath, [](const FAssetData* A, const FAssetData* B) { return A->GetSoftObjectPath().LexicalLess(B->GetSoftObjectPath()); });
+	for (FAssetData* AssetData : SortedAssetsByObjectPath)
 	{
 		// Hardcoding FAssetRegistryVersion::LatestVersion here so that branches can get optimized out in the forceinlined SerializeForCache
-		Pair.Value->SerializeForCache(Ar);
+		AssetData->SerializeForCache(Ar);
 	}
 
 	// Serialize Dependencies
@@ -1096,7 +1105,7 @@ bool FAssetRegistryState::Save(FArchive& OriginalAr, const FAssetRegistrySeriali
 		{
 			if (!bAsReferencer)
 			{
-				InDependency = ResolveRedirector(InDependency, CachedAssetsByObjectPath, RedirectCache);
+				InDependency = ResolveRedirector(InDependency, CachedAssets, RedirectCache);
 			}
 			if (!InDependency)
 			{
@@ -1386,7 +1395,7 @@ void FAssetRegistryState::LoadDependencies_BeforeFlags(FArchive& Ar, bool bSeria
 
 SIZE_T FAssetRegistryState::GetAllocatedSize(bool bLogDetailed) const
 {
-	SIZE_T MapMemory = CachedAssetsByObjectPath.GetAllocatedSize();
+	SIZE_T MapMemory = CachedAssets.GetAllocatedSize();
 	MapMemory += CachedAssetsByPackageName.GetAllocatedSize();
 	MapMemory += CachedAssetsByPath.GetAllocatedSize();
 	MapMemory += CachedAssetsByClass.GetAllocatedSize();
@@ -1421,25 +1430,48 @@ SIZE_T FAssetRegistryState::GetAllocatedSize(bool bLogDetailed) const
 		UE_LOG(LogAssetRegistry, Log, TEXT("Index Size: %" SIZE_T_FMT "k"), MapMemory / 1024);
 	}
 
-	SIZE_T AssetDataSize = 0;
+	SIZE_T AssetDataSize = 0, AssetBundlesSize = 0, NumAssetBundles = 0, NumSoftObjectPaths = 0, NumTopLevelAssetPaths = 0;
 	FAssetDataTagMapSharedView::FMemoryCounter TagMemoryUsage;
 
-	for (const TPair<FName, FAssetData*>& AssetDataPair : CachedAssetsByObjectPath)
+	for (const FAssetData* AssetData : CachedAssets)
 	{
-		const FAssetData& AssetData = *AssetDataPair.Value;
-		
-		AssetDataSize += sizeof(AssetData);
-		AssetDataSize += AssetData.ChunkIDs.GetAllocatedSize();
-		TagMemoryUsage.Include(AssetData.TagsAndValues); 
+		AssetDataSize += sizeof(*AssetData);
+		AssetDataSize += AssetData->ChunkIDs.GetAllocatedSize();
+		TagMemoryUsage.Include(AssetData->TagsAndValues);
+		if (AssetData->TaggedAssetBundles.IsValid())
+		{
+			AssetBundlesSize += sizeof(FAssetBundleData);
+			AssetBundlesSize += AssetData->TaggedAssetBundles->Bundles.GetAllocatedSize();
+			NumAssetBundles += AssetData->TaggedAssetBundles->Bundles.Num();
+			for( const FAssetBundleEntry& Entry : AssetData->TaggedAssetBundles->Bundles)
+			{
+#if WITH_EDITORONLY_DATA
+				PRAGMA_DISABLE_DEPRECATION_WARNINGS;
+				AssetBundlesSize += Entry.BundleAssets.GetAllocatedSize();
+				NumSoftObjectPaths += Entry.BundleAssets.Num();
+				for (const FSoftObjectPath& Path : Entry.BundleAssets)
+				{
+					AssetBundlesSize += Path.GetSubPathString().GetAllocatedSize();
+				}
+				PRAGMA_ENABLE_DEPRECATION_WARNINGS;
+#endif
+				AssetBundlesSize += Entry.AssetPaths.GetAllocatedSize();
+				NumTopLevelAssetPaths += Entry.AssetPaths.Num();
+			}
+		}
 	}
 
 	if (bLogDetailed)
 	{
-		UE_LOG(LogAssetRegistry, Log, TEXT("AssetData Count: %d"), CachedAssetsByObjectPath.Num());
+		UE_LOG(LogAssetRegistry, Log, TEXT("AssetData Count: %d"), CachedAssets.Num());
 		UE_LOG(LogAssetRegistry, Log, TEXT("AssetData Static Size: %" SIZE_T_FMT "k"), AssetDataSize / 1024);
 		UE_LOG(LogAssetRegistry, Log, TEXT("Loose Tags: %" SIZE_T_FMT "k"), TagMemoryUsage.GetLooseSize() / 1024);
 		UE_LOG(LogAssetRegistry, Log, TEXT("Fixed Tags: %" SIZE_T_FMT "k"), TagMemoryUsage.GetFixedSize() / 1024);
 		UE_LOG(LogAssetRegistry, Log, TEXT("TArray<FAssetData*>: %" SIZE_T_FMT "k"), MapArrayMemory / 1024);
+		UE_LOG(LogAssetRegistry, Log, TEXT("AssetBundle Count: %u"), NumAssetBundles);
+		UE_LOG(LogAssetRegistry, Log, TEXT("AssetBundle Size: %uk"), AssetBundlesSize / 1024);
+		UE_LOG(LogAssetRegistry, Log, TEXT("AssetBundle FSoftObjectPath Count: %u"), NumSoftObjectPaths);
+		UE_LOG(LogAssetRegistry, Log, TEXT("AssetBundle FTopLevelAssetPath Count: %u"), NumTopLevelAssetPaths);
 	}
 
 	SIZE_T DependNodesSize = 0, DependenciesSize = 0;
@@ -1465,7 +1497,7 @@ SIZE_T FAssetRegistryState::GetAllocatedSize(bool bLogDetailed) const
 		PackageDataSize += PackageDataPair.Value->GetAllocatedSize();
 	}
 
-	SIZE_T TotalBytes = MapMemory + AssetDataSize + TagMemoryUsage.GetFixedSize()  + TagMemoryUsage.GetLooseSize() + DependNodesSize + DependenciesSize + PackageDataSize + MapArrayMemory;
+	SIZE_T TotalBytes = MapMemory + AssetDataSize + AssetBundlesSize + TagMemoryUsage.GetFixedSize()  + TagMemoryUsage.GetLooseSize() + DependNodesSize + DependenciesSize + PackageDataSize + MapArrayMemory;
 
 	if (bLogDetailed)
 	{
@@ -1478,7 +1510,7 @@ SIZE_T FAssetRegistryState::GetAllocatedSize(bool bLogDetailed) const
 }
 
 FDependsNode* FAssetRegistryState::ResolveRedirector(FDependsNode* InDependency,
-													 TMap<FName, FAssetData*>& InAllowedAssets,
+													 const FAssetDataMap& InAllowedAssets,
 													 TMap<FDependsNode*, FDependsNode*>& InCache)
 {
 	if (InCache.Contains(InDependency))
@@ -1518,7 +1550,9 @@ FDependsNode* FAssetRegistryState::ResolveRedirector(FDependsNode* InDependency,
 						{
 							return; // Already looked at this dependency node
 						}
-						if (InAllowedAssets.Contains(InDepends->GetPackageName()))
+						const FAssetIdentifier& AssetId = InDepends->GetIdentifier();
+						FSoftObjectPath AssetPath(AssetId.PackageName, AssetId.ObjectName, FString());
+						if (InAllowedAssets.Contains(FCachedAssetKey(AssetPath)))
 						{
 							// This asset is in the allowed asset list, so take this as the redirect target
 							Result = InDepends;
@@ -1579,12 +1613,12 @@ void FAssetRegistryState::SetAssetDatas(TArrayView<FAssetData> AssetDatas, const
 	
 	auto SetPathCache = [&]() 
 	{
-		CachedAssetsByObjectPath.Empty(AssetDatas.Num());
+		CachedAssets.Empty(AssetDatas.Num());
 		for (FAssetData& AssetData : AssetDatas)
 		{
-			CachedAssetsByObjectPath.Add(AssetData.ObjectPath, &AssetData);
+			CachedAssets.Add(&AssetData);
 		}
-		ensure(NumAssets == CachedAssetsByObjectPath.Num());
+		ensure(NumAssets == CachedAssets.Num());
 	};
 
 	// FAssetDatas sharing package name are very rare.
@@ -1647,18 +1681,17 @@ void FAssetRegistryState::SetAssetDatas(TArrayView<FAssetData> AssetDatas, const
 
 void FAssetRegistryState::AddAssetData(FAssetData* AssetData)
 {
-
-	FAssetData*& ExistingByObjectPath = CachedAssetsByObjectPath.FindOrAdd(AssetData->ObjectPath);
-	if (ExistingByObjectPath)
+	bool bAlreadyInSet = false;
+	CachedAssets.Add(AssetData, &bAlreadyInSet);
+	if (bAlreadyInSet)
 	{
 		UE_LOG(LogAssetRegistry, Error, TEXT("AddAssetData called with ObjectPath %s which already exists. ")
-			TEXT("This will overwrite and leak the existing AssetData."), *AssetData->ObjectPath.ToString());
+			TEXT("This will overwrite and leak the existing AssetData."), *FCachedAssetKey(*AssetData).ToString());
 	}
 	else
 	{
 		++NumAssets;
 	}
-	ExistingByObjectPath = AssetData;
 
 	TArray<FAssetData*, TInlineAllocator<1>>& PackageAssets = CachedAssetsByPackageName.FindOrAdd(AssetData->PackageName);
 	TArray<FAssetData*>& PathAssets = CachedAssetsByPath.FindOrAdd(AssetData->PackagePath);
@@ -1684,12 +1717,13 @@ void FAssetRegistryState::AddTagsToAssetData(const FSoftObjectPath& InObjectPath
 		return;
 	}
 
-	FAssetData* AssetData = CachedAssetsByObjectPath.FindRef(InObjectPath.GetAssetPathName());
-	if (AssetData == nullptr)
+	FSetElementId Id = CachedAssets.FindId(FCachedAssetKey(InObjectPath));
+	if (!Id.IsValidId())
 	{
-		UE_LOG(LogAssetRegistry, Error, TEXT("AddTagsToAssetData called with asset data that doesn't exist! Tags not added. ObjectPath: %s"), *InObjectPath.GetAssetPathName().ToString());
+		UE_LOG(LogAssetRegistry, Error, TEXT("AddTagsToAssetData called with asset data that doesn't exist! Tags not added. ObjectPath: %s"), *InObjectPath.ToString());
 		return;
 	}
+	FAssetData* AssetData = CachedAssets[Id];
 	FAssetDataTagMap Tags = AssetData->TagsAndValues.CopyMap();
 	Tags.Append(MoveTemp(InTagsAndValues));
 	SetTagsOnExistingAsset(AssetData, MoveTemp(Tags));
@@ -1698,9 +1732,8 @@ void FAssetRegistryState::AddTagsToAssetData(const FSoftObjectPath& InObjectPath
 
 void FAssetRegistryState::FilterTags(const FAssetRegistrySerializationOptions& Options)
 {
-	for (TPair<FName, FAssetData*>& Pair : CachedAssetsByObjectPath)
+	for (FAssetData* AssetData : CachedAssets) 
 	{
-		FAssetData* AssetData = Pair.Value;
 		check(AssetData);
 
 		FAssetDataTagMap LocalTagsAndValues;
@@ -1741,9 +1774,9 @@ void FAssetRegistryState::SetTagsOnExistingAsset(FAssetData* AssetData, FAssetDa
 
 void FAssetRegistryState::UpdateAssetData(const FAssetData& NewAssetData, bool bCreateIfNotExists)
 {
-	if (FAssetData* AssetData = CachedAssetsByObjectPath.FindRef(NewAssetData.ObjectPath))
+	if (FAssetData** AssetData = CachedAssets.Find(FCachedAssetKey(NewAssetData)))
 	{
-		UpdateAssetData(AssetData, NewAssetData);
+		UpdateAssetData(*AssetData, NewAssetData);
 	}
 	else
 	{
@@ -1753,9 +1786,9 @@ void FAssetRegistryState::UpdateAssetData(const FAssetData& NewAssetData, bool b
 
 void FAssetRegistryState::UpdateAssetData(FAssetData&& NewAssetData, bool bCreateIfNotExists)
 {
-	if (FAssetData* AssetData = CachedAssetsByObjectPath.FindRef(NewAssetData.ObjectPath))
+	if (FAssetData** AssetData = CachedAssets.Find(FCachedAssetKey(NewAssetData)))
 	{
-		UpdateAssetData(AssetData, MoveTemp(NewAssetData));
+		UpdateAssetData(*AssetData, MoveTemp(NewAssetData));
 	}
 	else if (bCreateIfNotExists)
 	{
@@ -1771,31 +1804,21 @@ void FAssetRegistryState::UpdateAssetData(FAssetData* AssetData, const FAssetDat
 void FAssetRegistryState::UpdateAssetData(FAssetData* AssetData, FAssetData&& NewAssetData, bool* bOutModified)
 {
 	bool bKeyFieldIsModified = false;
+	FCachedAssetKey OldKey(AssetData);
+	FCachedAssetKey NewKey(NewAssetData);
 
 	// Update ObjectPath
-	if (AssetData->ObjectPath != NewAssetData.ObjectPath)
+	if (OldKey != NewKey)
 	{
 		bKeyFieldIsModified = true;
-		int32 NumRemoved = CachedAssetsByObjectPath.Remove(AssetData->ObjectPath);
+		int32 NumRemoved = CachedAssets.Remove(OldKey);
 		check(NumRemoved <= 1);
 		if (NumRemoved == 0)
 		{
 			UE_LOG(LogAssetRegistry, Error, TEXT("UpdateAssetData called on AssetData %s that is not present in the AssetRegistry."),
-				*AssetData->ObjectPath.ToString());
+				*AssetData->GetObjectPathString());
 		}
 		NumAssets -= NumRemoved;
-		FAssetData*& Existing = CachedAssetsByObjectPath.FindOrAdd(NewAssetData.ObjectPath, AssetData);
-		if (Existing)
-		{
-			UE_LOG(LogAssetRegistry, Error, TEXT("UpdateAssetData called with a change in ObjectPath from Old=\"%s\" to New=\"%s\", ")
-				TEXT("but the new ObjectPath is already present with another AssetData. This will overwrite and leak the existing AssetData."),
-				*AssetData->ObjectPath.ToString(), *NewAssetData.ObjectPath.ToString());
-		}
-		else
-		{
-			++NumAssets;
-		}
-		Existing = AssetData;
 	}
 
 	// Update PackageName
@@ -1885,6 +1908,23 @@ void FAssetRegistryState::UpdateAssetData(FAssetData* AssetData, FAssetData&& Ne
 
 	// Copy in new values
 	*AssetData = MoveTemp(NewAssetData);
+
+	// Can only re-add to asset map when key fields are updated
+	if (OldKey != NewKey)
+	{
+		bool bExisting = false;
+		CachedAssets.Add(AssetData, &bExisting);
+		if (bExisting)
+		{
+			UE_LOG(LogAssetRegistry, Error, TEXT("UpdateAssetData called with a change in ObjectPath from Old=\"%s\" to New=\"%s\", ")
+				TEXT("but the new ObjectPath is already present with another AssetData. This will overwrite and leak the existing AssetData."),
+				*OldKey.ToString(), *NewKey.ToString());
+}
+		else
+		{
+			++NumAssets;
+		}
+	}
 }
 
 bool FAssetRegistryState::UpdateAssetDataPackageFlags(FName PackageName, uint32 PackageFlags)
@@ -1912,12 +1952,12 @@ void FAssetRegistryState::RemoveAssetData(FAssetData* AssetData, bool bRemoveDep
 		return;
 	}
 
-	int32 NumRemoved = CachedAssetsByObjectPath.Remove(AssetData->ObjectPath);
+	int32 NumRemoved = CachedAssets.Remove(FCachedAssetKey(AssetData));
 	check(NumRemoved <= 1);
 	if (NumRemoved == 0)
 	{
 		UE_LOG(LogAssetRegistry, Error, TEXT("RemoveAssetData called on AssetData %s that is not present in the AssetRegistry."),
-			*AssetData->ObjectPath.ToString());
+			*FCachedAssetKey(*AssetData).ToString());
 		return;
 	}
 
@@ -2034,11 +2074,11 @@ bool FAssetRegistryState::RemoveDependsNode(const FAssetIdentifier& Identifier)
 
 void FAssetRegistryState::GetPrimaryAssetsIds(TSet<FPrimaryAssetId>& OutPrimaryAssets) const
 {
-	for (TMap<FName, FAssetData*>::ElementType Element : CachedAssetsByObjectPath)
+	for (FAssetData* AssetData: CachedAssets)
 	{
-		if (Element.Value)
+		if (AssetData)
 		{
-			FPrimaryAssetId PrimaryAssetId = Element.Value->GetPrimaryAssetId();
+			FPrimaryAssetId PrimaryAssetId = AssetData->GetPrimaryAssetId();
 			if (PrimaryAssetId.IsValid())
 			{
 				OutPrimaryAssets.Add(PrimaryAssetId);
@@ -2124,7 +2164,7 @@ bool IsFilterValid(const FARCompiledFilter& Filter)
 {
 	if (Filter.PackageNames.Contains(NAME_None) ||
 		Filter.PackagePaths.Contains(NAME_None) ||
-		Filter.ObjectPaths.Contains(NAME_None) ||
+		Filter.SoftObjectPaths.Contains(FSoftObjectPath()) ||
 		Filter.ClassPaths.Contains(FTopLevelAssetPath()) ||
 		Filter.TagsAndValues.Contains(NAME_None)
 		)
@@ -2212,7 +2252,7 @@ static void PrintAssetDataMap(FString Name, const MapType& AssetMap, TStringBuil
 		Items.Reset();
 		Items.Append(AssetArray);
 		Items.Sort([](const FAssetData& A, const FAssetData& B)
-			{ return A.ObjectPath.ToString() < B.ObjectPath.ToString(); }
+			{ return A.GetSoftObjectPath().LexicalLess(B.GetSoftObjectPath()); }
 		);
 
 		PageBuffer.Append(TEXT("\t"));
@@ -2222,7 +2262,7 @@ static void PrintAssetDataMap(FString Name, const MapType& AssetMap, TStringBuil
 		for (const FAssetData* Data : Items)
 		{
 			PageBuffer.Append(TEXT("\t\t"));
-			Data->ObjectPath.AppendString(PageBuffer);
+			Data->AppendObjectPath(PageBuffer);
 			if (PrintValue)
 			{
 				PrintValue(Key, *Data);
@@ -2237,7 +2277,7 @@ static void PrintAssetDataMap(FString Name, const MapType& AssetMap, TStringBuil
 
 void FAssetRegistryState::Dump(const TArray<FString>& Arguments, TArray<FString>& OutPages, int32 LinesPerPage) const
 {
-	int32 ExpectedNumLines = 14 + CachedAssetsByObjectPath.Num() * 5 + CachedDependsNodes.Num() + CachedPackageData.Num();
+	int32 ExpectedNumLines = 14 + CachedAssets.Num() * 5 + CachedDependsNodes.Num() + CachedPackageData.Num();
 	const int32 EstimatedLinksPerNode = 10*2; // Each dependency shows up once as a dependency and once as a reference
 	const int32 EstimatedCharactersPerLine = 100;
 	bool bAllFields = Arguments.Contains(TEXT("All"));
@@ -2345,18 +2385,24 @@ void FAssetRegistryState::Dump(const TArray<FString>& Arguments, TArray<FString>
 		PageBuffer.Append(TEXT("--- Begin CachedAssetsByObjectPath ---"));
 		AddLine();
 
-		TArray<FName> Keys;
-		CachedAssetsByObjectPath.GenerateKeyArray(Keys);
-		Keys.Sort(FNameLexicalLess());
+		TArray<FCachedAssetKey> Keys;
+		Keys.Reserve(CachedAssets.Num());
+		for( const FAssetData* AssetData : CachedAssets)
+		{
+			Keys.Emplace(FCachedAssetKey(*AssetData));
+		}
+		Keys.Sort([](const FCachedAssetKey& A, const FCachedAssetKey& B) {
+			return A.Compare(B) < 0;
+		});
 
-		for (FName ObjectPath : Keys)
+		for (const FCachedAssetKey& Key : Keys)
 		{
 			PageBuffer.Append(TEXT("	"));
-			ObjectPath.AppendString(PageBuffer);
+			Key.AppendString(PageBuffer);
 			AddLine();
 		}
 
-		PageBuffer.Appendf(TEXT("--- End CachedAssetsByObjectPath : %d entries ---"), CachedAssetsByObjectPath.Num());
+		PageBuffer.Appendf(TEXT("--- End CachedAssetsByObjectPath : %d entries ---"), CachedAssets.Num());
 		AddLine();
 	}
 
@@ -2375,7 +2421,8 @@ void FAssetRegistryState::Dump(const TArray<FString>& Arguments, TArray<FString>
 		PrintAssetDataMap(TEXT("CachedAssetsByClass"), CachedAssetsByClass, PageBuffer, AddLine);
 	}
 
-	if (bAllFields || Arguments.Contains(TEXT("Tag")))
+	// Only print this if it's requested specifically - '-all' will print tags-per-asset rather than assets-per-tag 
+	if (Arguments.Contains(TEXT("Tag")))
 	{
 		PrintAssetDataMap(TEXT("CachedAssetsByTag"), CachedAssetsByTag, PageBuffer, AddLine,
 			[&PageBuffer, &AddLine](const FName& TagName, const FAssetData& Data)
@@ -2384,6 +2431,35 @@ void FAssetRegistryState::Dump(const TArray<FString>& Arguments, TArray<FString>
 			});
 	}
 
+	if (bAllFields || Arguments.Contains(TEXT("AssetTags")))
+	{
+		int32 Counter = 0;
+		PageBuffer.Append(TEXT("--- Begin AssetTags ---"));
+		AddLine();
+
+		for (const FAssetData* AssetData : CachedAssets)
+		{
+			if (AssetData->TagsAndValues.Num() == 0)
+			{
+				continue;
+			}
+			++Counter;
+
+			PageBuffer << TEXT("  ") << FCachedAssetKey(AssetData);
+			AddLine();
+
+			AssetData->TagsAndValues.ForEach([&](const TPair<FName, FAssetTagValueRef>& TagPair) 
+			{
+				PageBuffer << TEXT("    ") << TagPair.Key << TEXT(" : ") << *TagPair.Value.AsString();
+				AddLine();	
+			});
+		}
+
+		PageBuffer.Appendf(TEXT("--- End AssetTags : %d entries ---"), Counter);
+		AddLine();
+	}
+
+
 	if ((bAllFields || Arguments.Contains(TEXT("Dependencies"))) && !bDumpDependencyDetails)
 	{
 		PageBuffer.Appendf(TEXT("--- Begin CachedDependsNodes ---"));
@@ -2548,7 +2624,7 @@ void FAssetRegistryState::Dump(const TArray<FString>& Arguments, TArray<FString>
 			Key.AppendString(PageBuffer);
 			PageBuffer.Append(TEXT(" : "));
 			PRAGMA_DISABLE_DEPRECATION_WARNINGS
-			PageBuffer.Append(PackageData->PackageGuid.ToString()); // TODO: Add AppendString for Guid
+			PageBuffer.Append(PackageData->PackageGuid.ToString());
 			PRAGMA_ENABLE_DEPRECATION_WARNINGS
 			PageBuffer.Appendf(TEXT(" : %d bytes"), PackageData->DiskSize);
 			AddLine();
@@ -2558,6 +2634,35 @@ void FAssetRegistryState::Dump(const TArray<FString>& Arguments, TArray<FString>
 		AddLine();
 	}
 
+	if (bAllFields || Arguments.Contains(TEXT("AssetBundles")))
+	{
+		int32 Counter = 0;
+		PageBuffer.Append(TEXT("--- Begin AssetBundles ---"));
+		AddLine();
+
+		for (FAssetData* AssetData : CachedAssets)
+		{
+			if (AssetData->TaggedAssetBundles.IsValid())
+			{
+				++Counter;
+				for (const FAssetBundleEntry& Entry : AssetData->TaggedAssetBundles->Bundles)
+				{
+					PageBuffer << TEXT("  Owner: ") << FCachedAssetKey(AssetData) << TEXT(" BundleName: ") << Entry.BundleName;
+					AddLine();
+
+					for (const FTopLevelAssetPath& Path : Entry.AssetPaths)
+					{
+						PageBuffer << TEXT("    ") << Path;
+						AddLine();
+					}
+				}
+			}
+		}
+
+		PageBuffer.Appendf(TEXT("--- End AssetBundles : %d entries ---"), Counter);
+		AddLine();
+	}
+
 	if (PageBuffer.Len() > 0)
 	{
 		if (LinesPerPage == 1)
@@ -2571,4 +2676,83 @@ void FAssetRegistryState::Dump(const TArray<FString>& Arguments, TArray<FString>
 	}
 }
 
-#endif // ASSET_REGISTRY_STATE_DUMPING_ENABLED
\ No newline at end of file
+#endif // ASSET_REGISTRY_STATE_DUMPING_ENABLED
+
+
+////////////////////////////////////////////////////////////////////////////
+
+#if WITH_DEV_AUTOMATION_TESTS 
+
+#include "Misc/AutomationTest.h"
+
+IMPLEMENT_SIMPLE_AUTOMATION_TEST(FAssetRegistryAssetPathStringsTest, "Engine.AssetRegistry.AssetPathStrings", EAutomationTestFlags::ApplicationContextMask | EAutomationTestFlags::EngineFilter);
+
+// Tests that we can produce correct paths for objects represented by FCachedAssetKey and FAssetData
+// E.g	PackageName.AssetName
+//		PackageName.TopLevel:Inner
+//		Packagename.TopLevel:Inner.Innermost
+bool FAssetRegistryAssetPathStringsTest::RunTest(const FString& Parameters)
+{
+	using namespace UE::AssetRegistry::Private;
+
+	// Construct these FNames before creating FCachedAssetKey because the key tries not to create unused path names.
+	const FName TopLevelOuter = TEXT("/Path/To/PackageName");
+	const FName DirectSubObjectOuter = TEXT("/Path/To/PackageName.OuterName");
+	const FName SubSubObjectOuter = TEXT("/Path/To/PackageName.OuterName:SubOuterName");
+	const FName AssetName = TEXT("AssetName");
+
+	const FString TopLevelPathString = FString::Printf(TEXT("%s.%s"), *TopLevelOuter.ToString(), *AssetName.ToString());
+	const FString DirectSubObjectPathString = FString::Printf(TEXT("%s:%s"), *DirectSubObjectOuter.ToString(), *AssetName.ToString());
+	const FString SubSubObjectPathString = FString::Printf(TEXT("%s.%s"), *SubSubObjectOuter.ToString(), *AssetName.ToString());
+
+	const FSoftObjectPath TopLevelPath(TopLevelPathString);
+	const FSoftObjectPath DirectSubObjectPath(DirectSubObjectPathString);
+	const FSoftObjectPath SubSubObjectPath(SubSubObjectPathString);
+
+	TestEqual(TEXT("SoftObjectPath::ToString() correct for top-level asset"),		TopLevelPath.ToString(), TopLevelPathString);
+	TestEqual(TEXT("SoftObjectPath::ToString() correct for subobject asset"),		DirectSubObjectPath.ToString(), DirectSubObjectPathString);
+	TestEqual(TEXT("SoftObjectPath::ToString() correct for sub-subobject asset"),	SubSubObjectPath.ToString(), SubSubObjectPathString);
+
+	// Construct FCachedAssetKey from FSoftObjectPath of various lengths + check they have the right components
+	const FCachedAssetKey TopLevelAssetKey(TopLevelPath);
+	TestEqual(TEXT("FCachedAssetKey::OuterPath correct for top-level asset"), TopLevelAssetKey.OuterPath.ToString(), TopLevelOuter.ToString());
+	TestEqual(TEXT("FCachedAssetKey::ObjectName correct for top-level asset"), TopLevelAssetKey.ObjectName.ToString(), AssetName.ToString());
+
+	const FCachedAssetKey DirectSubObjectKey(DirectSubObjectPath);
+	TestEqual(TEXT("FCachedAssetKey::OuterPath correct for subobject asset"), DirectSubObjectKey.OuterPath.ToString(), DirectSubObjectOuter.ToString());
+	TestEqual(TEXT("FCachedAssetKey::ObjectName correct for subobject asset"), DirectSubObjectKey.ObjectName.ToString(), AssetName.ToString());
+
+	const FCachedAssetKey SubSubObjectKey(SubSubObjectPath);
+	TestEqual(TEXT("FCachedAssetKey::OuterPath correct for sub-subobject asset"), SubSubObjectKey.OuterPath.ToString(), SubSubObjectOuter.ToString());
+	TestEqual(TEXT("FCachedAssetKey::ObjectName correct for sub-subobject asset"), SubSubObjectKey.ObjectName.ToString(), AssetName.ToString());
+
+	// Construct FCachedAssetKey from FSoftObjectPath of various lengths + check they give the right strings from AppendString
+	TestEqual(TEXT("FCachedAssetKey::ToString() correct for top-level asset"),		TopLevelAssetKey.ToString(), TopLevelPathString);
+	TestEqual(TEXT("FCachedAssetKey::ToString() correct for subobject asset"),		DirectSubObjectKey.ToString(), DirectSubObjectPathString);
+	TestEqual(TEXT("FCachedAssetKey::ToString() correct for sub-subobject asset"),	SubSubObjectKey.ToString(), SubSubObjectPathString);
+
+	auto PathToAssetData = [](const FString& Path)
+	{
+		FString PackageName = FPackageName::ObjectPathToPackageName(Path);
+		return FAssetData(PackageName, Path, FTopLevelAssetPath("/Script/CoreUObject.Object"));
+	};
+
+	FAssetData TopLevelAssetData = PathToAssetData(TopLevelPathString);
+	FAssetData DirectSubObjectAssetData = PathToAssetData(DirectSubObjectPathString);
+	FAssetData SubSubObjectAssetData = PathToAssetData(SubSubObjectPathString);
+
+	auto GetAssetDataPath = [](const FAssetData& AssetData)
+	{
+		TStringBuilder<FName::StringBufferSize> Buffer;
+		AssetData.AppendObjectPath(Buffer);
+		return FString(Buffer);
+	};
+
+	// Test FAssetData::AppendPath for asset data with variable length of OptionalOuterPath
+	TestEqual(TEXT("FAssetData::AppendPath() correct for top-level asset"),			TopLevelAssetData.GetObjectPathString(), TopLevelPathString);
+	TestEqual(TEXT("FAssetData::AppendPath() correct for subobject asset"),			DirectSubObjectAssetData.GetObjectPathString(), DirectSubObjectPathString);
+	TestEqual(TEXT("FAssetData::AppendPath() correct for sub-subobject asset"),		SubSubObjectAssetData.GetObjectPathString(), SubSubObjectPathString);
+	
+	return true;
+}
+#endif // WITH_DEV_AUTOMATION_TESTS
diff --git a/Engine/Source/Runtime/AssetRegistry/Private/NameTableArchive.cpp b/Engine/Source/Runtime/AssetRegistry/Private/NameTableArchive.cpp
--- a/Engine/Source/Runtime/AssetRegistry/Private/NameTableArchive.cpp
+++ b/Engine/Source/Runtime/AssetRegistry/Private/NameTableArchive.cpp
@@ -159,6 +159,13 @@ FArchive& FNameTableArchiveReader::operator<<(FName& OutName)
 }
 
 void FNameTableArchiveReader::SerializeTagsAndBundles(FAssetData& Out)
+{
+	FAssetDataTagMap Map;
+	*this << Map;
+	Out.SetTagsAndAssetBundles(MoveTemp(Map));
+}
+	
+void FNameTableArchiveReader::SerializeTagsAndBundlesOldVersion(FAssetData& Out, int32 Version)
 {
 	FAssetDataTagMap Map;
 	*this << Map;
diff --git a/Engine/Source/Runtime/AssetRegistry/Private/NameTableArchive.h b/Engine/Source/Runtime/AssetRegistry/Private/NameTableArchive.h
--- a/Engine/Source/Runtime/AssetRegistry/Private/NameTableArchive.h
+++ b/Engine/Source/Runtime/AssetRegistry/Private/NameTableArchive.h
@@ -31,6 +31,7 @@ public:
 	virtual FArchive& operator<<(FName& Name);
 
 	void SerializeTagsAndBundles(FAssetData& Out);
+	void SerializeTagsAndBundlesOldVersion(FAssetData& Out, int32 Version);
 
 private:
 	/** Serializers for different package maps */
diff --git a/Engine/Source/Runtime/AssetRegistry/Public/AssetRegistry/AssetRegistryState.h b/Engine/Source/Runtime/AssetRegistry/Public/AssetRegistry/AssetRegistryState.h
--- a/Engine/Source/Runtime/AssetRegistry/Public/AssetRegistry/AssetRegistryState.h
+++ b/Engine/Source/Runtime/AssetRegistry/Public/AssetRegistry/AssetRegistryState.h
@@ -116,10 +116,74 @@ struct FAssetRegistryLoadOptions
 	int32 ParallelWorkers = 0;
 };
 
+namespace UE::AssetRegistry::Private
+{
+	/* 
+	* Key type for TSet<FAssetData*> in the asset registry.
+	* Top level assets are searched for by their asset path as two names (e.g. '/Path/ToPackageName' + 'AssetName')
+	* Other assets (e.g. external actors) are searched for by their full path with the whole outer chain as a single name. 
+	* (e.g. '/Path/To/Package.TopLevel:Subobject' + 'DeeperSubobject')
+	*/
+	struct FCachedAssetKey
+	{
+		explicit FCachedAssetKey(const FAssetData* InAssetData);
+		explicit FCachedAssetKey(const FAssetData& InAssetData);
+		explicit FCachedAssetKey(FTopLevelAssetPath InAssetPath);
+		explicit FCachedAssetKey(const FSoftObjectPath& InObjectPath);
+
+		FString ToString() const;
+		int32 Compare(const FCachedAssetKey& Other) const;	// Order asset keys with fast non-lexical comparison
+		void AppendString(FStringBuilderBase& Builder) const;
+
+		FName OuterPath = NAME_None;
+		FName ObjectName = NAME_None;
+	};
+
+	inline FStringBuilderBase& operator<<(FStringBuilderBase& Builder, const FCachedAssetKey& Key);
+	inline bool operator==(const FCachedAssetKey& A, const FCachedAssetKey& B);
+	inline bool operator!=(const FCachedAssetKey& A, const FCachedAssetKey& B);
+	inline uint32 GetTypeHash(const FCachedAssetKey& A);
+
+	/* 
+	* Policy type for TSet<FAssetData*> to use FCachedAssetKey for hashing/equality.
+	* This allows is to store just FAssetData* in the map without storing an extra copy of the key fields to save memory.
+	*/
+	struct FCachedAssetKeyFuncs
+	{
+		using KeyInitType = FCachedAssetKey;
+		using ElementInitType = void; // TSet doesn't actually use this type 
+
+		enum { bAllowDuplicateKeys = false };
+
+		static FORCEINLINE KeyInitType GetSetKey(const FAssetData* Element)
+		{
+			return FCachedAssetKey(*Element);
+		}
+
+		static FORCEINLINE bool Matches(KeyInitType A, KeyInitType B)
+		{
+			return A == B;
+		}
+
+		static FORCEINLINE uint32 GetKeyHash(KeyInitType Key)
+		{
+			return GetTypeHash(Key);
+		}
+	};
+
+	using FAssetDataMap = TSet<FAssetData*, FCachedAssetKeyFuncs>;
+	using FConstAssetDataMap = TSet<const FAssetData*, FCachedAssetKeyFuncs>;
+}
+
 /** The state of an asset registry, this is used internally by IAssetRegistry to represent the disk cache, and is also accessed directly to save/load cooked caches */
 class ASSETREGISTRY_API FAssetRegistryState
 {
+	using FCachedAssetKey = UE::AssetRegistry::Private::FCachedAssetKey;
 public:
+	// These types are an implementation detail and they and the functions which take/return them are subject to change without deprecation warnings.
+	using FAssetDataMap = UE::AssetRegistry::Private::FAssetDataMap;
+	using FConstAssetDataMap = UE::AssetRegistry::Private::FConstAssetDataMap;
+
 	FAssetRegistryState() = default;
 	FAssetRegistryState(const FAssetRegistryState&) = delete;
 	FAssetRegistryState(FAssetRegistryState&& Rhs) { *this = MoveTemp(Rhs); }
@@ -187,6 +251,7 @@ public:
 	 * @param bARFiltering Whether to apply filtering from UE::AssetRegistry::FFiltering (false by default)
 	 */
 	bool EnumerateAllAssets(const TSet<FName>& PackageNamesToSkip, TFunctionRef<bool(const FAssetData&)> Callback, bool bARFiltering = false) const;
+	void EnumerateAllAssets(TFunctionRef<void(const FAssetData&)> Callback) const;
 
 	/**
 	 * Gets the LongPackageNames for all packages with the given PackageName.
@@ -239,9 +304,22 @@ public:
 	 * @param ObjectPath the path of the object to be looked up
 	 * @return the assets data, null if not found
 	 */
+	UE_DEPRECATED(5.1, "Asset path FNames have been deprecated, use FSoftObjectPath instead.")
 	const FAssetData* GetAssetByObjectPath(const FName ObjectPath) const
 	{
-		FAssetData* const* FoundAsset = CachedAssetsByObjectPath.Find(ObjectPath);
+		return GetAssetByObjectPath(FSoftObjectPath(ObjectPath.ToString()));
+	}
+
+	/**
+	 * Gets the asset data for the specified object path
+	 *
+	 * @param ObjectPath the path of the object to be looked up
+	 * @return the assets data, null if not found
+	 */
+	const FAssetData* GetAssetByObjectPath(const FSoftObjectPath& ObjectPath) const 
+	{
+		FCachedAssetKey Key(ObjectPath);
+		FAssetData* const* FoundAsset = CachedAssets.Find(Key);
 		if (FoundAsset)
 		{
 			return *FoundAsset;
@@ -312,9 +390,9 @@ public:
 	}
 
 	/** Returns const version of internal ObjectPath->AssetData map for fast iteration */
-	const TMap<FName, const FAssetData*>& GetObjectPathToAssetDataMap() const
+	const FConstAssetDataMap& GetAssetDataMap() const
 	{
-		return reinterpret_cast<const TMap<FName, const FAssetData*>&>(CachedAssetsByObjectPath);
+		return reinterpret_cast<const FConstAssetDataMap&>(CachedAssets);
 	}
 
 	/** Returns const version of internal Tag->AssetDatas map for fast iteration */
@@ -381,10 +459,10 @@ public:
 	void Reset();
 
 	/** Initializes cache from existing set of asset data and depends nodes */
-	void InitializeFromExisting(const TMap<FName, FAssetData*>& AssetDataMap, const TMap<FAssetIdentifier, FDependsNode*>& DependsNodeMap, const TMap<FName, FAssetPackageData*>& AssetPackageDataMap, const FAssetRegistrySerializationOptions& Options, EInitializationMode InitializationMode = EInitializationMode::Rebuild);
+	void InitializeFromExisting(const FAssetDataMap& AssetDataMap, const TMap<FAssetIdentifier, FDependsNode*>& DependsNodeMap, const TMap<FName, FAssetPackageData*>& AssetPackageDataMap, const FAssetRegistrySerializationOptions& Options, EInitializationMode InitializationMode = EInitializationMode::Rebuild);
 	void InitializeFromExisting(const FAssetRegistryState& Existing, const FAssetRegistrySerializationOptions& Options, EInitializationMode InitializationMode = EInitializationMode::Rebuild)
 	{
-		InitializeFromExisting(Existing.CachedAssetsByObjectPath, Existing.CachedDependsNodes, Existing.CachedPackageData, Options, InitializationMode);
+		InitializeFromExisting(Existing.CachedAssets, Existing.CachedDependsNodes, Existing.CachedPackageData, Options, InitializationMode);
 	}
 
 	/** 
@@ -457,7 +535,7 @@ private:
 	void SetAssetDatas(TArrayView<FAssetData> AssetDatas, const FAssetRegistryLoadOptions& Options);
 
 	/** Find the first non-redirector dependency node starting from InDependency. */
-	FDependsNode* ResolveRedirector(FDependsNode* InDependency, TMap<FName, FAssetData*>& InAllowedAssets, TMap<FDependsNode*, FDependsNode*>& InCache);
+	FDependsNode* ResolveRedirector(FDependsNode* InDependency, const FAssetDataMap& InAllowedAssets, TMap<FDependsNode*, FDependsNode*>& InCache);
 
 	/** Finds an existing node for the given package and returns it, or returns null if one isn't found */
 	FDependsNode* FindDependsNode(const FAssetIdentifier& Identifier) const;
@@ -476,8 +554,8 @@ private:
 
 	void SetTagsOnExistingAsset(FAssetData* AssetData, FAssetDataTagMap&& NewTags);
 
-	/** The map of ObjectPath names to asset data for assets saved to disk */
-	TMap<FName, FAssetData*> CachedAssetsByObjectPath;
+	/** Set of asset data for assets saved to disk. Searched via path name types, implicitly converted to FCachedAssetKey. */
+	FAssetDataMap CachedAssets;
 
 	/** The map of package names to asset data for assets saved to disk */
 	TMap<FName, TArray<FAssetData*, TInlineAllocator<1>> > CachedAssetsByPackageName;
@@ -510,3 +588,116 @@ private:
 	friend class UAssetRegistryImpl;
 	friend class UE::AssetRegistry::FAssetRegistryImpl;
 };
+
+namespace UE::AssetRegistry::Private
+{
+	FORCEINLINE uint32 HashCombineQuick(uint32 A, uint32 B)
+	{
+		return A ^ (B + 0x9e3779b9 + (A << 6) + (A >> 2));
+	}
+
+	inline FCachedAssetKey::FCachedAssetKey(const FAssetData* InAssetData)
+	{
+		if (!InAssetData)
+		{
+			return;
+		}
+
+#if WITH_EDITORONLY_DATA
+		if (!InAssetData->GetOptionalOuterPathName().IsNone())
+		{
+			OuterPath = InAssetData->GetOptionalOuterPathName();
+		}
+		else
+#endif
+		{
+			OuterPath = InAssetData->PackageName;
+		}
+		ObjectName = InAssetData->AssetName;
+	}
+
+	inline FCachedAssetKey::FCachedAssetKey(const FAssetData& InAssetData)
+		: FCachedAssetKey(&InAssetData)
+	{
+	}
+
+	inline FCachedAssetKey::FCachedAssetKey(FTopLevelAssetPath InAssetPath)
+		: OuterPath(InAssetPath.GetPackageName())
+		, ObjectName(InAssetPath.GetAssetName())
+	{
+	}
+
+	inline FCachedAssetKey::FCachedAssetKey(const FSoftObjectPath& InObjectPath)
+	{
+		if (InObjectPath.GetAssetFName().IsNone())
+		{
+			// Packages themselves never appear in the asset registry
+			return;
+		}
+		else if (InObjectPath.GetSubPathString().IsEmpty())
+		{
+			// If InObjectPath represents a top-level asset we can just take the existing FNames.
+			OuterPath = InObjectPath.GetLongPackageFName();
+			ObjectName = InObjectPath.GetAssetFName();
+		}
+		else
+		{
+			// If InObjectPath represents a subobject we need to split the path into the path of the outer and the name of the innermost object.
+			TStringBuilder<FName::StringBufferSize> Builder;
+			InObjectPath.ToString(Builder);
+
+			const FAssetPathParts Parts = SplitIntoOuterPathAndAssetName(Builder);
+
+			// This should be impossible as at bare minimum concatenating the package name and asset name should add a separator
+			check(!Parts.OuterPath.IsEmpty() && !Parts.InnermostName.IsEmpty()); 
+
+			// Don't create FNames for this query struct. If the AssetData exists to find, the FName will already exist due to OptionalOuterPath on FAssetData.
+			OuterPath = FName(Parts.OuterPath, FNAME_Find); 
+			ObjectName = FName(Parts.InnermostName);
+		}
+	}
+	inline FString FCachedAssetKey::ToString() const
+	{
+		TStringBuilder<FName::StringBufferSize> Builder;
+		AppendString(Builder);
+		return FString(Builder);
+	}
+
+	inline int32 FCachedAssetKey::Compare(const FCachedAssetKey& Other) const
+	{
+		if (OuterPath == Other.OuterPath)
+		{
+			return ObjectName.CompareIndexes(Other.ObjectName);
+		}
+		else
+		{
+			return OuterPath.CompareIndexes(Other.OuterPath);
+		}
+	}
+
+	inline void FCachedAssetKey::AppendString(FStringBuilderBase& Builder) const
+	{
+		ConcatenateOuterPathAndObjectName(Builder, OuterPath, ObjectName);
+	}
+
+	inline FStringBuilderBase& operator<<(FStringBuilderBase& Builder, const FCachedAssetKey& Key)
+	{
+		Key.AppendString(Builder);
+		return Builder;
+	}
+
+	inline bool operator==(const FCachedAssetKey& A, const FCachedAssetKey& B)
+	{
+		return A.OuterPath == B.OuterPath && A.ObjectName == B.ObjectName;
+	}
+
+	inline bool operator!=(const FCachedAssetKey& A, const FCachedAssetKey& B)
+	{
+		return A.OuterPath != B.OuterPath || A.ObjectName != B.ObjectName;
+	}
+
+	inline uint32 GetTypeHash(const FCachedAssetKey& A)
+	{
+		return HashCombineQuick(GetTypeHash(A.OuterPath), GetTypeHash(A.ObjectName));
+	}
+}
\ No newline at end of file
diff --git a/Engine/Source/Runtime/AssetRegistry/Public/AssetRegistry/IAssetRegistry.h b/Engine/Source/Runtime/AssetRegistry/Public/AssetRegistry/IAssetRegistry.h
--- a/Engine/Source/Runtime/AssetRegistry/Public/AssetRegistry/IAssetRegistry.h
+++ b/Engine/Source/Runtime/AssetRegistry/Public/AssetRegistry/IAssetRegistry.h
@@ -216,9 +216,32 @@ public:
 	 * @param bIncludeOnlyOnDiskAssets if true, in-memory objects will be ignored. The call will be faster.
 	 * @return the assets data;Will be invalid if object could not be found
 	 */
+	// UE_DEPRECATED(5.1, "Asset path FNames have been deprecated, use Soft Object Path instead.")
 	UFUNCTION(BlueprintCallable, BlueprintPure=false, Category="AssetRegistry")
 	virtual FAssetData GetAssetByObjectPath( const FName ObjectPath, bool bIncludeOnlyOnDiskAssets = false ) const = 0;
 
+	/**
+	 * Gets the asset data for the specified object path
+	 *
+	 * @param ObjectPath the path of the object to be looked up
+	 * @param bIncludeOnlyOnDiskAssets if true, in-memory objects will be ignored. The call will be faster.
+	 * @return the assets data;Will be invalid if object could not be found
+	 */
+	UFUNCTION(BlueprintCallable, BlueprintPure=false, Category = "AssetRegistry", DisplayName="Get Asset By Object Path")
+	virtual FAssetData K2_GetAssetByObjectPath(const FSoftObjectPath& ObjectPath, bool bIncludeOnlyOnDiskAssets = false) const
+	{
+		return GetAssetByObjectPath(ObjectPath, bIncludeOnlyOnDiskAssets);
+	}
+
+	/**
+	 * Gets the asset data for the specified object path
+	 *
+	 * @param ObjectPath the path of the object to be looked up
+	 * @param bIncludeOnlyOnDiskAssets if true, in-memory objects will be ignored. The call will be faster.
+	 * @return the assets data;Will be invalid if object could not be found
+	 */
+	virtual FAssetData GetAssetByObjectPath(const FSoftObjectPath& ObjectPath, bool bIncludeOnlyOnDiskAssets = false) const = 0;
+
 	/**
 	 * Gets asset data for all assets in the registry.
 	 * This method may be slow, use a filter if possible to avoid iterating over the entire registry.
@@ -339,6 +362,9 @@ public:
 	virtual bool DoesPackageExistOnDisk(FName PackageName, FString* OutCorrectCasePackageName = nullptr, FString* OutExtension = nullptr) const = 0;
 
 	/** Uses the asset registry to look for ObjectRedirectors. This will follow the chain of redirectors. It will return the original path if no redirectors are found */
+	virtual FSoftObjectPath GetRedirectedObjectPath(const FSoftObjectPath& ObjectPath) const = 0;
+
+	// UE_DEPRECATED(5.1, "Asset path FNames have been deprecated, use FSoftObjectPath instead.")
 	virtual FName GetRedirectedObjectPath(const FName ObjectPath) const = 0;
 
 	UE_DEPRECATED(4.27, "Loading then discarding tags is no longer allowed as it can "
@@ -592,9 +618,8 @@ public:
 	 * @param OutState			This will be filled in with a copy of the asset data, platform data, and dependency data
 	 * @param Options			Serialization options that will be used to write this later
 	 * @param bRefreshExisting	If true, will not delete or add packages in OutState and will just update things that already exist
-	 * @param OverrideData		Map of ObjectPath to AssetData. If non empty, it will use this map of AssetData, and will filter Platform/Dependency data to only include this set
 	 */
-	virtual void InitializeTemporaryAssetRegistryState(FAssetRegistryState& OutState, const FAssetRegistrySerializationOptions& Options, bool bRefreshExisting = false, const TMap<FName, FAssetData*>& OverrideData = TMap<FName, FAssetData*>()) const = 0;
+	virtual void InitializeTemporaryAssetRegistryState(FAssetRegistryState& OutState, const FAssetRegistrySerializationOptions& Options, bool bRefreshExisting = false) const = 0;
 
 	UE_DEPRECATED(5.0, "Receiving a pointer is not threadsafe. Use other functions on IAssetRegistry to access the same data, or contact Epic Core team to add the threadsafe functions you require..")
 	virtual const FAssetRegistryState* GetAssetRegistryState() const = 0;
@@ -681,7 +706,9 @@ protected:
 	virtual void SetManageReferences(const TMultiMap<FAssetIdentifier, FAssetIdentifier>& ManagerMap, bool bClearExisting, UE::AssetRegistry::EDependencyCategory RecurseType, TSet<FDependsNode*>& ExistingManagedNodes, ShouldSetManagerPredicate ShouldSetManager = nullptr) = 0;
 
 	/** Sets the PrimaryAssetId for a specific asset. This should only be called by the AssetManager, and is needed when the AssetManager is more up to date than the on disk Registry */
+	// UE_DEPRECATED(5.1, "Asset path FNames have been deprecated, use FSoftObjectPath instead.")
 	virtual bool SetPrimaryAssetIdForObjectPath(const FName ObjectPath, FPrimaryAssetId PrimaryAssetId) = 0;
+	virtual bool SetPrimaryAssetIdForObjectPath(const FSoftObjectPath& ObjectPath, FPrimaryAssetId PrimaryAssetId) = 0;
 };
 
 namespace UE
diff --git a/Engine/Source/Runtime/Core/Public/UObject/ObjectVersion.h b/Engine/Source/Runtime/Core/Public/UObject/ObjectVersion.h
--- a/Engine/Source/Runtime/Core/Public/UObject/ObjectVersion.h
+++ b/Engine/Source/Runtime/Core/Public/UObject/ObjectVersion.h
@@ -62,7 +62,10 @@ enum class EUnrealEngineObjectUE5Version : uint32
 
 	// Add IsInherited to the FObjectExport entry
 	TRACK_OBJECT_EXPORT_IS_INHERITED,
-
+	
+	// Replace FName asset path in FSoftObjectPath with (package name, asset name) pair FTopLevelAssetPath
+	FSOFTOBJECTPATH_REMOVE_ASSET_PATH_FNAMES,
+	
 	// -----<new versions can be added before this line>-------------------------------------------------
 	// - this needs to be the last line (see note below)
 	AUTOMATIC_VERSION_PLUS_ONE,
diff --git a/Engine/Source/Runtime/CoreUObject/Private/AssetRegistry/AssetBundleData.cpp b/Engine/Source/Runtime/CoreUObject/Private/AssetRegistry/AssetBundleData.cpp
--- a/Engine/Source/Runtime/CoreUObject/Private/AssetRegistry/AssetBundleData.cpp
+++ b/Engine/Source/Runtime/CoreUObject/Private/AssetRegistry/AssetBundleData.cpp
@@ -7,28 +7,29 @@
 UE_IMPLEMENT_STRUCT("/Script/CoreUObject", AssetBundleData);
 
 
-namespace UE { namespace AssetBundleEntry { namespace Private
+namespace UE::AssetBundleEntry::Private
 {
-
-const FStringView BundleNamePrefix(TEXT("(BundleName=\""));
-const FStringView BundleAssetsPrefix(TEXT(",BundleAssets=("));
-const FStringView BundleAssetsSuffix(TEXT("))"));
-const FStringView BundlesPrefix(TEXT("(Bundles=("));
-const FStringView BundlesSuffix(TEXT("))"));
-const FStringView EmptyBundles(TEXT("(Bundles=)"));
-
-static bool SkipPrefix(const TCHAR*& InOutIt, FStringView Prefix)
-{
-	const bool bOk = FStringView(InOutIt, Prefix.Len()) == Prefix;
-	if (bOk)
+	const FStringView BundleNamePrefix(TEXT("(BundleName=\""));
+	const FStringView BundleAssetsPrefix(TEXT(",BundleAssets=("));
+	const FStringView AssetPathsPrefix(TEXT(",AssetPaths=("));
+	const FStringView EmptyAssetPathsPrefix(TEXT(",AssetPaths="));
+	const FStringView AssetPathsSuffix(TEXT("))"));
+	const FStringView BundlesPrefix(TEXT("(Bundles=("));
+	const FStringView BundlesSuffix(TEXT("))"));
+	const FStringView EmptyBundles(TEXT("(Bundles=)"));
+	const FStringView StructOrArraySuffix(TEXT(")"));
+
+	static bool SkipPrefix(const TCHAR*& InOutIt, FStringView Prefix)
 	{
-		InOutIt += Prefix.Len();
+		const bool bOk = FStringView(InOutIt, Prefix.Len()) == Prefix;
+		if (bOk)
+		{
+			InOutIt += Prefix.Len();
+		}
+		return bOk;
 	}
-	return bOk;
 }
 
-}}}
-
 bool FAssetBundleEntry::ExportTextItem(FString& ValueStr, const FAssetBundleEntry& DefaultValue, UObject* Parent, int32 PortFlags, UObject* ExportRootScope) const
 {	
 	if (DefaultValue.IsValid())
@@ -44,24 +45,34 @@ bool FAssetBundleEntry::ExportTextItem(FString& ValueStr, const FAssetBundleEntr
 	ValueStr += BundleNamePrefix;
 	BundleName.AppendString(ValueStr);
 	ValueStr += '\"';
-	ValueStr += BundleAssetsPrefix;
 
-	const FSoftObjectPath EmptyPath;
-	for (const FSoftObjectPath& Path : BundleAssets)
+	if (AssetPaths.Num())
 	{
-		if (!Path.ExportTextItem(ValueStr, EmptyPath, Parent, PortFlags, ExportRootScope))
+		ValueStr += AssetPathsPrefix;
+		const FSoftObjectPath EmptyAssetPath;
+		for (const FTopLevelAssetPath& Path : AssetPaths)
 		{
-			ValueStr.LeftInline(OriginalLen);
-			return false;
+			FSoftObjectPath ObjectPath(Path, {});
+			if (!ObjectPath.ExportTextItem(ValueStr, EmptyAssetPath, Parent, PortFlags, ExportRootScope))
+			{
+				ValueStr.LeftInline(OriginalLen);
+				return false;
+			}
+			
+			ValueStr.AppendChar(TEXT(','));
 		}
-		
-		ValueStr.AppendChar(TEXT(','));
-	}
 
-	// Remove last comma
-	ValueStr.LeftChopInline(1, /* shrink */ false);
+		// Remove last comma
+		ValueStr.LeftChopInline(1, /* shrink */ false);
 
-	ValueStr += BundleAssetsSuffix;
+		ValueStr += AssetPathsSuffix;
+	}
+	else
+	{
+
+		ValueStr += EmptyAssetPathsPrefix;
+		ValueStr += StructOrArraySuffix;
+	}
 
 	return true;
 }
@@ -70,41 +81,155 @@ bool FAssetBundleEntry::ImportTextItem(const TCHAR*& Buffer, int32 PortFlags, UO
 {
 	using namespace UE::AssetBundleEntry::Private;
 
-	const TCHAR* BundleNameBegin = Buffer;
-	if (SkipPrefix(/* in-out */ BundleNameBegin, BundleNamePrefix))
+	const TCHAR* BufferIt = Buffer;
+	FName LocalBundleName;
+	TArray<FSoftObjectPath> LocalObjectPaths;
+
+	if (SkipPrefix(/* in-out */ BufferIt, BundleNamePrefix))
+	{
+		if (const TCHAR* BundleNameEnd = FCString::Strchr(BufferIt, TCHAR('"')))
+		{
+			LocalBundleName = FName(UE_PTRDIFF_TO_INT32(BundleNameEnd - BufferIt), BufferIt);
+			BufferIt = BundleNameEnd + 1;
+		}
+		else
+		{
+			ErrorText->Logf(ELogVerbosity::Error, TEXT("FAssetBundleEntry::ImportTextItem - Unterminated BundleName string. When parsing import text '%s'."), Buffer);
+			return false;
+		}
+	}
+	else
+	{
+		ErrorText->Logf(ELogVerbosity::Error, TEXT("FAssetBundleEntry::ImportTextItem - Text does not start with expected prefix '%.*s'. When parsing import text '%s'."), 
+			BundleNamePrefix.Len(), BundleNamePrefix.GetData(), Buffer);
+		return false;
+	}
+
+	bool bBundleAssetsSeen = false;
+	bool bAssetPathsSeen = false;
+	while (true)
 	{
-		if (const TCHAR* BundleNameEnd = FCString::Strchr(BundleNameBegin, TCHAR('"')))
+		// Parse deprecated format 
+		if (SkipPrefix(/* in-out */ BufferIt, BundleAssetsPrefix))
 		{
-			FName Name(UE_PTRDIFF_TO_INT32(BundleNameEnd - BundleNameBegin), BundleNameBegin);
+			if (bBundleAssetsSeen)
+			{
+				// Parse error 
+				ErrorText->Logf(ELogVerbosity::Error, TEXT("FAssetBundleEntry::ImportTextItem - Encountered key BundleAssets twice. When parsing import text '%s'."), Buffer);
+				return false;
+			}
+			bBundleAssetsSeen = true;
 
-			const TCHAR* PathIt = BundleNameEnd + 1;
-			if (SkipPrefix(/* in-out */ PathIt, BundleAssetsPrefix))
+			if (SkipPrefix(/* in-out */ BufferIt, StructOrArraySuffix))
 			{
-				TArray<FSoftObjectPath> Paths;
-				while (Paths.Emplace_GetRef().ImportTextItem(/* in-out */ PathIt, PortFlags, Parent, ErrorText))
+				// End of array
+				continue;
+			}
+
+			while (LocalObjectPaths.Emplace_GetRef().ImportTextItem(/* in-out */ BufferIt, PortFlags, Parent, ErrorText))
+			{
+				if(!LocalObjectPaths.Last().GetSubPathString().IsEmpty())
+				{
+					ErrorText->Logf(ELogVerbosity::Warning, TEXT("FAssetBundleEntry::ImportTextItem - Asset bundle entries should all be top level objects, but found a subobject path '%s'. When parsing import text '%s'."), 
+						*LocalObjectPaths.Last().ToString(), Buffer);
+				}
+
+				if (*BufferIt == ',')
 				{
-					if (*PathIt == ',')
-					{
-						++PathIt;
-					}
-					else if (SkipPrefix(/* in-out */ PathIt, BundleAssetsSuffix))
-					{
-						BundleName = Name;
-						BundleAssets = MoveTemp(Paths);
-						Buffer = PathIt;
-
-						return true;
-					}
-					else
-					{
-						return false;
-					}
+					++BufferIt; // We expect another path
+				}
+				else if (SkipPrefix(/* in-out */ BufferIt, StructOrArraySuffix))
+				{
+					// End of array
+					break;
+				}
+				else
+				{
+					// Parse error 
+					ErrorText->Logf(ELogVerbosity::Error, TEXT("FAssetBundleEntry::ImportTextItem - Unterminated or ill formed BundleAssets list while importing asset bundle entry. When parsing import text '%s'."), Buffer);
+					return false;
 				}
 			}
 		}
+		else if (SkipPrefix(/* in-out */ BufferIt, AssetPathsPrefix))
+		{
+			if (bAssetPathsSeen)
+			{
+				// Parse error 
+				ErrorText->Logf(ELogVerbosity::Error, TEXT("FAssetBundleEntry::ImportTextItem - Encountered key AssetPaths twice. When parsing import text '%s'."), Buffer);
+				return false;
+			}
+			bAssetPathsSeen = true;
+
+			if (SkipPrefix(/* in-out */ BufferIt, StructOrArraySuffix))
+			{
+				// End of array
+				continue;
+			}
+
+			while (LocalObjectPaths.Emplace_GetRef().ImportTextItem(/* in-out */ BufferIt, PortFlags, Parent, ErrorText))
+			{
+				if (!LocalObjectPaths.Last().GetSubPathString().IsEmpty())
+				{
+					// Parse error, these should only have top-level paths
+					ErrorText->Logf(ELogVerbosity::Error, TEXT("FAssetBundleEntry::ImportTextItem - Asset bundle entries should all be top level objects, but found a subobject path '%s'. When parsing import text '%s'."), 
+						*LocalObjectPaths.Last().ToString(), Buffer);
+					return false;
+				}
+
+				if (*BufferIt == ',')
+				{
+					++BufferIt; // We expect another path
+				}
+				else if (SkipPrefix(/* in-out */ BufferIt, StructOrArraySuffix))
+				{
+					// End of array
+					break;
+				}
+				else
+				{
+					// Parse error 
+					ErrorText->Logf(ELogVerbosity::Error, TEXT("FAssetBundleEntry::ImportTextItem - unterminated or ill formed AssetPaths list while importing asset bundle entry: '%s'. When parsing import text '%s'."), BufferIt, Buffer);
+					return false;
+				}
+			}
+		}
+		else if (SkipPrefix(/* in-out */ BufferIt, EmptyAssetPathsPrefix))
+		{
+			// Empty array 
+			continue;
+		}
+		else if (SkipPrefix(/* in-out */ BufferIt, StructOrArraySuffix))
+		{
+			// End
+			break;
+		}
+		else
+		{
+			// Parse error 
+			ErrorText->Logf(ELogVerbosity::Error, TEXT("FAssetBundleEntry::ImportTextItem - Ill-formed asset bundle entry, unexpected: '%s'. When parsing import text '%s'."), BufferIt, Buffer);
+			return false;
+		}
+	}
+
+	BundleName = LocalBundleName;
+	AssetPaths.Reset(LocalObjectPaths.Num());
+	for (const FSoftObjectPath& Path : LocalObjectPaths)
+	{
+		AssetPaths.Add(Path.GetAssetPath());
 	}
 
-	return false;
+PRAGMA_DISABLE_DEPRECATION_WARNINGS
+	BundleAssets.Reset(AssetPaths.Num());
+	for (const FTopLevelAssetPath& Path : AssetPaths)
+	{
+		BundleAssets.Add(FSoftObjectPath(Path, {}));
+	}
+PRAGMA_ENABLE_DEPRECATION_WARNINGS
+
+	// Success, update how much we read for the caller
+	Buffer = BufferIt;
+	return true;
 }
 
 FAssetBundleEntry* FAssetBundleData::FindEntry(FName SearchName)
@@ -119,6 +244,65 @@ FAssetBundleEntry* FAssetBundleData::FindEntry(FName SearchName)
 	return nullptr;
 }
 
+
+void FAssetBundleData::AddBundleAsset(FName BundleName, const FTopLevelAssetPath& AssetPath)
+{
+	if (!AssetPath.IsValid())
+	{
+		return;
+	}
+
+	FAssetBundleEntry* FoundEntry = FindEntry(BundleName);
+
+	if (!FoundEntry)
+	{
+		FoundEntry = new(Bundles) FAssetBundleEntry(BundleName);
+	}
+
+PRAGMA_DISABLE_DEPRECATION_WARNINGS
+	FoundEntry->BundleAssets.AddUnique(FSoftObjectPath(AssetPath, {}));
+PRAGMA_ENABLE_DEPRECATION_WARNINGS
+	FoundEntry->AssetPaths.AddUnique(AssetPath);
+}
+
+void FAssetBundleData::AddBundleAssets(FName BundleName, const TArray<FTopLevelAssetPath>& AssetPaths)
+{
+	FAssetBundleEntry* FoundEntry = FindEntry(BundleName);
+
+	for (const FTopLevelAssetPath& Path : AssetPaths)
+	{
+		if (Path.IsValid())
+		{
+			// Only create if required
+			if (!FoundEntry)
+			{
+				FoundEntry = new(Bundles) FAssetBundleEntry(BundleName);
+			}
+
+PRAGMA_DISABLE_DEPRECATION_WARNINGS
+			FoundEntry->BundleAssets.AddUnique(FSoftObjectPath(Path, {}));
+PRAGMA_ENABLE_DEPRECATION_WARNINGS
+			FoundEntry->AssetPaths.AddUnique(Path);
+		}
+	}
+}
+
+void FAssetBundleData::SetBundleAssets(FName BundleName, TArray<FTopLevelAssetPath>&& AssetPaths)
+{
+	FAssetBundleEntry* FoundEntry = FindEntry(BundleName);
+
+	if (!FoundEntry)
+	{
+		FoundEntry = new(Bundles) FAssetBundleEntry(BundleName);
+	}
+
+	FoundEntry->AssetPaths = MoveTemp(AssetPaths);
+PRAGMA_DISABLE_DEPRECATION_WARNINGS
+	FoundEntry->BundleAssets.Reset();
+	Algo::Transform(FoundEntry->AssetPaths, FoundEntry->BundleAssets, [](FTopLevelAssetPath Path) { return FSoftObjectPath(Path, {}); });
+PRAGMA_ENABLE_DEPRECATION_WARNINGS
+}
+
 void FAssetBundleData::AddBundleAsset(FName BundleName, const FSoftObjectPath& AssetPath)
 {
 	if (!AssetPath.IsValid())
@@ -133,7 +317,10 @@ void FAssetBundleData::AddBundleAsset(FName BundleName, const FSoftObjectPath& A
 		FoundEntry = new(Bundles) FAssetBundleEntry(BundleName);
 	}
 
+PRAGMA_DISABLE_DEPRECATION_WARNINGS
 	FoundEntry->BundleAssets.AddUnique(AssetPath);
+PRAGMA_ENABLE_DEPRECATION_WARNINGS
+	FoundEntry->AssetPaths.AddUnique(AssetPath.GetAssetPath());
 }
 
 void FAssetBundleData::AddBundleAssets(FName BundleName, const TArray<FSoftObjectPath>& AssetPaths)
@@ -150,7 +337,10 @@ void FAssetBundleData::AddBundleAssets(FName BundleName, const TArray<FSoftObjec
 				FoundEntry = new(Bundles) FAssetBundleEntry(BundleName);
 			}
 
+PRAGMA_DISABLE_DEPRECATION_WARNINGS
 			FoundEntry->BundleAssets.AddUnique(Path);
+PRAGMA_ENABLE_DEPRECATION_WARNINGS
+			FoundEntry->AssetPaths.AddUnique(Path.GetAssetPath());
 		}
 	}
 }
@@ -164,7 +354,79 @@ void FAssetBundleData::SetBundleAssets(FName BundleName, TArray<FSoftObjectPath>
 		FoundEntry = new(Bundles) FAssetBundleEntry(BundleName);
 	}
 
-	FoundEntry->BundleAssets = AssetPaths;
+	FoundEntry->AssetPaths.Reset(AssetPaths.Num());
+	for( const FSoftObjectPath& Path : AssetPaths )
+	{	
+		FoundEntry->AssetPaths.Add(Path.GetAssetPath());
+	}
+
+PRAGMA_DISABLE_DEPRECATION_WARNINGS
+	FoundEntry->BundleAssets.Reset();
+	Algo::Transform(FoundEntry->AssetPaths, FoundEntry->BundleAssets, [](FTopLevelAssetPath Path) { return FSoftObjectPath(Path, {}); });
+PRAGMA_ENABLE_DEPRECATION_WARNINGS
+}
+
+void FAssetBundleData::AddBundleAssetTruncated(FName BundleName, const FSoftObjectPath& AssetPath)
+{
+	if (!AssetPath.IsValid())
+	{
+		return;
+	}
+
+	FAssetBundleEntry* FoundEntry = FindEntry(BundleName);
+
+	if (!FoundEntry)
+	{
+		FoundEntry = new(Bundles) FAssetBundleEntry(BundleName);
+	}
+
+PRAGMA_DISABLE_DEPRECATION_WARNINGS
+	FoundEntry->BundleAssets.AddUnique(AssetPath.GetWithoutSubPath());
+PRAGMA_ENABLE_DEPRECATION_WARNINGS
+	FoundEntry->AssetPaths.AddUnique(AssetPath.GetAssetPath());
+}
+
+void FAssetBundleData::AddBundleAssetsTruncated(FName BundleName, const TArray<FSoftObjectPath>& AssetPaths)
+{
+	FAssetBundleEntry* FoundEntry = FindEntry(BundleName);
+
+	for (const FSoftObjectPath& Path : AssetPaths)
+	{
+		if (Path.IsValid())
+		{
+			// Only create if required
+			if (!FoundEntry)
+			{
+				FoundEntry = new(Bundles) FAssetBundleEntry(BundleName);
+			}
+
+PRAGMA_DISABLE_DEPRECATION_WARNINGS
+			FoundEntry->BundleAssets.AddUnique(Path.GetWithoutSubPath());
+PRAGMA_ENABLE_DEPRECATION_WARNINGS
+			FoundEntry->AssetPaths.AddUnique(Path.GetAssetPath());
+		}
+	}
+}
+
+void FAssetBundleData::SetBundleAssetsTruncated(FName BundleName, const TArray<FSoftObjectPath>& AssetPaths)
+{
+	FAssetBundleEntry* FoundEntry = FindEntry(BundleName);
+
+	if (!FoundEntry)
+	{
+		FoundEntry = new(Bundles) FAssetBundleEntry(BundleName);
+	}
+
+	FoundEntry->AssetPaths.Reset(AssetPaths.Num());
+	for (const FSoftObjectPath& Path : AssetPaths)
+	{	
+		FoundEntry->AssetPaths.Add(Path.GetAssetPath());
+	}
+
+PRAGMA_DISABLE_DEPRECATION_WARNINGS
+	FoundEntry->BundleAssets.Reset();
+	Algo::Transform(FoundEntry->AssetPaths, FoundEntry->BundleAssets, [](FTopLevelAssetPath Path) { return FSoftObjectPath(Path); });
+PRAGMA_ENABLE_DEPRECATION_WARNINGS
 }
 
 void FAssetBundleData::Reset()
@@ -261,10 +523,174 @@ FString FAssetBundleData::ToDebugString() const
 
 		Result.Appendf(TEXT("%s -> (%s)"),
 			*Data.BundleName.ToString(),
-			*FString::JoinBy(Data.BundleAssets, TEXT(", "), [](const FSoftObjectPath& LoadBundle) { return LoadBundle.ToString(); }));
+			*FString::JoinBy(Data.AssetPaths, TEXT(", "), [](const FTopLevelAssetPath& LoadBundle) { return LoadBundle.ToString(); })
+		);
 
 		bFirstLine = false;
 	}
 
 	return Result.ToString();
 }
+
+
+#if WITH_DEV_AUTOMATION_TESTS 
+
+#include "Misc/AutomationTest.h"
+
+// Combine import/export tests
+
+IMPLEMENT_SIMPLE_AUTOMATION_TEST(FAssetBundleEntryImportExportTextTest, "Engine.AssetRegistry.AssetBundleEntry.ImportExportText", EAutomationTestFlags::ApplicationContextMask | EAutomationTestFlags::EngineFilter);
+bool FAssetBundleEntryImportExportTextTest::RunTest(const FString& Parameters)
+{
+	FAssetBundleEntry DefaultEntry;
+
+	FAssetBundleEntry EmptyEntry;
+	EmptyEntry.BundleName = "Empty";
+	FString EmptyEntryString = TEXT("(BundleName=\"Empty\",AssetPaths=)");
+
+	FAssetBundleEntry SingleAssetEntry;
+	SingleAssetEntry.BundleName = "Single";
+	SingleAssetEntry.AssetPaths.Add(FTopLevelAssetPath(TEXT("/Game/Characters/Steve.Steve")));
+	FString SingleAssetEntryString = TEXT("(BundleName=\"Single\",AssetPaths=(\"/Game/Characters/Steve.Steve\"))");
+
+	FAssetBundleEntry MultiAssetEntry;
+	MultiAssetEntry.BundleName = "Multi";
+	MultiAssetEntry.AssetPaths.Add(FTopLevelAssetPath(TEXT("/Game/Meshes/ConcreteWall.ConcreteWall")));
+	MultiAssetEntry.AssetPaths.Add(FTopLevelAssetPath(TEXT("/Game/Meshes/GarbageCan.GarbageCan")));
+	MultiAssetEntry.AssetPaths.Add(FTopLevelAssetPath(TEXT("/Game/Blueprints/Widget.Widget_C")));
+	FString MultiAssetEntryString = TEXT("(BundleName=\"Multi\",AssetPaths=(\"/Game/Meshes/ConcreteWall.ConcreteWall\",\"/Game/Meshes/GarbageCan.GarbageCan\",\"/Game/Blueprints/Widget.Widget_C\"))");
+
+	FString EmptyExported;
+	TestTrue(TEXT("Empty asset bundle exports to text"), EmptyEntry.ExportTextItem(EmptyExported, FAssetBundleEntry{}, nullptr, PPF_Delimited, nullptr));
+	TestEqual(TEXT("Empty asset bundle exports to correct string"), EmptyExported, EmptyEntryString);
+
+	FString SingleExported;
+	TestTrue(TEXT("Single-asset asset bundle exports to text"), SingleAssetEntry.ExportTextItem(SingleExported, FAssetBundleEntry{}, nullptr, PPF_Delimited, nullptr));
+	TestEqual(TEXT("Single-asset asset bundle exports to correct string"), SingleExported, SingleAssetEntryString);
+
+	FString MultiExported;
+	TestTrue(TEXT("Multi-asset asset bundle exports to text"), MultiAssetEntry.ExportTextItem(MultiExported, FAssetBundleEntry{}, nullptr, PPF_Delimited, nullptr));
+	TestEqual(TEXT("Multi-asset asset bundle exports to correct string"), MultiExported, MultiAssetEntryString);
+
+	const TCHAR* EmptyBuffer = *EmptyEntryString;
+	FAssetBundleEntry EmptyImported;
+	TestTrue(TEXT("Empty asset bundle imports from text"), EmptyImported.ImportTextItem(EmptyBuffer, PPF_Delimited, nullptr, GLog->Get()));
+	TestEqual(TEXT("Empty asset bundle imports correctly"), EmptyImported, EmptyEntry);
+
+	const TCHAR* SingleBuffer = *SingleAssetEntryString;
+	FAssetBundleEntry SingleImported;
+	TestTrue(TEXT("Single-asset asset bundle imports from text"), SingleImported.ImportTextItem(SingleBuffer, PPF_Delimited, nullptr, GLog->Get()));
+	TestEqual(TEXT("Single-asset asset bundle imports correctly"), SingleImported, SingleAssetEntry);
+
+	const TCHAR* MultiBuffer = *MultiAssetEntryString;
+	FAssetBundleEntry MultiImported;
+	TestTrue(TEXT("Multi-asset asset bundle imports from text"), MultiImported.ImportTextItem(MultiBuffer, PPF_Delimited, nullptr, GLog->Get()));
+	TestEqual(TEXT("Multi-asset asset bundle imports correctly"), MultiImported, MultiAssetEntry);
+
+	AddExpectedError(TEXT("FAssetBundleEntry::ImportTextItem - Asset bundle entries should all be top level objects"));
+	FString SubobjectEntryString = TEXT("(BundleName=\"Subobject\",AssetPaths=(\"/Game/Characters/Steve.Steve:Hat\"))");
+	const TCHAR* SubobjectBuffer = *SubobjectEntryString;
+	FAssetBundleEntry SubobjectImported;
+	TestFalse(TEXT("Subobject bunld does not import from text"), SubobjectImported.ImportTextItem(SubobjectBuffer, PPF_Delimited, nullptr, GLog->Get()));
+	TestEqual(TEXT("Subobject bundle does not import from text"), SubobjectImported, DefaultEntry);
+
+	AddExpectedError(TEXT("FAssetBundleEntry::ImportTextItem - Unterminated BundleName string"));
+	FString UnterminatedBundleNameString = TEXT("(BundleName=\"TestUnterminated,BundleAssets=))");
+	const TCHAR* UnterminatedBundleNameBuffer = *UnterminatedBundleNameString;
+	FAssetBundleEntry UnterminatedEntry;
+	TestFalse(TEXT("Unterminated asset bundle name fails import"), UnterminatedEntry.ImportTextItem(UnterminatedBundleNameBuffer, PPF_Delimited, nullptr, GLog->Get()));
+	TestEqual(TEXT("Unterminated asset bundle name fails import"), UnterminatedEntry, DefaultEntry);
+
+	AddExpectedError(TEXT("AssetBundleEntry::ImportTextItem - Ill-formed asset bundle entry"));
+	FString MismatchedQuotesString = TEXT("(BundleName=\"TestUnterminated,BundleAssets=(\"/Game/Characters/Steve.Steve\")))");
+	const TCHAR* MismatchedQuotesBuffer = *MismatchedQuotesString;
+	FAssetBundleEntry MismatchedQuotesEntry;
+	TestFalse(TEXT("Mismatched quotes asset bundle name fails import"), MismatchedQuotesEntry.ImportTextItem(MismatchedQuotesBuffer, PPF_Delimited, nullptr, GLog->Get()));
+	TestEqual(TEXT("Mismatched quotes asset bundle name fails import"), MismatchedQuotesEntry, DefaultEntry);
+
+	AddExpectedError(TEXT("FAssetBundleEntry::ImportTextItem - Text does not start with expected prefix"));
+	FString UnsupportedOrderString = TEXT("(BundleAssets=(\"/Game/Characters/Steve.Steve\"), BundleName=\"UnsupportedOrder\")");
+	const TCHAR* UnsupportedOrderBuffer = *UnsupportedOrderString;
+	FAssetBundleEntry UnsupportedOrderEntry;
+	TestFalse(TEXT("Unsupported order fails import"), UnsupportedOrderEntry.ImportTextItem(UnsupportedOrderBuffer, PPF_Delimited, nullptr, GLog->Get()));
+	TestEqual(TEXT("Unsupported order fails import"), UnsupportedOrderEntry, DefaultEntry);
+
+	AddExpectedError(TEXT("FAssetBundleEntry::ImportTextItem - Encountered key BundleAssets twice"));
+	FString DuplicateKeysString = TEXT("(BundleName=\"DuplicateKeys\",BundleAssets=(\"/Game/Characters/Steve.Steve\"),BundleAssets=(\"/Game/Characters/Will.Will\"))");
+	const TCHAR* DuplicateKeysBuffer = *DuplicateKeysString;
+	FAssetBundleEntry DuplicateKeysEntry;
+	TestFalse(TEXT("Duplicate keys fails import"), DuplicateKeysEntry.ImportTextItem(DuplicateKeysBuffer, PPF_Delimited, nullptr, GLog->Get()));
+	TestEqual(TEXT("Duplicate keys fails import"), DuplicateKeysEntry, DefaultEntry);
+
+	return true;
+}
+
+
+IMPLEMENT_SIMPLE_AUTOMATION_TEST(FLegacyAssetBundleEntryTest, "Engine.AssetRegistry.AssetBundleEntry.LegacyAssetBundleEntry", EAutomationTestFlags::ApplicationContextMask | EAutomationTestFlags::EngineFilter);
+// Test ImportText for old asset bundles with BundleAssets field
+bool FLegacyAssetBundleEntryTest::RunTest(const FString& Parameters)
+{
+	// Test subobject path truncation
+	FAssetBundleEntry DefaultEntry;
+
+	FAssetBundleEntry EmptyEntry;
+	EmptyEntry.BundleName = "Empty";
+	FString EmptyEntryString = TEXT("(BundleName=\"Empty\",AssetPaths=())");
+	FString EmptyEntryString2 = TEXT("(BundleName=\"Empty\",AssetPaths=)");
+
+	FAssetBundleEntry SingleAssetEntry;
+	SingleAssetEntry.BundleName = "Single";
+	SingleAssetEntry.AssetPaths.Add(FTopLevelAssetPath(TEXT("/Game/Characters/Steve.Steve")));
+	FString SingleAssetEntryString = TEXT("(BundleName=\"Single\",BundleAssets=(\"/Game/Characters/Steve.Steve\"))");
+
+	FAssetBundleEntry MultiAssetEntry;
+	MultiAssetEntry.BundleName = "Multi";
+	MultiAssetEntry.AssetPaths.Add(FTopLevelAssetPath(TEXT("/Game/Meshes/ConcreteWall.ConcreteWall")));
+	MultiAssetEntry.AssetPaths.Add(FTopLevelAssetPath(TEXT("/Game/Meshes/GarbageCan.GarbageCan")));
+	MultiAssetEntry.AssetPaths.Add(FTopLevelAssetPath(TEXT("/Game/Blueprints/Widget.Widget_C")));
+	FString MultiAssetEntryString = TEXT("(BundleName=\"Multi\",BundleAssets=(\"/Game/Meshes/ConcreteWall.ConcreteWall\",\"/Game/Meshes/GarbageCan.GarbageCan\",\"/Game/Blueprints/Widget.Widget_C\"))");
+
+	FAssetBundleEntry SubobjectEntry;
+	SubobjectEntry.BundleName = "Subobject";
+	SubobjectEntry.AssetPaths.Add(FTopLevelAssetPath(TEXT("/Game/Characters/Steve.Steve")));
+	FString SubobjectEntryString = TEXT("(BundleName=\"Subobject\",BundleAssets=(\"/Game/Characters/Steve.Steve:Hat\"))");
+
+	const TCHAR* EmptyBuffer = *EmptyEntryString;
+	FAssetBundleEntry EmptyImported;
+	TestTrue(TEXT("Empty asset bundle imports from text"), EmptyImported.ImportTextItem(EmptyBuffer, PPF_Delimited, nullptr, GLog->Get()));
+	TestEqual(TEXT("Empty asset bundle imports correctly"), EmptyImported, EmptyEntry);
+
+	const TCHAR* EmptyBuffer2 = *EmptyEntryString2;
+	FAssetBundleEntry EmptyImported2;
+	TestTrue(TEXT("Empty asset bundle without inner parens imports from text"), EmptyImported2.ImportTextItem(EmptyBuffer2, PPF_Delimited, nullptr, GLog->Get()));
+	TestEqual(TEXT("Empty asset bundle without inner parens imports correctly"), EmptyImported2, EmptyEntry);
+
+	const TCHAR* SingleBuffer = *SingleAssetEntryString;
+	FAssetBundleEntry SingleImported;
+	TestTrue(TEXT("Single-asset asset bundle imports from text"), SingleImported.ImportTextItem(SingleBuffer, PPF_Delimited, nullptr, GLog->Get()));
+	TestEqual(TEXT("Single-asset asset bundle imports correctly"), SingleImported, SingleAssetEntry);
+
+	const TCHAR* MultiBuffer = *MultiAssetEntryString;
+	FAssetBundleEntry MultiImported;
+	TestTrue(TEXT("Multi-asset asset bundle imports from text"), MultiImported.ImportTextItem(MultiBuffer, PPF_Delimited, nullptr, GLog->Get()));
+	TestEqual(TEXT("Multi-asset asset bundle imports correctly"), MultiImported, MultiAssetEntry);
+
+	const TCHAR* SubobjectBuffer = *SubobjectEntryString;
+	FAssetBundleEntry SubobjectImported;
+	TestTrue(TEXT("Subobject asset bundle imports from text"), SubobjectImported.ImportTextItem(SubobjectBuffer, PPF_Delimited, nullptr, GLog->Get()));
+	TestEqual(TEXT("Subobject asset bundle imports correctly"), SubobjectImported, SubobjectEntry);
+
+	FAssetBundleEntry MixedFormatEntry;
+	MixedFormatEntry.BundleName = "MixedFormat";
+	MixedFormatEntry.AssetPaths.Add(FTopLevelAssetPath(TEXT("/Game/Meshes/ConcreteWall.ConcreteWall")));
+	MixedFormatEntry.AssetPaths.Add(FTopLevelAssetPath(TEXT("/Game/Meshes/GarbageCan.GarbageCan")));
+	MixedFormatEntry.AssetPaths.Add(FTopLevelAssetPath(TEXT("/Game/Blueprints/Widget.Widget_C")));
+	FString MixedFormatString = TEXT("(BundleName=\"MixedFormat\",AssetPaths=(\"/Game/Meshes/ConcreteWall.ConcreteWall\"),BundleAssets=(\"/Game/Meshes/GarbageCan.GarbageCan\",\"/Game/Blueprints/Widget.Widget_C\"))");
+	const TCHAR* MixedFormatBuffer = *MixedFormatString;
+	FAssetBundleEntry MixedFormatImported;
+	TestTrue(TEXT("Mixed format asset bundle imports from text"), MixedFormatImported.ImportTextItem(MixedFormatBuffer, PPF_Delimited, nullptr, GLog->Get()));
+	TestEqual(TEXT("Mixed format asset bundle imports correctly"), MixedFormatImported, MixedFormatEntry);
+
+	return true;
+}
+#endif // WITH_DEV_AUTOMATION_TESTS
diff --git a/Engine/Source/Runtime/CoreUObject/Private/AssetRegistry/AssetData.cpp b/Engine/Source/Runtime/CoreUObject/Private/AssetRegistry/AssetData.cpp
--- a/Engine/Source/Runtime/CoreUObject/Private/AssetRegistry/AssetData.cpp
+++ b/Engine/Source/Runtime/CoreUObject/Private/AssetRegistry/AssetData.cpp
@@ -7,6 +7,7 @@
 #include "AssetRegistry/ARFilter.h"
 #include "Containers/Set.h"
 #include "HAL/CriticalSection.h"
+#include "Misc/AsciiSet.h"
 #include "Misc/PathViews.h"
 #include "Misc/ScopeRWLock.h"
 #include "Serialization/CompactBinary.h"
@@ -25,6 +26,49 @@ UE_IMPLEMENT_STRUCT("/Script/CoreUObject", AssetData);
 const FGuid FAssetRegistryVersion::GUID(0x717F9EE7, 0xE9B0493A, 0x88B39132, 0x1B388107);
 FCustomVersionRegistration GRegisterAssetRegistryVersion(FAssetRegistryVersion::GUID, FAssetRegistryVersion::LatestVersion, TEXT("AssetRegistry"));
 
+namespace UE::AssetRegistry::Private
+{
+	FAssetPathParts SplitIntoOuterPathAndAssetName(FStringView InObjectPath)
+	{
+		constexpr FAsciiSet Delimiters(SUBOBJECT_DELIMITER ".");
+		FStringView OuterPathPlusDelimiter = FAsciiSet::TrimSuffixWithout(InObjectPath, Delimiters);
+
+		return FAssetPathParts{
+			OuterPathPlusDelimiter.LeftChop(1),
+			InObjectPath.RightChop(OuterPathPlusDelimiter.Len())
+		};
+	}
+
+	void ConcatenateOuterPathAndObjectName(FStringBuilderBase& Builder, FName OuterPath, FName ObjectName)
+	{
+		// We assume that OuterPath was correctly constructed with a subobject delimiter if it needed one
+		// So we only need to decide if OuterPath and ObjectName should be separated by '.' or ':'
+		// See UObjectBaseUtility::GetPathName.
+		// We don't have access to type information here so the best we can do is rely on the fact that we don't have
+		// UPackage anywhere but top-level and ensure that the second delimiter in any path string is a ':'
+
+		int32 StartingLen = Builder.Len();
+		Builder << OuterPath;
+
+		TCHAR Delimiter = '.';
+
+		FStringView OuterPathView = Builder.ToView().Mid(StartingLen, Builder.Len() - StartingLen);
+		int32 DotIndex = 0;
+		if (OuterPathView.FindChar('.', DotIndex))
+		{
+			// Contains a dot delimiter, so we may need to use SUBOBJECT_DELIMITER_CHAR
+			int32 SubobjectDelimIndex = 0;
+			if (!OuterPathView.RightChop(DotIndex + 1).FindChar(SUBOBJECT_DELIMITER_CHAR, SubobjectDelimIndex))
+			{
+				// No delimiter, OuterPath must be of the form 'A.B' so we need SUBOBJECT_DELIMITER_CHAR to produce full path 'A.B:C'
+				Delimiter = SUBOBJECT_DELIMITER_CHAR;
+			}
+		}
+
+		Builder << Delimiter << ObjectName;
+	}
+}
+
 const FName GAssetBundleDataName("AssetBundleData");
 
 namespace UE { namespace AssetData { namespace Private {
@@ -78,26 +122,36 @@ FAssetData::FAssetData(FName InPackageName, FName InPackagePath, FName InAssetNa
 	FNameBuilder ObjectPathStr(PackageName);
 	ObjectPathStr << TEXT('.');
 	AssetName.AppendString(ObjectPathStr);
+PRAGMA_DISABLE_DEPRECATION_WARNINGS
 	ObjectPath = FName(FStringView(ObjectPathStr));
+PRAGMA_ENABLE_DEPRECATION_WARNINGS
 }
 
 FAssetData::FAssetData(const FString& InLongPackageName, const FString& InObjectPath, FTopLevelAssetPath InAssetClassPathName, FAssetDataTagMap InTags, TArrayView<const int32> InChunkIDs, uint32 InPackageFlags)
-	: ObjectPath(*InObjectPath)
-	, PackageName(*InLongPackageName)
+	: PackageName(*InLongPackageName)
 	, AssetClassPath(InAssetClassPathName)
 	, PackageFlags(InPackageFlags)
 	, ChunkIDs(MoveTemp(InChunkIDs))
 {
+	using namespace UE::AssetRegistry::Private;
+PRAGMA_DISABLE_DEPRECATION_WARNINGS
+	ObjectPath = *InObjectPath;
+PRAGMA_ENABLE_DEPRECATION_WARNINGS
+
 	SetTagsAndAssetBundles(MoveTemp(InTags));
 
 	PackagePath = FName(*FPackageName::GetLongPackagePath(InLongPackageName));
 
-	// Find the object name from the path, FPackageName::ObjectPathToObjectName(InObjectPath)) doesn't provide what we want here
-	int32 CharPos = InObjectPath.FindLastCharByPredicate([](TCHAR Char)
+
+	const FAssetPathParts Parts = SplitIntoOuterPathAndAssetName(InObjectPath);
+	AssetName = FName(Parts.InnermostName);
+
+#if WITH_EDITORONLY_DATA
+	if (!Parts.OuterPath.Equals(InLongPackageName, ESearchCase::IgnoreCase))
 	{
-		return Char == ':' || Char == '.';
-	});
-	AssetName = FName(*InObjectPath.Mid(CharPos + 1));
+		OptionalOuterPath = FName(Parts.OuterPath);
+	}
+#endif
 }
 
 FAssetData::FAssetData(const UObject* InAsset, FAssetData::ECreationFlags InCreationFlags)
@@ -120,7 +174,15 @@ FAssetData::FAssetData(const UObject* InAsset, FAssetData::ECreationFlags InCrea
 		PackagePath = FName(*FPackageName::GetLongPackagePath(Package->GetName()));
 		AssetName = InAsset->GetFName();
 		AssetClassPath = InAsset->GetClass()->GetPathName();
+PRAGMA_DISABLE_DEPRECATION_WARNINGS
 		ObjectPath = FName(*InAsset->GetPathName());
+PRAGMA_ENABLE_DEPRECATION_WARNINGS
+#if WITH_EDITORONLY_DATA
+		if (InAsset->GetOuter() != Package)
+		{
+			OptionalOuterPath = *InAsset->GetOuter()->GetPathName();
+		}
+#endif
 
 		if (!EnumHasAnyFlags(InCreationFlags, FAssetData::ECreationFlags::SkipAssetRegistryTagsGathering))
 		{
@@ -132,6 +194,20 @@ FAssetData::FAssetData(const UObject* InAsset, FAssetData::ECreationFlags InCrea
 	}
 }
 
+FSoftObjectPath FAssetData::GetSoftObjectPath() const
+{
+	if (IsTopLevelAsset())
+	{
+		return FSoftObjectPath(PackageName, AssetName, FString());
+	}
+	else
+	{
+		TStringBuilder<FName::StringBufferSize> Builder;
+		AppendObjectPath(Builder);
+		return FSoftObjectPath(Builder.ToView());
+	}
+}
+
 bool FAssetData::IsUAsset(UObject* InAsset)
 {
 	if (InAsset == nullptr)
@@ -152,9 +228,23 @@ bool FAssetData::IsUAsset(UObject* InAsset)
 
 bool FAssetData::IsTopLevelAsset() const
 {
+#if WITH_EDITORONLY_DATA
+	if (OptionalOuterPath.IsNone())
+	{
+		// If no outer path, then path is PackageName.AssetName so we must be top level
+		return true;
+	}
+
+	TStringBuilder<FName::StringBufferSize> Builder;
+	AppendObjectPath(Builder);
+
 	int32 SubObjectIndex;
-	FStringView(WriteToString<256>(ObjectPath)).FindChar(SUBOBJECT_DELIMITER_CHAR, SubObjectIndex);
+	FStringView(Builder).FindChar(SUBOBJECT_DELIMITER_CHAR, SubObjectIndex);
 	return SubObjectIndex == INDEX_NONE;
+#else
+	// Non-top-level assets only appear in the editor
+	return true;
+#endif
 }
 
 bool FAssetData::IsTopLevelAsset(UObject* Object)
@@ -206,10 +296,19 @@ FPrimaryAssetId FAssetData::GetPrimaryAssetId() const
 	return FPrimaryAssetId();
 }
 
-void FAssetData::SerializeForCacheInternal(FArchive& Ar, FAssetRegistryVersion::Type Version, void (*SerializeTagsAndBundles)(FArchive& , FAssetData&))
+void FAssetData::SerializeForCacheInternal(FArchive& Ar, FAssetRegistryVersion::Type Version, void (*SerializeTagsAndBundles)(FArchive& , FAssetData&, FAssetRegistryVersion::Type))
 {
 	// Serialize out the asset info
-	Ar << ObjectPath;
+	// Only needed if Version < FAssetRegistryVersion::RemoveAssetPathFNames but we need to reference it later to 
+	// rebuild OptionalOuterPath for assets which are stored in a different package to their outer (e.g. external actors)
+	FName OldObjectPath; 
+	if (Version < FAssetRegistryVersion::RemoveAssetPathFNames)
+	{
+PRAGMA_DISABLE_DEPRECATION_WARNINGS
+		Ar << OldObjectPath;
+PRAGMA_ENABLE_DEPRECATION_WARNINGS
+	}
+
 	Ar << PackagePath;
 
 	// Serialize the asset class.
@@ -225,11 +324,45 @@ PRAGMA_DISABLE_DEPRECATION_WARNINGS
 PRAGMA_ENABLE_DEPRECATION_WARNINGS
 	}
 
-	// These are derived from ObjectPath, we manually serialize them because they get pooled
 	Ar << PackageName;
 	Ar << AssetName;
 
-	SerializeTagsAndBundles(Ar , *this);
+#if WITH_EDITORONLY_DATA
+	if (Version >= FAssetRegistryVersion::RemoveAssetPathFNames)
+	{
+		if (!Ar.IsFilterEditorOnly())
+		{
+			Ar << OptionalOuterPath;
+		}
+		else if (Ar.IsLoading())
+		{
+			OptionalOuterPath = NAME_None;
+		}
+	}
+	else
+	{
+		check(Ar.IsLoading());
+		check(!OldObjectPath.IsNone());
+
+		using namespace UE::AssetRegistry::Private;
+
+		OptionalOuterPath = NAME_None;
+		TStringBuilder<FName::StringBufferSize> Builder;
+		Builder << PackageName << '.' << AssetName;
+		if (OldObjectPath != *Builder)
+		{
+			Builder.Reset();
+			Builder << OldObjectPath;
+			FAssetPathParts Parts = SplitIntoOuterPathAndAssetName(Builder.ToView());
+			if (PackageName.ToString() != Parts.OuterPath)
+			{
+				OptionalOuterPath = FName(Parts.OuterPath);
+			}
+		}
+	}
+#endif // WITH_EDITORONLY_DATA
+
+	SerializeTagsAndBundles(Ar, *this, Version);
 
 	if (Ar.IsSaving() && ChunkIDs.Num() > 1)
 	{
@@ -242,6 +375,13 @@ PRAGMA_ENABLE_DEPRECATION_WARNINGS
 		Ar << ChunkIDs;
 	}
 	Ar << PackageFlags;
+
+	// Rebuild deprecated ObjectPath field 
+	TStringBuilder<FName::StringBufferSize> Builder;
+	AppendObjectPath(Builder);
+PRAGMA_DISABLE_DEPRECATION_WARNINGS
+	ObjectPath = *Builder;
+PRAGMA_ENABLE_DEPRECATION_WARNINGS
 }
 
 void FAssetData::NetworkWrite(FCbWriter& Writer, bool bWritePackageName) const
@@ -250,9 +390,11 @@ void FAssetData::NetworkWrite(FCbWriter& Writer, bool bWritePackageName) const
 	Writer.BeginObject();
 	if (bWritePackageName)
 	{
-		Writer << "O" << ObjectPath;
 		Writer << "P" << PackagePath;
 		Writer << "Q" << PackageName;
+#if WITH_EDITORONLY_DATA
+		Writer << "OO" << OptionalOuterPath;
+#endif
 	}
 	Writer << "N" << AssetName;
 	Writer << "C" << AssetClassPath.ToString();
@@ -295,21 +437,21 @@ bool FAssetData::TryNetworkRead(FCbFieldView Field, bool bReadPackageName, FName
 	bOk &= bHasAssetName;
 	if (bReadPackageName)
 	{
-		bOk = LoadFromCompactBinary(Object["O"], ObjectPath) & bOk;
 		bOk = LoadFromCompactBinary(Object["P"], PackagePath) & bOk;
 		bOk = LoadFromCompactBinary(Object["Q"], PackageName) & bOk;
+#if WITH_EDITORONLY_DATA
+		bOk = LoadFromCompactBinary(Object["OO"], OptionalOuterPath) & bOk;
+#endif
 	}
 	else
 	{
 		if (bHasAssetName)
 		{
 			WriteToString<256> PackageNameStr(InPackageName);
-			ObjectPath = FName(WriteToString<256>(PackageNameStr.ToView(), TEXT("."), AssetName));
 			PackagePath = FName(FPathViews::GetPath(PackageNameStr.ToView()));
 		}
 		else
 		{
-			ObjectPath = NAME_None;
 			PackagePath = NAME_None;
 		}
 		PackageName = InPackageName;
@@ -346,14 +488,22 @@ bool FAssetData::TryNetworkRead(FCbFieldView Field, bool bReadPackageName, FName
 	}
 
 	LoadFromCompactBinary(Object["I"], ChunkIDs); // Ok if it does not exist
+
+	// Rebuild deprecated ObjectPath field 
+	TStringBuilder<FName::StringBufferSize> Builder;
+	AppendObjectPath(Builder);
+PRAGMA_DISABLE_DEPRECATION_WARNINGS
+	ObjectPath = *Builder;
+PRAGMA_ENABLE_DEPRECATION_WARNINGS
+
 	return bOk;
 }
 
-void FAssetData::SerializeForCacheWithTagsAndBundles(FArchive& Ar, void (*SerializeTagsAndBundles)(FArchive&, FAssetData&))
+void FAssetData::SerializeForCacheWithTagsAndBundles(FArchive& Ar, void (*SerializeTagsAndBundles)(FArchive&, FAssetData&, FAssetRegistryVersion::Type))
 {
 	SerializeForCacheInternal(Ar, FAssetRegistryVersion::LatestVersion, SerializeTagsAndBundles);
 }
-void FAssetData::SerializeForCacheOldVersionWithTagsAndBundles(FArchive& Ar, FAssetRegistryVersion::Type Version, void (*SerializeTagsAndBundles)(FArchive&, FAssetData&))
+void FAssetData::SerializeForCacheOldVersionWithTagsAndBundles(FArchive& Ar, FAssetRegistryVersion::Type Version, void (*SerializeTagsAndBundles)(FArchive&, FAssetData&, FAssetRegistryVersion::Type Version))
 {
 	SerializeForCacheInternal(Ar, Version, SerializeTagsAndBundles);
 }
diff --git a/Engine/Source/Runtime/CoreUObject/Private/Misc/RedirectCollector.cpp b/Engine/Source/Runtime/CoreUObject/Private/Misc/RedirectCollector.cpp
--- a/Engine/Source/Runtime/CoreUObject/Private/Misc/RedirectCollector.cpp
+++ b/Engine/Source/Runtime/CoreUObject/Private/Misc/RedirectCollector.cpp
@@ -28,9 +28,9 @@ FAutoConsoleCommand CVarResolveAllSoftObjects(
 	})
 );
 
-void FRedirectCollector::OnSoftObjectPathLoaded(const FSoftObjectPath& InPath, FArchive* InArchive)
+void FRedirectCollector::OnSoftObjectPathLoaded(const FSoftObjectPath& ObjectPath, FArchive* InArchive)
 {
-	if (InPath.IsNull() || !GIsEditor)
+	if (ObjectPath.IsNull() || !GIsEditor)
 	{
 		// No need to track empty strings, or in standalone builds
 		return;
@@ -51,22 +51,22 @@ void FRedirectCollector::OnSoftObjectPathLoaded(const FSoftObjectPath& InPath, F
 	}
 
 	const bool bReferencedByEditorOnlyProperty = (CollectType == ESoftObjectPathCollectType::EditorOnlyCollect);
-	FName AssetPathName = InPath.GetAssetPathName();
+	FTopLevelAssetPath AssetPath = ObjectPath.GetAssetPath();
 
 	FScopeLock ScopeLock(&CriticalSection);
 	if (CollectType != ESoftObjectPathCollectType::NeverCollect)
 	{
 		// Add this reference to the soft object inclusion list for the cook's iterative traversal of the soft dependency graph
-		FSoftObjectPathProperty SoftObjectPathProperty(AssetPathName, PropertyName, bReferencedByEditorOnlyProperty);
+		FSoftObjectPathProperty SoftObjectPathProperty(FSoftObjectPath(AssetPath, {}), PropertyName, bReferencedByEditorOnlyProperty);
 		SoftObjectPathMap.FindOrAdd(PackageName).Add(SoftObjectPathProperty);
 	}
 
 	if (ShouldTrackPackageReferenceTypes())
 	{
 		// Add the referenced package to the potential-exclusion list for the cook's up-front traversal of the soft dependency graph
-		TStringBuilder<256> AssetPathString;
-		AssetPathName.ToString(AssetPathString);
-		FName ReferencedPackageName = FName(FPackageName::ObjectPathToPackageName(AssetPathString));
+		TStringBuilder<FName::StringBufferSize> ObjectPathString;
+		ObjectPath.AppendString(ObjectPathString);
+		FName ReferencedPackageName = FName(FPackageName::ObjectPathToPackageName(ObjectPathString));
 		if (PackageName != ReferencedPackageName)
 		{
 			TMap<FName, ESoftObjectPathCollectType>& PackageReferences = PackageReferenceTypes.FindOrAdd(PackageName);
@@ -81,7 +81,7 @@ void FRedirectCollector::CollectSavedSoftPackageReferences(FName ReferencingPack
 	TArray<FSoftObjectPathProperty, TInlineAllocator<4>> SoftObjectPathArray;
 	Algo::Transform(PackageNames, SoftObjectPathArray, [ReferencingPackage, bEditorOnlyReferences](const FName& PackageName)
 		{
-			return FSoftObjectPathProperty(PackageName, NAME_None, bEditorOnlyReferences);
+			return FSoftObjectPathProperty(FSoftObjectPath(PackageName, NAME_None, FString()), NAME_None, bEditorOnlyReferences);
 		});
 
 	FScopeLock ScopeLock(&CriticalSection);
@@ -92,8 +92,8 @@ void FRedirectCollector::ResolveAllSoftObjectPaths(FName FilterPackage)
 {	
 	auto LoadSoftObjectPathLambda = [this](const FSoftObjectPathProperty& SoftObjectPathProperty, const FName ReferencerPackageName)
 	{
-		const FName& ToLoadFName = SoftObjectPathProperty.GetAssetPathName();
-		const FString ToLoad = ToLoadFName.ToString();
+		FSoftObjectPath ToLoadPath = SoftObjectPathProperty.GetObjectPath();
+		const FString ToLoad = ToLoadPath.ToString();
 
 		if (ToLoad.Len() > 0 )
 		{
@@ -113,11 +113,11 @@ void FRedirectCollector::ResolveAllSoftObjectPaths(FName FilterPackage)
 
 			if (Loaded)
 			{
-				FString Dest = Loaded->GetPathName();
-				UE_LOG(LogRedirectors, Verbose, TEXT("    Resolved to '%s'"), *Dest);
-				if (Dest != ToLoad)
+				FSoftObjectPath Dest(Loaded);
+				UE_LOG(LogRedirectors, Verbose, TEXT("    Resolved to '%s'"), *Dest.ToString());
+				if (Dest.ToString() != ToLoad)
 				{
-					AssetPathRedirectionMap.Add(ToLoadFName, FName(*Dest));
+					ObjectPathRedirectionMap.Add(ToLoadPath, Dest);
 				}
 			}
 			else
@@ -186,8 +186,8 @@ void FRedirectCollector::ProcessSoftObjectPathPackageList(FName FilterPackage, b
 	{
 		if (!SoftObjectPathProperty.GetReferencedByEditorOnlyProperty() || bGetEditorOnly)
 		{
-			const FName& ToLoadFName = SoftObjectPathProperty.GetAssetPathName();
-			FString PackageNameString = FPackageName::ObjectPathToPackageName(ToLoadFName.ToString());
+			FSoftObjectPath ToLoadPath = SoftObjectPathProperty.GetObjectPath();
+			FString PackageNameString = FPackageName::ObjectPathToPackageName(ToLoadPath.ToString());
 			OutReferencedPackages.Add(FName(*PackageNameString));
 		}
 	}
@@ -232,57 +232,78 @@ bool FRedirectCollector::ShouldTrackPackageReferenceTypes()
 	return TrackingReferenceTypesState == ETrackingReferenceTypesState::Enabled;
 }
 
-void FRedirectCollector::AddAssetPathRedirection(FName OriginalPath, FName RedirectedPath)
+void FRedirectCollector::AddAssetPathRedirection(const FSoftObjectPath& OriginalPath, const FSoftObjectPath& RedirectedPath)
 {
 	FScopeLock ScopeLock(&CriticalSection);
 
-	if (!ensureMsgf(OriginalPath != NAME_None, TEXT("Cannot add redirect from Name_None!")))
+	if (!ensureMsgf(!OriginalPath.IsNull(), TEXT("Cannot add redirect from Name_None!")))
 	{
 		return;
 	}
 
-	FName FinalRedirection = GetAssetPathRedirection(RedirectedPath);
+	FSoftObjectPath FinalRedirection = GetAssetPathRedirection(RedirectedPath);
 	if (FinalRedirection == OriginalPath)
 	{
 		// If RedirectedPath points back to OriginalPath, remove that to avoid a circular reference
 		// This can happen when renaming assets in the editor but not actually dropping redirectors because it was new
-		AssetPathRedirectionMap.Remove(RedirectedPath);
+		ObjectPathRedirectionMap.Remove(RedirectedPath);
 	}
 
 	// This replaces an existing mapping, can happen in the editor if things are renamed twice
-	AssetPathRedirectionMap.Add(OriginalPath, RedirectedPath);
+	ObjectPathRedirectionMap.Add(OriginalPath, RedirectedPath);
 }
 
-void FRedirectCollector::RemoveAssetPathRedirection(FName OriginalPath)
+void FRedirectCollector::AddAssetPathRedirection(FName OriginalPath, FName RedirectedPath)
+{
+PRAGMA_DISABLE_DEPRECATION_WARNINGS
+	AddAssetPathRedirection(FSoftObjectPath(OriginalPath), FSoftObjectPath(RedirectedPath));
+PRAGMA_ENABLE_DEPRECATION_WARNINGS
+}
+
+void FRedirectCollector::RemoveAssetPathRedirection(const FSoftObjectPath& OriginalPath)
 {
 	FScopeLock ScopeLock(&CriticalSection);
 
-	FName* Found = AssetPathRedirectionMap.Find(OriginalPath);
+	FSoftObjectPath* Found = ObjectPathRedirectionMap.Find(OriginalPath);
 
 	if (ensureMsgf(Found, TEXT("Cannot remove redirection from %s, it was not registered"), *OriginalPath.ToString()))
 	{
-		AssetPathRedirectionMap.Remove(OriginalPath);
+		ObjectPathRedirectionMap.Remove(OriginalPath);
 	}
 }
 
+void FRedirectCollector::RemoveAssetPathRedirection(FName OriginalPath)
+{
+PRAGMA_DISABLE_DEPRECATION_WARNINGS
+	RemoveAssetPathRedirection(FSoftObjectPath(OriginalPath));
+PRAGMA_ENABLE_DEPRECATION_WARNINGS
+}
+
 FName FRedirectCollector::GetAssetPathRedirection(FName OriginalPath)
+{
+PRAGMA_DISABLE_DEPRECATION_WARNINGS
+	return FName(*GetAssetPathRedirection(FSoftObjectPath(OriginalPath)).ToString());
+PRAGMA_ENABLE_DEPRECATION_WARNINGS
+}
+
+FSoftObjectPath FRedirectCollector::GetAssetPathRedirection(const FSoftObjectPath& OriginalPath)
 {
 	FScopeLock ScopeLock(&CriticalSection);
-	TArray<FName> SeenPaths;
+	TArray<FSoftObjectPath, TInlineAllocator<2>> SeenPaths;
 
 	// We need to follow the redirect chain recursively
-	FName CurrentPath = OriginalPath;
+	FSoftObjectPath CurrentPath = OriginalPath;
 
-	while (CurrentPath != NAME_None)
+	while (!CurrentPath.IsNull())
 	{
 		SeenPaths.Add(CurrentPath);
-		FName NewPath = AssetPathRedirectionMap.FindRef(CurrentPath);
+		FSoftObjectPath NewPath = ObjectPathRedirectionMap.FindRef(CurrentPath);
 
-		if (NewPath != NAME_None)
+		if (!NewPath.IsNull())
 		{
 			if (!ensureMsgf(!SeenPaths.Contains(NewPath), TEXT("Found circular redirect from %s to %s! Returning None instead"), *CurrentPath.ToString(), *NewPath.ToString()))
 			{
-				return NAME_None;
+				return FSoftObjectPath();
 			}
 
 			// Continue trying to follow chain
@@ -299,7 +320,7 @@ FName FRedirectCollector::GetAssetPathRedirection(FName OriginalPath)
 	{
 		return CurrentPath;
 	}
-	return NAME_None;
+	return FSoftObjectPath();
 }
 
 FRedirectCollector GRedirectCollector;
diff --git a/Engine/Source/Runtime/CoreUObject/Private/UObject/LinkerLoad.cpp b/Engine/Source/Runtime/CoreUObject/Private/UObject/LinkerLoad.cpp
--- a/Engine/Source/Runtime/CoreUObject/Private/UObject/LinkerLoad.cpp
+++ b/Engine/Source/Runtime/CoreUObject/Private/UObject/LinkerLoad.cpp
@@ -6256,11 +6256,7 @@ void FLinkerLoad::FixupSoftObjectPathForInstancedPackage(FSoftObjectPath& InOutS
 			FName LongPackageName = InOutSoftObjectPath.GetLongPackageFName();
 			if (FName RemappedPackage = InstancingContext.RemapPackage(LongPackageName); RemappedPackage != LongPackageName)
 			{
-				FNameBuilder TmpSoftObjectPathBuilder;
-				InOutSoftObjectPath.ToString(TmpSoftObjectPathBuilder);
-
-				TmpSoftObjectPathBuilder.ReplaceAt(0, LongPackageName.GetStringLength(), RemappedPackage.ToString());
-				InOutSoftObjectPath.SetPath(FName(TmpSoftObjectPathBuilder));
+				InOutSoftObjectPath = FSoftObjectPath(RemappedPackage, InOutSoftObjectPath.GetAssetFName(), InOutSoftObjectPath.GetSubPathString());
 			}
 		}
 	}
diff --git a/Engine/Source/Runtime/CoreUObject/Private/UObject/SoftObjectPath.cpp b/Engine/Source/Runtime/CoreUObject/Private/UObject/SoftObjectPath.cpp
--- a/Engine/Source/Runtime/CoreUObject/Private/UObject/SoftObjectPath.cpp
+++ b/Engine/Source/Runtime/CoreUObject/Private/UObject/SoftObjectPath.cpp
@@ -10,6 +10,7 @@
 #include "UObject/UObjectThreadContext.h"
 #include "UObject/CoreRedirects.h"
 #include "Misc/RedirectCollector.h"
+#include "Misc/AutomationTest.h"
 #include "String/Find.h"
 
 FSoftObjectPath::FSoftObjectPath(const UObject* InObject)
@@ -20,6 +21,13 @@ FSoftObjectPath::FSoftObjectPath(const UObject* InObject)
 	}
 }
 
+// Deprecated constructor
+FSoftObjectPath::FSoftObjectPath(FName InAssetPathName, FString InSubPathString)
+{
+	AssetPath = FTopLevelAssetPath(WriteToString<FName::StringBufferSize>(InAssetPathName).ToView());
+	SubPathString = MoveTemp(InSubPathString);
+}
+
 FString FSoftObjectPath::ToString() const
 {
 	// Most of the time there is no sub path so we can do a single string allocation
@@ -28,78 +36,56 @@ FString FSoftObjectPath::ToString() const
 		return GetAssetPathString();
 	}
 
-	TCHAR Buffer[FName::StringBufferSize];
-	FStringView AssetPathString;
-	if (AssetPathName.IsValid())
-	{
-		AssetPathString = FStringView(Buffer, AssetPathName.ToString(Buffer));
-	}
-
-	FString FullPathString;
-
-	// Preallocate to correct size and then append strings
-	FullPathString.Reserve(AssetPathString.Len() + SubPathString.Len() + 1);
-	FullPathString += AssetPathString;
-	FullPathString += SUBOBJECT_DELIMITER_CHAR;
-	FullPathString += SubPathString;
-	return FullPathString;
+	TStringBuilder<FName::StringBufferSize> Builder;
+	Builder << AssetPath << SUBOBJECT_DELIMITER_CHAR << SubPathString;
+	return Builder.ToString();
 }
 
 void FSoftObjectPath::ToString(FStringBuilderBase& Builder) const
 {
-	if (!AssetPathName.IsNone())
+	AppendString(Builder);
+}
+
+void FSoftObjectPath::AppendString(FStringBuilderBase& Builder) const
+{
+	if (AssetPath.IsNull())
 	{
-		Builder << AssetPathName;
+		return;
 	}
 
+	Builder << AssetPath;
 	if (SubPathString.Len() > 0)
 	{
 		Builder << SUBOBJECT_DELIMITER_CHAR << SubPathString;
 	}
 }
 
-FString FSoftObjectPath::GetLongPackageName() const
+void FSoftObjectPath::AppendString(FString& Builder) const
 {
-	if (!AssetPathName.IsNone())
+	if (AssetPath.IsNull())
 	{
-		FNameBuilder AssetPathNameBuilder(AssetPathName);
-		FStringView AssetPathNameView = AssetPathNameBuilder.ToView();
-
-		int32 DotPos = INDEX_NONE;
-		if (AssetPathNameView.FindChar(TEXT('.'), DotPos))
-		{
-			return FString(AssetPathNameView.Left(DotPos));
-		}
-
-		return FString(AssetPathNameView);
+		return;
 	}
 
-	return FString();
+	AssetPath.AppendString(Builder);
+	if (SubPathString.Len() > 0)
+	{
+		Builder += SUBOBJECT_DELIMITER_CHAR;
+		Builder += SubPathString;
+	}
 }
 
-FName FSoftObjectPath::GetLongPackageFName() const
+FName FSoftObjectPath::GetAssetPathName() const
 {
-	TCHAR Buffer[NAME_SIZE];
-	FStringView PlainAssetPath(Buffer, /* len */ AssetPathName.GetPlainNameString(Buffer));
-	int32 DotPos = UE::String::FindFirstChar(PlainAssetPath, '.');
-	return DotPos == INDEX_NONE ? AssetPathName : FName(PlainAssetPath.Left(DotPos));
+PRAGMA_DISABLE_DEPRECATION_WARNINGS
+	return AssetPath.ToFName();
+PRAGMA_ENABLE_DEPRECATION_WARNINGS
 }
 
-FString FSoftObjectPath::GetAssetName() const
+// Deprecated setter
+void FSoftObjectPath::SetAssetPathName(FName InPath)
 {
-	if (!AssetPathName.IsNone())
-	{
-		FNameBuilder AssetPathNameBuilder(AssetPathName);
-		FStringView AssetPathNameView = AssetPathNameBuilder.ToView();
-
-		int32 DotPos = INDEX_NONE;
-		if (AssetPathNameView.FindChar(TEXT('.'), DotPos))
-		{
-			return FString(AssetPathNameView.Mid(DotPos + 1));
-		}
-	}
-
-	return FString();
+	AssetPath = WriteToString<FName::StringBufferSize>(InPath).ToView();
 }
 
 /** Helper function that adds info about the object currently being serialized when triggering an ensure about invalid soft object path */
@@ -114,6 +100,12 @@ static FString GetObjectBeingSerializedForSoftObjectPath()
 	return Result;
 }
 
+void FSoftObjectPath::SetPath(const FTopLevelAssetPath& InAssetPath, FString InSubPathString)
+{
+	AssetPath = InAssetPath;
+	SubPathString = MoveTemp(InSubPathString);
+}
+
 void FSoftObjectPath::SetPath(FWideStringView Path)
 {
 	if (Path.IsEmpty() || Path.Equals(TEXT("None"), ESearchCase::CaseSensitive))
@@ -130,13 +122,13 @@ void FSoftObjectPath::SetPath(FWideStringView Path)
 		if (Path.FindChar(SUBOBJECT_DELIMITER_CHAR, ColonIndex))
 		{
 			// Has a subobject, split on that then create a name from the temporary path
-			AssetPathName = FName(Path.Left(ColonIndex));
+			AssetPath = Path.Left(ColonIndex);
 			SubPathString = Path.Mid(ColonIndex + 1);
 		}
 		else
 		{
 			// No Subobject
-			AssetPathName = FName(Path);
+			AssetPath = Path;
 			SubPathString.Empty();
 		}
 	}
@@ -156,37 +148,10 @@ void FSoftObjectPath::SetPath(FUtf8StringView Path)
 	SetPath(Wide);
 }
 
+// Deprecated setter
 void FSoftObjectPath::SetPath(FName PathName)
 {
-	if (PathName.IsNone())
-	{
-		Reset();
-	}
-	else
-	{
-		TCHAR Buffer[FName::StringBufferSize];
-		FStringView Path(Buffer, PathName.ToString(Buffer));
-
-		if (ensureMsgf(!FPackageName::IsShortPackageName(Path), TEXT("Cannot create SoftObjectPath with short package name '%s'%s! You must pass in fully qualified package names"), Path.GetData(), *GetObjectBeingSerializedForSoftObjectPath()))
-		{
-			// Possibly an ExportText path. Trim the ClassName.
-			Path = FPackageName::ExportTextPathToObjectPath(Path);
-
-			int32 ColonIndex;
-			if (Path.FindChar(SUBOBJECT_DELIMITER_CHAR, ColonIndex))
-			{
-				// Has a subobject, split on that then create a name from the temporary path
-				AssetPathName = FName(Path.Left(ColonIndex));
-				SubPathString = Path.Mid(ColonIndex + 1);
-			}
-			else
-			{
-				// No Subobject
-				AssetPathName = Path.GetData() == Buffer ? PathName : FName(Path);
-				SubPathString.Empty();
-			}
-		}
-	}
+	SetPath(PathName.ToString());
 }
 
 #if WITH_EDITOR
@@ -201,15 +166,15 @@ bool FSoftObjectPath::PreSavePath(bool* bReportSoftObjectPathRedirects)
 		return false;
 	}
 
-	FName FoundRedirection = GRedirectCollector.GetAssetPathRedirection(AssetPathName);
+	FSoftObjectPath FoundRedirection = GRedirectCollector.GetAssetPathRedirection(*this);
 
-	if (FoundRedirection != NAME_None)
+	if (!FoundRedirection.IsNull())
 	{
-		if (AssetPathName != FoundRedirection && bReportSoftObjectPathRedirects)
+		if (*this != FoundRedirection && bReportSoftObjectPathRedirects)
 		{
 			*bReportSoftObjectPathRedirects = true;
 		}
-		AssetPathName = FoundRedirection;
+		*this = FoundRedirection;
 		return true;
 	}
 
@@ -285,9 +250,17 @@ void FSoftObjectPath::SerializePath(FArchive& Ar)
 
 			SetPath(MoveTemp(Path));
 		}
-		else
+		else if( Ar.IsLoading() && Ar.UEVer() < EUnrealEngineObjectUE5Version::FSOFTOBJECTPATH_REMOVE_ASSET_PATH_FNAMES)
 		{
+			FName AssetPathName;
 			Ar << AssetPathName;
+			AssetPath = WriteToString<FName::StringBufferSize>(AssetPathName).ToView();
+
+			Ar << SubPathString;
+		}
+		else
+		{
+			Ar << AssetPath;
 			Ar << SubPathString;
 		}
 	}
@@ -320,7 +293,7 @@ void FSoftObjectPath::SerializePath(FArchive& Ar)
 
 bool FSoftObjectPath::operator==(FSoftObjectPath const& Other) const
 {
-	return AssetPathName == Other.AssetPathName && SubPathString == Other.SubPathString;
+	return AssetPath == Other.AssetPath && SubPathString == Other.SubPathString;
 }
 
 bool FSoftObjectPath::ExportTextItem(FString& ValueStr, FSoftObjectPath const& DefaultValue, UObject* Parent, int32 PortFlags, UObject* ExportRootScope) const
@@ -494,7 +467,7 @@ UObject* FSoftObjectPath::TryLoad(FUObjectSerializeContext* InLoadContext) const
 		if (IsSubobject())
 		{
 			// For subobjects, it's not safe to call LoadObject directly, so we want to load the parent object and then resolve again
-			FSoftObjectPath TopLevelPath = FSoftObjectPath(AssetPathName, FString());
+			FSoftObjectPath TopLevelPath = FSoftObjectPath(AssetPath, FString());
 			UObject* TopLevelObject = TopLevelPath.TryLoad(InLoadContext);
 
 			// This probably loaded the top-level object, so re-resolve ourselves
@@ -572,16 +545,9 @@ UObject* FSoftObjectPath::ResolveObject() const
 
 UObject* FSoftObjectPath::ResolveObjectInternal() const
 {
-	if (SubPathString.IsEmpty())
-	{
-		TCHAR PathString[FName::StringBufferSize];
-		AssetPathName.ToString(PathString);
-		return ResolveObjectInternal(PathString);
-	}
-	else
-	{
-		return ResolveObjectInternal(*ToString());
-	}
+	TStringBuilder<FName::StringBufferSize> Builder;
+	Builder << *this;
+	return ResolveObjectInternal(*Builder);
 }
 
 UObject* FSoftObjectPath::ResolveObjectInternal(const TCHAR* PathString) const
@@ -591,7 +557,7 @@ UObject* FSoftObjectPath::ResolveObjectInternal(const TCHAR* PathString) const
 	if (!FoundObject && IsSubobject())
 	{
 		// Try to resolve through the top level object
-		FSoftObjectPath TopLevelPath = FSoftObjectPath(AssetPathName, FString());
+		FSoftObjectPath TopLevelPath = FSoftObjectPath(AssetPath, FString());
 		UObject* TopLevelObject = TopLevelPath.ResolveObject();
 
 		// If the the top-level object exists but we can't find the object, defer the resolving to the top-level container object in case
@@ -834,3 +800,33 @@ bool FSoftObjectPathThreadContext::GetSerializationOptions(FName& OutPackageName
 
 FThreadSafeCounter FSoftObjectPath::CurrentTag(1);
 TSet<FName> FSoftObjectPath::PIEPackageNames;
+
+
+
+IMPLEMENT_SIMPLE_AUTOMATION_TEST(FSoftObjectPathTests, "System.CoreUObject.SoftObjectPath", EAutomationTestFlags::ApplicationContextMask | EAutomationTestFlags::EngineFilter);
+bool FSoftObjectPathTests::RunTest(const FString& Parameters)
+{
+	const TCHAR* PackageName = TEXT("/Game/Environments/Sets/Arid/Materials/M_Arid");
+	const TCHAR* AssetName = TEXT("M_Arid");
+	const FString String = FString::Printf(TEXT("%s.%s"), PackageName, AssetName);
+
+	const FString QuotedPath = FString::Printf(TEXT("\"%s\""), *String);
+	const FString UnquotedPath = String;
+
+	FSoftObjectPath Path(String);
+	TestEqual(TEXT("Correct package name"), Path.GetLongPackageName(), PackageName);
+	TestEqual(TEXT("Correct asset name"), Path.GetAssetName(), AssetName);
+	TestEqual(TEXT("Empty subpath"), Path.GetSubPathString(), TEXT(""));
+
+	FSoftObjectPath ImportQuoted;
+	const TCHAR* QuotedBuffer = *QuotedPath;
+	TestTrue(TEXT("Quoted path imports successfully"), ImportQuoted.ImportTextItem(QuotedBuffer, PPF_None, nullptr, GLog->Get()));
+	TestEqual(TEXT("Quoted path imports correctly"), ImportQuoted, Path);
+
+	FSoftObjectPath ImportUnquoted;
+	const TCHAR* UnquotedBuffer = *UnquotedPath;
+	TestTrue(TEXT("Unquoted path imports successfully"), ImportUnquoted.ImportTextItem(UnquotedBuffer, PPF_None, nullptr, GLog->Get()));
+	TestEqual(TEXT("Unquoted path imports correctly"), ImportUnquoted, Path);
+
+	return true;
+}
diff --git a/Engine/Source/Runtime/CoreUObject/Public/AssetRegistry/ARFilter.h b/Engine/Source/Runtime/CoreUObject/Public/AssetRegistry/ARFilter.h
--- a/Engine/Source/Runtime/CoreUObject/Public/AssetRegistry/ARFilter.h
+++ b/Engine/Source/Runtime/CoreUObject/Public/AssetRegistry/ARFilter.h
@@ -5,12 +5,14 @@
 #include "CoreMinimal.h"
 #include "UObject/Class.h"
 #include "UObject/ObjectMacros.h"
+#include "UObject/TopLevelAssetPath.h"
 
 struct FAssetData;
 
 /**
  * A struct to serve as a filter for Asset Registry queries.
  * Each component element is processed as an 'OR' operation while all the components are processed together as an 'AND' operation.
+ * This type is mirrored in NoExportTypes.h 
  */
 struct FARFilter
 {
@@ -29,8 +31,16 @@ PRAGMA_ENABLE_DEPRECATION_WARNINGS
 	TArray<FName> PackagePaths;
 
 	/** The filter component containing specific object paths */
+	// UE_DEPRECATED(5.1, "Asset path FNames have been deprecated, use FSoftObjectPath instead.")
 	TArray<FName> ObjectPaths;
 
+	/** 
+	 * The filter component containing the paths of specific assets to match. 
+	 * Matches against FAssetData::ToSoftObjectPath().
+	 * This is a top level asset path for most assets and a subobject path for assets such as world partition external actors.
+	 */
+	TArray<FSoftObjectPath> SoftObjectPaths;
+	
 	/** Deprecated. The filter component for class names. Instances of the specified classes, but not subclasses (by default), will be included. Derived classes will be included only if bRecursiveClasses is true. */
 	UE_DEPRECATED(5.1, "Class names are now represented by path names. Please use ClassPaths.")
 	TArray<FName> ClassNames;
@@ -68,7 +78,7 @@ PRAGMA_ENABLE_DEPRECATION_WARNINGS
 	{
 		PackageNames.Append(Other.PackageNames);
 		PackagePaths.Append(Other.PackagePaths);
-		ObjectPaths.Append(Other.ObjectPaths);
+		SoftObjectPaths.Append(Other.SoftObjectPaths);
 		ClassPaths.Append(Other.ClassPaths);
 
 		for (auto TagIt = Other.TagsAndValues.CreateConstIterator(); TagIt; ++TagIt)
@@ -88,7 +98,7 @@ PRAGMA_ENABLE_DEPRECATION_WARNINGS
 	/** Returns true if this filter has no entries */
 	bool IsEmpty() const
 	{
-		return PackageNames.Num() + PackagePaths.Num() + ObjectPaths.Num() + ClassPaths.Num() + TagsAndValues.Num() + WithoutPackageFlags + WithPackageFlags == 0;
+		return PackageNames.Num() + PackagePaths.Num() + SoftObjectPaths.Num() + ClassPaths.Num() + TagsAndValues.Num() + WithoutPackageFlags + WithPackageFlags == 0;
 	}
 
 	/** Returns true if this filter is recursive */
@@ -102,7 +112,7 @@ PRAGMA_ENABLE_DEPRECATION_WARNINGS
 	{
 		PackageNames.Empty();
 		PackagePaths.Empty();
-		ObjectPaths.Empty();
+		SoftObjectPaths.Empty();
 		ClassPaths.Empty();
 		TagsAndValues.Empty();
 		RecursiveClassPathsExclusionSet.Empty();
@@ -150,8 +160,12 @@ PRAGMA_ENABLE_DEPRECATION_WARNINGS
 	TSet<FName> PackagePaths;
 
 	/** The filter component containing specific object paths */
+	// UE_DEPRECATED(5.1, "Object path FNames have been deprecated, use FSoftObjectPath instead.")
 	TSet<FName> ObjectPaths;
 
+	/** The filter component containing specific object paths */
+	TSet<FSoftObjectPath> SoftObjectPaths;
+
 	/** Deprecated. The filter component for class names. Instances of the specified classes, but not subclasses (by default), will be included. Derived classes will be included only if bRecursiveClasses is true. */
 	UE_DEPRECATED(5.1, "Class names are now represented by path names. Please use ClassPaths.")
 	TSet<FName> ClassNames;
@@ -174,7 +188,7 @@ PRAGMA_ENABLE_DEPRECATION_WARNINGS
 	/** Returns true if this filter has no entries */
 	bool IsEmpty() const
 	{
-		return PackageNames.Num() + PackagePaths.Num() + ObjectPaths.Num() + ClassPaths.Num() + TagsAndValues.Num() == 0;
+		return PackageNames.Num() + PackagePaths.Num() + SoftObjectPaths.Num() + ClassPaths.Num() + TagsAndValues.Num() == 0;
 	}
 
 	/** Clears this filter of all entries */
@@ -182,7 +196,7 @@ PRAGMA_ENABLE_DEPRECATION_WARNINGS
 	{
 		PackageNames.Empty();
 		PackagePaths.Empty();
-		ObjectPaths.Empty();
+		SoftObjectPaths.Empty();
 		ClassPaths.Empty();
 		TagsAndValues.Empty();
 
diff --git a/Engine/Source/Runtime/CoreUObject/Public/AssetRegistry/AssetBundleData.h b/Engine/Source/Runtime/CoreUObject/Public/AssetRegistry/AssetBundleData.h
--- a/Engine/Source/Runtime/CoreUObject/Public/AssetRegistry/AssetBundleData.h
+++ b/Engine/Source/Runtime/CoreUObject/Public/AssetRegistry/AssetBundleData.h
@@ -12,6 +12,7 @@
 #include "UObject/PrimaryAssetId.h"
 #include "UObject/SoftObjectPath.h"
 #include "UObject/SoftObjectPtr.h"
+#include "UObject/TopLevelAssetPath.h"
 
 class FOutputDevice;
 class UObject;
@@ -21,20 +22,25 @@ struct COREUOBJECT_API FAssetBundleEntry
 {
 	/** Declare constructors inline so this can be a header only class */
 	FORCEINLINE FAssetBundleEntry() {}
+PRAGMA_DISABLE_DEPRECATION_WARNINGS
 	FORCEINLINE ~FAssetBundleEntry() {}
-
+	// Explicitly declared constructors/operators because we explicitly declared the destructor to hide deprecation warnings.
+	// These can be removed along with the destructor after deprecation.
+	FAssetBundleEntry(const FAssetBundleEntry&) = default;
+	FAssetBundleEntry(FAssetBundleEntry&&) = default;
+	FAssetBundleEntry& operator=(const FAssetBundleEntry&) = default;
+	FAssetBundleEntry& operator=(FAssetBundleEntry&&) = default;
+PRAGMA_ENABLE_DEPRECATION_WARNINGS
 
 	/** Specific name of this bundle, should be unique for a given scope */
 	FName BundleName;
 
 	/** List of string assets contained in this bundle */
+	// UE_DEPRECATED(5.1, "BundleAssets has been replaced with AssetPaths.")
 	TArray<FSoftObjectPath> BundleAssets;
 
-	FAssetBundleEntry(const FAssetBundleEntry& OldEntry)
-		: BundleName(OldEntry.BundleName), BundleAssets(OldEntry.BundleAssets)
-	{
-
-	}
+	/** List of references to top-level assets contained in this bundle */
+	TArray<FTopLevelAssetPath> AssetPaths;
 	
 	explicit FAssetBundleEntry(FName InBundleName)
 		: BundleName(InBundleName)
@@ -56,7 +62,9 @@ struct COREUOBJECT_API FAssetBundleEntry
 	/** Equality */
 	bool operator==(const FAssetBundleEntry& Other) const
 	{
-		return BundleName == Other.BundleName && BundleAssets == Other.BundleAssets;
+		// Ignore deprecated fields for equality checks
+		return BundleName == Other.BundleName
+			&& AssetPaths == Other.AssetPaths;
 	}
 	bool operator!=(const FAssetBundleEntry& Other) const
 	{
@@ -102,20 +110,42 @@ struct COREUOBJECT_API FAssetBundleData
 	}
 
 	/** Adds or updates an entry with the given BundleName -> Path. Scope is empty and will be filled in later */
+	// UE_DEPRECATED(5.1, "The contents of an asset bundle now use FTopLevelAssetPath instead of FSoftObjectPath.")
 	void AddBundleAsset(FName BundleName, const FSoftObjectPath& AssetPath);
 
 	template< typename T >
+	// UE_DEPRECATED(5.1, "The contents of an asset bundle now use FTopLevelAssetPath instead of FSoftObjectPath.")
 	void AddBundleAsset(FName BundleName, const TSoftObjectPtr<T>& SoftObjectPtr)
 	{
-		AddBundleAsset(BundleName, SoftObjectPtr.ToSoftObjectPath());
+		AddBundleAsset(BundleName, SoftObjectPtr.ToSoftObjectPath().GetAssetPath());
 	}
 
 	/** Adds multiple assets at once */
+	// UE_DEPRECATED(5.1, "The contents of an asset bundle now use FTopLevelAssetPath instead of FSoftObjectPath.")
 	void AddBundleAssets(FName BundleName, const TArray<FSoftObjectPath>& AssetPaths);
 
 	/** A fast set of asset bundle assets, will destroy copied in path list */
+	// UE_DEPRECATED(5.1, "The contents of an asset bundle now use FTopLevelAssetPath instead of FSoftObjectPath.")
 	void SetBundleAssets(FName BundleName, TArray<FSoftObjectPath>&& AssetPaths);
 
+	/** Adds or updates an entry with the given BundleName -> Path.*/
+	void AddBundleAsset(FName BundleName, const FTopLevelAssetPath& AssetPath);
+
+	/** Adds multiple assets at once */
+	void AddBundleAssets(FName BundleName, const TArray<FTopLevelAssetPath>& AssetPaths);
+
+	/** A fast set of asset bundle assets, will destroy copied in path list */
+	void SetBundleAssets(FName BundleName, TArray<FTopLevelAssetPath>&& AssetPaths);
+
+	/** Adds or updates an entry with the given BundleName -> Path, truncating to a top-level asset path if necessary.*/
+	void AddBundleAssetTruncated(FName BundleName, const FSoftObjectPath& AssetPath);
+
+	/** Adds multiple assets at once, truncating to top-level asset paths if necessary */
+	void AddBundleAssetsTruncated(FName BundleName, const TArray<FSoftObjectPath>& AssetPaths);
+
+	/** Replaces the contents of this bundle with the given assets, truncating to top-level asset paths if necessary. */
+	void SetBundleAssetsTruncated(FName BundleName, const TArray<FSoftObjectPath>& AssetPaths);
+
 	/** Resets the data to defaults */
 	void Reset();
 
diff --git a/Engine/Source/Runtime/CoreUObject/Public/AssetRegistry/AssetData.h b/Engine/Source/Runtime/CoreUObject/Public/AssetRegistry/AssetData.h
--- a/Engine/Source/Runtime/CoreUObject/Public/AssetRegistry/AssetData.h
+++ b/Engine/Source/Runtime/CoreUObject/Public/AssetRegistry/AssetData.h
@@ -92,6 +92,7 @@ struct COREUOBJECT_API FAssetRegistryVersion
 		ObjectResourceOptionalVersionChange,// Change to linker export/import resource serialization
 		AddedChunkHashes,					// Added FIoHash for each FIoChunkId in the package to the AssetPackageData.
 		ClassPaths,							// Classes are serialized as path names rather than short object names, e.g. /Script/Engine.StaticMesh
+		RemoveAssetPathFNames,				// Asset bundles are serialized as FTopLevelAssetPath instead of FSoftObjectPath, deprecated FAssetData::ObjectPath	
 
 		// -----<new versions can be added above this line>-------------------------------------------------
 		VersionPlusOne,
@@ -108,6 +109,23 @@ private:
 	FAssetRegistryVersion() {}
 };
 
+namespace UE::AssetRegistry::Private
+{
+	struct FAssetPathParts
+	{
+		FStringView OuterPath;
+		FStringView InnermostName;
+	};
+
+	/** 
+	* Split a full object path into the path of the innermost object's outer and the name of that object.
+	* E.g. splits /Path/To/PackageName.AssetName:SubObject.Innermost into /Path/To/PackageName.AssetName:SubObject and Innermost
+	*/
+	COREUOBJECT_API FAssetPathParts SplitIntoOuterPathAndAssetName(FStringView InObjectPath);
+	/** Concatenates an existing object path with an inner object with the correct separator ('.' or ':'). Assumes the outer path already contains correct separators. */
+	COREUOBJECT_API void ConcatenateOuterPathAndObjectName(FStringBuilderBase& Builder, FName OuterPath, FName ObjectName);
+}
+
 /** 
  * A struct to hold important information about an assets found by the Asset Registry
  * This struct is transient and should never be serialized
@@ -131,6 +149,7 @@ public:
 
 public:
 	/** The object path for the asset in the form PackageName.ObjectName, or PackageName.ObjectName:SubObjectName. */
+	// UE_DEPRECATED(5.1, "FName asset paths have been deprecated. Use GetSoftObjectPath to get the path this asset will use in memory when loaded.")
 	FName ObjectPath;
 	/** The name of the package in which the asset is found, this is the full long package name such as /Game/Path/Package */
 	FName PackageName;
@@ -139,12 +158,27 @@ public:
 	/** The name of the asset without the package */
 	FName AssetName;
 	/** The name of the asset's class */
-	UE_DEPRECATED(5.1, "Class names are now represented by path names. Please use AssetClassPath.")
+	// UE_DEPRECATED(5.1, "Class names are now represented by path names. Please use AssetClassPath.")
 	FName AssetClass;
 	/** The path of the asset's class, e.g. /Script/Engine.StaticMesh */
 	FTopLevelAssetPath AssetClassPath;
 	/** Asset package flags */
 	uint32 PackageFlags = 0;
+
+#if WITH_EDITORONLY_DATA
+private:
+	/**
+	 * If this object is not a top level asset, this contains the path of the outer of this object.
+	 * Non top-level assets may only be used in the editor.
+	 * For some assets (such as external actors) this may not start with PackageName.
+	 * e.g. PackageName			= /Game/__EXTERNAL_ACTORS__/Maps/MyMap/ABCDE12345
+			OptionalOuterPath	= /Game/Maps/MyMap.MyMap:PersistentLevel
+			AssetName			= SomeExternalActor
+	*/
+	FName OptionalOuterPath;
+public:
+#endif
+
 	/** The map of values for properties that were marked AssetRegistrySearchable or added by GetAssetRegistryTags */
 	FAssetDataTagMapSharedView TagsAndValues;
 	/**
@@ -189,31 +223,41 @@ PRAGMA_ENABLE_DEPRECATION_WARNINGS
 	{
 	}
 
-	/** FAssetDatas are equal if their object paths match */
+	/** FAssetDatas are uniquely identified by PackageName and AssetName. */
 	bool operator==(const FAssetData& Other) const
 	{
-		return ObjectPath == Other.ObjectPath;
+		return PackageName == Other.PackageName && AssetName == Other.AssetName;
 	}
 
 	bool operator!=(const FAssetData& Other) const
 	{
-		return ObjectPath != Other.ObjectPath;
+		return PackageName != Other.PackageName || AssetName != Other.AssetName;
 	}
 
+	/** Perform a lexical greater-than operation on the PackageName and AssetName that uniquely identify two FAssetData. */
 	bool operator>(const FAssetData& Other) const
 	{
-		return  Other.ObjectPath.LexicalLess(ObjectPath);
+		if (PackageName == Other.PackageName)
+		{
+			return Other.AssetName.LexicalLess(AssetName);
+		}
+		return Other.PackageName.LexicalLess(PackageName);
 	}
 
+	/** Perform a lexical less-than operation on the PackageName and AssetName that uniquely identify two FAssetData. */
 	bool operator<(const FAssetData& Other) const
 	{
-		return ObjectPath.LexicalLess(Other.ObjectPath);
+		if (PackageName == Other.PackageName)
+		{
+			return AssetName.LexicalLess(Other.AssetName);
+		}
+		return PackageName.LexicalLess(Other.PackageName);
 	}
 
 	/** Checks to see if this AssetData refers to an asset or is NULL */
 	bool IsValid() const
 	{
-		return !ObjectPath.IsNone();
+		return !PackageName.IsNone() && !AssetName.IsNone();
 	}
 
 	/**
@@ -237,6 +281,23 @@ PRAGMA_ENABLE_DEPRECATION_WARNINGS
 		return DetectIsUAssetByNames(PackageNameStrBuilder, AssetNameStrBuilder);
 	}
 
+	/** Convert to a SoftObjectPath. PackageName.AssetName for ToplevelAssets. OptionalOuterName and AssetName joined by : or . for sub-object assets */
+	COREUOBJECT_API FSoftObjectPath GetSoftObjectPath() const;
+	
+	// TODO: Deprecate in favor of GetSoftObjectPath
+	FSoftObjectPath ToSoftObjectPath() const
+	{
+		return GetSoftObjectPath();
+	}
+	
+	/** Return the object path as a string. */
+	FString GetObjectPathString() const
+	{
+		TStringBuilder<FName::StringBufferSize> Builder;
+		AppendObjectPath(Builder);
+		return FString(Builder);
+	}
+
 	/**
 	 * Returns true if the given UObject is the main asset in a package, true for maps and assets but false for secondary objects like class redirectors
 	 * Every UAsset is also a TopLevelAsset.
@@ -264,7 +325,12 @@ PRAGMA_ENABLE_DEPRECATION_WARNINGS
 		TagsAndValues.Shrink();
 	}
 
-	/** Returns the full name for the asset in the form: Class ObjectPath */
+#if WITH_EDITORONLY_DATA	
+	FName GetOptionalOuterPathName() const { return OptionalOuterPath; }
+	void SetOptionalOuterPathName(FName InName) { OptionalOuterPath = InName; }
+#endif
+
+	/** Returns the full name for the asset in the form: Class FullPath */
 	FString GetFullName() const
 	{
 		FString FullName;
@@ -272,13 +338,13 @@ PRAGMA_ENABLE_DEPRECATION_WARNINGS
 		return FullName;
 	}
 
-	/** Populates OutFullName with the full name for the asset in the form: Class ObjectPath */
+	/** Populates OutFullName with the full name for the asset in the form: Class FullPath */
 	void GetFullName(FString& OutFullName) const
 	{
 		OutFullName.Reset();
 		OutFullName += AssetClassPath.ToString();
 		OutFullName.AppendChar(TEXT(' '));
-		ObjectPath.AppendString(OutFullName);
+		AppendObjectPath(OutFullName);
 	}
 
 	/** Populates OutFullNameBuilder with the full name for the asset in the form: Class ObjectPath */
@@ -287,10 +353,10 @@ PRAGMA_ENABLE_DEPRECATION_WARNINGS
 		OutFullNameBuilder.Reset();
 		AssetClassPath.AppendString(OutFullNameBuilder);
 		OutFullNameBuilder.AppendChar(TEXT(' '));
-		ObjectPath.AppendString(OutFullNameBuilder);
+		AppendObjectPath(OutFullNameBuilder);
 	}
 
-	/** Returns the name for the asset in the form: Class'ObjectPath' */
+	/** Returns the name for the asset in the form: Class'FullPath' */
 	FString GetExportTextName() const
 	{
 		FString ExportTextName;
@@ -298,23 +364,23 @@ PRAGMA_ENABLE_DEPRECATION_WARNINGS
 		return ExportTextName;
 	}
 
-	/** Populates OutExportTextName with the name for the asset in the form: Class'ObjectPath' */
+	/** Populates OutExportTextName with the name for the asset in the form: Class'FullPath' */
 	void GetExportTextName(FString& OutExportTextName) const
 	{
 		OutExportTextName.Reset();
 		OutExportTextName.Append(AssetClassPath.ToString());
 		OutExportTextName.AppendChar(TEXT('\''));
-		ObjectPath.AppendString(OutExportTextName);
+		AppendObjectPath(OutExportTextName);
 		OutExportTextName.AppendChar(TEXT('\''));
 	}
 
-	/** Populates OutExportTextNameBuilder with the name for the asset in the form: Class'ObjectPath' */
+	/** Populates OutExportTextNameBuilder with the name for the asset in the form: Class'FullPath' */
 	void GetExportTextName(FStringBuilderBase& OutExportTextNameBuilder) const
 	{
 		OutExportTextNameBuilder.Reset();
 		AssetClassPath.AppendString(OutExportTextNameBuilder);
 		OutExportTextNameBuilder.AppendChar(TEXT('\''));
-		ObjectPath.AppendString(OutExportTextNameBuilder);
+		AppendObjectPath(OutExportTextNameBuilder);
 		OutExportTextNameBuilder.AppendChar(TEXT('\''));
 	}
 
@@ -358,16 +424,35 @@ PRAGMA_ENABLE_DEPRECATION_WARNINGS
 		return ClassPointer && ClassPointer->IsChildOf(BaseClass);
 	}
 
-	/** Convert to a SoftObjectPath for loading */
-	FSoftObjectPath ToSoftObjectPath() const
+	/** Append the object path to the given string builder. */
+	void AppendObjectPath(FStringBuilderBase& Builder) const
+	{
+#if WITH_EDITORONLY_DATA
+		if (!OptionalOuterPath.IsNone())
+		{
+			UE::AssetRegistry::Private::ConcatenateOuterPathAndObjectName(Builder, OptionalOuterPath, AssetName);
+		}
+		else
+#endif
+		{
+			Builder << PackageName << '.' << AssetName;
+		}
+
+	}
+
+	/** Append the object path to the given string. */
+	void AppendObjectPath(FString& String) const
 	{
-		return FSoftObjectPath(ObjectPath);
+		TStringBuilder<FName::StringBufferSize> Builder;
+		AppendObjectPath(Builder);
+		String = FString(Builder);
 	}
 
-	UE_DEPRECATED(4.18, "ToStringReference was renamed to ToSoftObjectPath")
+
+	UE_DEPRECATED(4.18, "ToStringReference was renamed to GetSoftObjectPath")
 	FSoftObjectPath ToStringReference() const
 	{
-		return ToSoftObjectPath();
+		return GetSoftObjectPath();
 	}
 	
 	/** Gets primary asset id of this data */
@@ -381,33 +466,38 @@ PRAGMA_ENABLE_DEPRECATION_WARNINGS
 	 */
 	UObject* FastGetAsset(bool bLoad = false, TSet<FName> LoadTags = {}) const
 	{
-		if ( !IsValid())
+		if (!IsValid())
 		{
-			// Do not try to find the object if the objectpath is not set
+			// Do not try to find the object if the fields are not set
 			return nullptr;
 		}
 
+		// We load PackageName rather than using GetObjectPath because external assets may be saved in a different package to their loaded path. 
 		UPackage* FoundPackage = FindObjectFast<UPackage>(nullptr, PackageName);
-		if (FoundPackage == nullptr)
+		if (FoundPackage == nullptr && bLoad)
 		{
-			if (bLoad)
-			{
-				FLinkerInstancingContext InstancingContext(MoveTemp(LoadTags));
-				return LoadObject<UObject>(nullptr, *ObjectPath.ToString(), nullptr, 0, nullptr, &InstancingContext);
-			}
-			else
-			{
-				return nullptr;
-			}
+			FLinkerInstancingContext InstancingContext(MoveTemp(LoadTags));
+			FoundPackage = LoadPackage(nullptr, *PackageName.ToString(), LOAD_None, nullptr, &InstancingContext);
 		}
 
-		UObject* Asset = FindObjectFast<UObject>(FoundPackage, AssetName);
-		if (Asset == nullptr && bLoad)
+		if (!FoundPackage)
 		{
-			return LoadObject<UObject>(nullptr, *ObjectPath.ToString());
+			return nullptr;
 		}
 
-		return Asset;
+#if WITH_EDITORONLY_DATA
+		if (!OptionalOuterPath.IsNone())
+		{
+			TStringBuilder<FName::StringBufferSize> Builder;
+			AppendObjectPath(Builder);
+			return FindObject<UObject>(nullptr, *Builder);
+		}
+		else
+#endif
+		{
+			// If the package is loaded all objects in it must be loaded, just search for the asset in memory with its outer.
+			return FindObjectFast<UObject>(FoundPackage, AssetName);
+		}
 	}
 
 	/** 
@@ -417,20 +507,8 @@ PRAGMA_ENABLE_DEPRECATION_WARNINGS
 	 */
 	UObject* GetAsset(TSet<FName> LoadTags = {}) const
 	{
-		if ( !IsValid())
-		{
-			// Dont even try to find the object if the objectpath isn't set
-			return nullptr;
-		}
 
-		UObject* Asset = FindObject<UObject>(nullptr, *ObjectPath.ToString());
-		if ( Asset == nullptr)
-		{
-			FLinkerInstancingContext InstancingContext(MoveTemp(LoadTags));
-			Asset = LoadObject<UObject>(nullptr, *ObjectPath.ToString(), nullptr, 0, nullptr, &InstancingContext);
-		}
-
-		return Asset;
+		return FastGetAsset(true /* bLoad */, MoveTemp(LoadTags));
 	}
 
 	/**
@@ -455,11 +533,12 @@ PRAGMA_ENABLE_DEPRECATION_WARNINGS
 		return ((PackageFlags & FlagsToCheck) == FlagsToCheck);
 	}
 
+	/** Tries to find the package in memory if it is loaded, otherwise loads it. */
 	UPackage* GetPackage() const
 	{
-		if (PackageName == NAME_None)
+		if (PackageName.IsNone())
 		{
-			return NULL;
+			return nullptr;
 		}
 
 		UPackage* Package = FindPackage(NULL, *PackageName.ToString());
@@ -499,13 +578,13 @@ PRAGMA_ENABLE_DEPRECATION_WARNINGS
 	/** Returns true if the asset is loaded */
 	bool IsAssetLoaded() const
 	{
-		return IsValid() && FindObjectSafe<UObject>(NULL, *ObjectPath.ToString()) != NULL;
+		return FastGetAsset(false) != nullptr;
 	}
 
 	/** Prints the details of the asset to the log */
 	void PrintAssetData() const
 	{
-		UE_LOG(LogAssetData, Log, TEXT("    FAssetData for %s"), *ObjectPath.ToString());
+		UE_LOG(LogAssetData, Log, TEXT("    FAssetData for %s"), *GetObjectPathString());
 		UE_LOG(LogAssetData, Log, TEXT("    ============================="));
 		UE_LOG(LogAssetData, Log, TEXT("        PackageName: %s"), *PackageName.ToString());
 		UE_LOG(LogAssetData, Log, TEXT("        PackagePath: %s"), *PackagePath.ToString());
@@ -559,7 +638,7 @@ PRAGMA_ENABLE_DEPRECATION_WARNINGS
 	template<class Archive>
 	FORCEINLINE void SerializeForCache(Archive&& Ar)
 	{
-		SerializeForCacheWithTagsAndBundles(Ar, [](FArchive& Ar, FAssetData& Ad) {
+		SerializeForCacheWithTagsAndBundles(Ar, [](FArchive& Ar, FAssetData& Ad, FAssetRegistryVersion::Type) {
 			static_cast<Archive&>(Ar).SerializeTagsAndBundles(Ad);
 			});
 	}
@@ -569,11 +648,13 @@ PRAGMA_ENABLE_DEPRECATION_WARNINGS
 	template<class Archive>
 	FORCEINLINE void SerializeForCacheOldVersion(Archive&& Ar, FAssetRegistryVersion::Type Version = FAssetRegistryVersion::LatestVersion)
 	{
-		SerializeForCacheOldVersionWithTagsAndBundles(Ar, Version, [](FArchive& Ar, FAssetData& Ad) {
-			static_cast<Archive&>(Ar).SerializeTagsAndBundles(Ad);
+		SerializeForCacheOldVersionWithTagsAndBundles(Ar, Version, [](FArchive& Ar, FAssetData& Ad, FAssetRegistryVersion::Type Version) {
+			static_cast<Archive&>(Ar).SerializeTagsAndBundlesOldVersion(Ad, Version);
 			});
 	}
 
+	// Note: these functions should only be used for live communication between processing running the same version of the engine.
+	// There is no versioning support
 	COREUOBJECT_API void NetworkWrite(FCbWriter& Writer, bool bWritePackageName) const;
 	COREUOBJECT_API bool TryNetworkRead(FCbFieldView Field, bool bReadPackageName, FName InPackageName);
 private:
@@ -582,9 +663,9 @@ private:
 	 * Note that this function is force-inlined but defined in AssetData.cpp which is fine as functions will get inlined
 	 * as long as they're defined before they are used for the first time by other functions (SerializeForCacheWithTagsAndBundles in this case)
 	 */
-	FORCEINLINE void SerializeForCacheInternal(FArchive& Ar, FAssetRegistryVersion::Type Version, void (*SerializeTagsAndBundles)(FArchive& , FAssetData&));
-	COREUOBJECT_API void SerializeForCacheWithTagsAndBundles(FArchive& Ar, void (*SerializeTagsAndBundles)(FArchive&, FAssetData&));
-	COREUOBJECT_API void SerializeForCacheOldVersionWithTagsAndBundles(FArchive& Ar, FAssetRegistryVersion::Type Version, void (*SerializeTagsAndBundles)(FArchive&, FAssetData&));
+	FORCEINLINE void SerializeForCacheInternal(FArchive& Ar, FAssetRegistryVersion::Type Version, void (*SerializeTagsAndBundles)(FArchive& , FAssetData&, FAssetRegistryVersion::Type));
+	COREUOBJECT_API void SerializeForCacheWithTagsAndBundles(FArchive& Ar, void (*SerializeTagsAndBundles)(FArchive&, FAssetData&, FAssetRegistryVersion::Type));
+	COREUOBJECT_API void SerializeForCacheOldVersionWithTagsAndBundles(FArchive& Ar, FAssetRegistryVersion::Type Version, void (*SerializeTagsAndBundles)(FArchive&, FAssetData&, FAssetRegistryVersion::Type));
 
 	static bool DetectIsUAssetByNames(FStringView PackageName, FStringView ObjectPathName)
 	{
@@ -614,7 +695,7 @@ ENUM_CLASS_FLAGS(FAssetData::ECreationFlags);
 
 FORCEINLINE uint32 GetTypeHash(const FAssetData& AssetData)
 {
-	return GetTypeHash(AssetData.ObjectPath);
+	return HashCombine(GetTypeHash(AssetData.PackageName), GetTypeHash(AssetData.AssetName));
 }
 
 
diff --git a/Engine/Source/Runtime/CoreUObject/Public/Misc/RedirectCollector.h b/Engine/Source/Runtime/CoreUObject/Public/Misc/RedirectCollector.h
--- a/Engine/Source/Runtime/CoreUObject/Public/Misc/RedirectCollector.h
+++ b/Engine/Source/Runtime/CoreUObject/Public/Misc/RedirectCollector.h
@@ -31,15 +31,15 @@ private:
 	/** Helper struct for soft object path tracking */
 	struct FSoftObjectPathProperty
 	{
-		FSoftObjectPathProperty(FName InAssetPathName, FName InProperty, bool bInReferencedByEditorOnlyProperty)
-			: AssetPathName(InAssetPathName)
+		FSoftObjectPathProperty(FSoftObjectPath InObjectPath, FName InProperty, bool bInReferencedByEditorOnlyProperty)
+			: ObjectPath(MoveTemp(InObjectPath))
 			, PropertyName(InProperty)
 			, bReferencedByEditorOnlyProperty(bInReferencedByEditorOnlyProperty)
 		{}
 
 		 bool operator==(const FSoftObjectPathProperty& Other) const
 		 {
-		 	return AssetPathName == Other.AssetPathName &&
+		 	return ObjectPath == Other.ObjectPath &&
 		 		PropertyName == Other.PropertyName &&
 		 		bReferencedByEditorOnlyProperty == Other.bReferencedByEditorOnlyProperty;
 		 }
@@ -47,15 +47,15 @@ private:
 		friend inline uint32 GetTypeHash(const FSoftObjectPathProperty& Key)
 		{
 			uint32 Hash = 0;
-			Hash = HashCombine(Hash, GetTypeHash(Key.AssetPathName));
+			Hash = HashCombine(Hash, GetTypeHash(Key.ObjectPath));
 			Hash = HashCombine(Hash, GetTypeHash(Key.PropertyName));
 			Hash = HashCombine(Hash, (uint32)Key.bReferencedByEditorOnlyProperty);
 			return Hash;
 		}
 
-		const FName& GetAssetPathName() const
+		const FSoftObjectPath& GetObjectPath() const
 		{
-			return AssetPathName;
+			return ObjectPath;
 		}
 
 		const FName& GetPropertyName() const
@@ -69,7 +69,7 @@ private:
 		}
 
 	private:
-		FName AssetPathName;
+		FSoftObjectPath ObjectPath;
 		FName PropertyName;
 		bool bReferencedByEditorOnlyProperty;
 	};
@@ -106,13 +106,20 @@ public:
 	void ProcessSoftObjectPathPackageList(FName FilterPackage, bool bGetEditorOnly, TSet<FName>& OutReferencedPackages);
 
 	/** Adds a new mapping for redirector path to destination path, this is called from the Asset Registry to register all redirects it knows about */
+	// UE_DEPRECATED(5.1, "Asset path FNames have been deprecated, use FSoftObjectPath instead.")
 	void AddAssetPathRedirection(FName OriginalPath, FName RedirectedPath);
+	void AddAssetPathRedirection(const FSoftObjectPath& OriginalPath, const FSoftObjectPath& RedirectedPath);
 
 	/** Removes an asset path redirection, call this when deleting redirectors */
+	// UE_DEPRECATED(5.1, "Asset path FNames have been deprecated, use FSoftObjectPath instead.")
 	void RemoveAssetPathRedirection(FName OriginalPath);
+	void RemoveAssetPathRedirection(const FSoftObjectPath& OriginalPath);
 
 	/** Returns a remapped asset path, if it returns null there is no relevant redirector */
+	// UE_DEPRECATED(5.1, "Asset path FNames have been deprecated, use FSoftObjectPath instead.")
 	FName GetAssetPathRedirection(FName OriginalPath);
+	/** Returns a remapped asset path, if there is no relevant redirector, the return value reports true from IsNull() */
+	FSoftObjectPath GetAssetPathRedirection(const FSoftObjectPath& OriginalPath);
 
 	/**
 	 * Do we have any references to resolve.
@@ -151,7 +158,7 @@ private:
 	TMap<FName, TMap<FName, ESoftObjectPathCollectType>> PackageReferenceTypes;
 
 	/** When saving, apply this remapping to all soft object paths */
-	TMap<FName, FName> AssetPathRedirectionMap;
+	TMap<FSoftObjectPath, FSoftObjectPath> ObjectPathRedirectionMap;
 
 	/** For SoftObjectPackageMap map */
 	FCriticalSection CriticalSection;
diff --git a/Engine/Source/Runtime/CoreUObject/Public/UObject/LinkerInstancingContext.h b/Engine/Source/Runtime/CoreUObject/Public/UObject/LinkerInstancingContext.h
--- a/Engine/Source/Runtime/CoreUObject/Public/UObject/LinkerInstancingContext.h
+++ b/Engine/Source/Runtime/CoreUObject/Public/UObject/LinkerInstancingContext.h
@@ -44,11 +44,9 @@ public:
 	 */
 	FSoftObjectPath RemapPath(const FSoftObjectPath& Path) const
 	{
-		if (const FName* Remapped = PathMapping.Find(Path.GetAssetPathName()))
+		if (const FTopLevelAssetPath* Remapped = PathMapping.Find(Path.GetAssetPath()))
 		{
-			FSoftObjectPath NewPath = Path;
-			NewPath.SetAssetPathName(*Remapped);
-			return NewPath;
+			return FSoftObjectPath(*Remapped, Path.GetSubPathString());
 		}
 		return Path;
 	}
@@ -67,7 +65,7 @@ public:
 		ensureAlwaysMsgf(Instanced.GetSubPathString().IsEmpty(), 
 			TEXT("Linker instance remap paths should be top-level assets only: ->%s"), *Instanced.ToString());
 	
-		PathMapping.Emplace(Original.GetAssetPathName(), Instanced.GetAssetPathName());
+		PathMapping.Emplace(Original.GetAssetPath(), Instanced.GetAssetPath());
 	}
 
 	void AddTag(FName NewTag)
@@ -116,8 +114,8 @@ private:
 
 	/** Map of original package name to their instance counterpart. */
 	TMap<FName, FName> PackageMapping;
-	/** Map of original top level asset path to their instance counterpart. This should be FTopLevelAssetPath when FSoftObjectPath is refactored to contain that. */
-	TMap<FName, FName> PathMapping;
+	/** Map of original top level asset path to their instance counterpart. */
+	TMap<FTopLevelAssetPath, FTopLevelAssetPath> PathMapping;
 	/** Tags can be used to determine some loading behavior. */
 	TSet<FName> Tags;
 	/** Remap soft object paths */
diff --git a/Engine/Source/Runtime/CoreUObject/Public/UObject/NoExportTypes.h b/Engine/Source/Runtime/CoreUObject/Public/UObject/NoExportTypes.h
--- a/Engine/Source/Runtime/CoreUObject/Public/UObject/NoExportTypes.h
+++ b/Engine/Source/Runtime/CoreUObject/Public/UObject/NoExportTypes.h
@@ -22,6 +22,7 @@
 #include "Misc/DateTime.h"
 #include "Misc/Timespan.h"
 
+#include "UObject/TopLevelAssetPath.h"
 #include "UObject/SoftObjectPath.h"
 #include "UObject/PropertyAccessUtil.h"
 
@@ -1995,6 +1996,22 @@ struct FTimespan
 	int64 Ticks;
 };
 
+/**
+ * A struct that can reference a top level asset such as '/Path/To/Package.AssetName'
+ * @note The full C++ class is located here: Engine\Source\Runtime\CoreUObject\Public\UObject\TopLevelAssetPath.h
+ */
+USTRUCT(noexport, BlueprintType)
+struct FTopLevelAssetPath
+{
+private:
+	/** Name of the package containing the asset e.g. /Path/To/Package */
+	UPROPERTY(EditAnywhere, SaveGame, BlueprintReadWrite, Category = TopLevelAssetPath, meta = (AllowPrivateAccess = "true"))
+	FName PackageName;
+	/** Name of the asset within the package e.g. 'AssetName' */
+	UPROPERTY(EditAnywhere, SaveGame, BlueprintReadWrite, Category = TopLevelAssetPath, meta = (AllowPrivateAccess = "true"))
+	FName AssetName;
+};
+
 /** 
  * A struct that contains a string reference to an object, either a top level asset or a subobject.
  * @note The full C++ class is located here: Engine\Source\Runtime\CoreUObject\Public\UObject\SoftObjectPath.h
@@ -2004,7 +2021,7 @@ struct FSoftObjectPath
 {
 	/** Asset path, patch to a top level object in a package */
 	UPROPERTY()
-	FName AssetPathName;
+	FTopLevelAssetPath AssetPath;
 
 	/** Optional FString for subobject within an asset */
 	UPROPERTY()
@@ -2020,22 +2037,6 @@ struct FSoftClassPath : public FSoftObjectPath
 {
 };
 
-/**
- * A struct that can reference a top level asset such as '/Path/To/Package.AssetName'
- * @note The full C++ class is located here: Engine\Source\Runtime\CoreUObject\Public\UObject\TopLevelAssetPath.h
- */
-USTRUCT(noexport, BlueprintType)
-struct FTopLevelAssetPath
-{
-private:
-	/** Name of the package containing the asset e.g. /Path/To/Package */
-	UPROPERTY(EditAnywhere, SaveGame, BlueprintReadWrite, Category = TopLevelAssetPath, meta = (AllowPrivateAccess = "true"))
-	FName PackageName;
-	/** Name of the asset within the package e.g. 'AssetName' */
-	UPROPERTY(EditAnywhere, SaveGame, BlueprintReadWrite, Category = TopLevelAssetPath, meta = (AllowPrivateAccess = "true"))
-	FName AssetName;
-};
-
 /** 
  * A type of primary asset, used by the Asset Manager system.
  * @note The full C++ class is located here: Engine\Source\Runtime\CoreUObject\Public\UObject\PrimaryAssetId.h
@@ -2458,6 +2459,9 @@ enum class EDataValidationResult : uint8
 	NotValidated
 };
 
+/**
+ * A struct to serve as a filter for Asset Registry queries. (mirrored in ARFilter.h)
+ */
 USTRUCT(noexport, BlueprintType, IsAlwaysAccessible, HasDefaults, meta = (HasNativeMake = "/Script/Engine.KismetSystemLibrary.MakeARFilter", HasNativeBreak = "/Script/Engine.KismetSystemLibrary.BreakARFilter"))
 struct FARFilter
 {
@@ -2469,9 +2473,13 @@ struct FARFilter
 	UPROPERTY(EditDefaultsOnly, BlueprintReadWrite, Category = AssetRegistry)
 	TArray<FName> PackagePaths;
 
+	/** The filter component containing specific object paths. Deprecated. */
+	UPROPERTY(EditDefaultsOnly, BlueprintReadWrite, Category = AssetRegistry, meta=(DeprecatedProperty, DeprecationMessage="Names containing full asset paths are deprecated. Use SoftObjectPaths instead."))
+	TArray<FName> ObjectPaths;
+
 	/** The filter component containing specific object paths */
 	UPROPERTY(EditDefaultsOnly, BlueprintReadWrite, Category = AssetRegistry)
-	TArray<FName> ObjectPaths;
+	TArray<FSoftObjectPath> SoftObjectPaths;
 
 	/** [DEPRECATED] - Class names are now represented by path names. Please use ClassPaths instead. */
 	UPROPERTY(BlueprintReadWrite, Category = AssetRegistry, meta=(DeprecatedProperty, DeprecationMessage="Short asset class names must be converted to full asset pathnames. Use ClassPaths instead."))
@@ -2519,8 +2527,11 @@ struct COREUOBJECT_API FAssetBundleEntry
 	FName BundleName;
 
 	/** List of string assets contained in this bundle */
-	UPROPERTY()
+	UPROPERTY(meta=(DeprecatedProperty, DeprecationMessage="Asset bundles may only contain top level asset paths which are referenced through the AssetPaths property."))
 	TArray<FSoftObjectPath> BundleAssets;
+
+	UPROPERTY()
+	TArray<FTopLevelAssetPath> AssetPaths;
 };
 
 /** A struct with a list of asset bundle entries. If one of these is inside a UObject it will get automatically exported as the asset registry tag AssetBundleData */
@@ -2539,8 +2550,7 @@ struct COREUOBJECT_API FAssetBundleData
 USTRUCT(noexport, BlueprintType, IsAlwaysAccessible, HasDefaults)
 struct FAssetData
 {
-	/** The object path for the asset in the form PackageName.AssetName. Only top level objects in a package can have AssetData */
-	UPROPERTY(BlueprintReadOnly, Category = AssetData, transient)
+	UPROPERTY(BlueprintReadOnly, Category = AssetData, transient, meta=(DeprecatedProperty, DeprecationMessage="Names containing full asset paths are deprecated. Use 'To Soft Object Path' instead."))
 	FName ObjectPath;
 	/** The name of the package in which the asset is found, this is the full long package name such as /Game/Path/Package */
 	UPROPERTY(BlueprintReadOnly, Category = AssetData, transient)
@@ -2560,6 +2570,12 @@ struct FAssetData
 
 	/** Asset package flags */
 	uint32 PackageFlags = 0;
+	
+#if WITH_EDITORONLY_DATA
+	/** If the outer of this object is not PackageName, it is the object specified by this path. Not exposed to blueprints except through 'To Soft Object Path'. */
+	FName OptionalOuterPath;
+#endif
+
 	/** The map of values for properties that were marked AssetRegistrySearchable or added by GetAssetRegistryTags */
 	FAssetDataTagMapSharedView TagsAndValues;
 	TSharedPtr<FAssetBundleData, ESPMode::ThreadSafe> TaggedAssetBundles;
diff --git a/Engine/Source/Runtime/CoreUObject/Public/UObject/SoftObjectPath.h b/Engine/Source/Runtime/CoreUObject/Public/UObject/SoftObjectPath.h
--- a/Engine/Source/Runtime/CoreUObject/Public/UObject/SoftObjectPath.h
+++ b/Engine/Source/Runtime/CoreUObject/Public/UObject/SoftObjectPath.h
@@ -2,12 +2,14 @@
 
 #pragma once
 
+#include "Algo/Transform.h"
 #include "Containers/Array.h"
 #include "Containers/Set.h"
 #include "Containers/StringFwd.h"
 #include "Containers/StringView.h"
 #include "Containers/UnrealString.h"
 #include "CoreMinimal.h"
+#include "Containers/StringView.h"
 #include "CoreTypes.h"
 #include "HAL/ThreadSafeCounter.h"
 #include "HAL/ThreadSingleton.h"
@@ -22,6 +24,7 @@
 #include "UObject/Object.h"
 #include "UObject/ObjectMacros.h"
 #include "UObject/ObjectPtr.h"
+#include "UObject/TopLevelAssetPath.h"
 #include "UObject/UObjectGlobals.h"
 #include "UObject/UObjectHash.h"
 
@@ -42,32 +45,47 @@ struct COREUOBJECT_API FSoftObjectPath
 
 	/** Construct from another soft object path */
 	FSoftObjectPath(const FSoftObjectPath& Other)
-		: AssetPathName(Other.AssetPathName)
+		: AssetPath(Other.AssetPath)
 		, SubPathString(Other.SubPathString)
 	{
 	}
 
 	/** Construct from a moveable soft object path */
 	FSoftObjectPath(FSoftObjectPath&& Other)
-		: AssetPathName(Other.AssetPathName)
+		: AssetPath(Other.AssetPath)
 		, SubPathString(MoveTemp(Other.SubPathString))
 	{
 	}
 
 	/** Construct from a path string. Non-explicit for backwards compatibility. */
 	FSoftObjectPath(const FString& Path)						{ SetPath(FStringView(Path)); }
+	FSoftObjectPath(FTopLevelAssetPath InAssetPath, FString InSubPathString)
+	{
+		SetPath(InAssetPath, MoveTemp(InSubPathString));
+	}
+	FSoftObjectPath(FName InPackageName, FName InAssetName, FString InSubPathString)
+	{
+		SetPath(FTopLevelAssetPath(InPackageName, InAssetName), MoveTemp(InSubPathString));
+	}
+	/** Explicitly extend a top-level object path with an empty subobject path. */
+	explicit FSoftObjectPath(FTopLevelAssetPath InAssetPath)	{ SetPath(InAssetPath, FString()); }
 	explicit FSoftObjectPath(FWideStringView Path)				{ SetPath(Path); }
 	explicit FSoftObjectPath(FAnsiStringView Path)				{ SetPath(Path); }
-	explicit FSoftObjectPath(FName Path)						{ SetPath(Path); }
 	explicit FSoftObjectPath(const WIDECHAR* Path)				{ SetPath(FWideStringView(Path)); }
 	explicit FSoftObjectPath(const ANSICHAR* Path)				{ SetPath(FAnsiStringView(Path)); }
 	explicit FSoftObjectPath(TYPE_OF_NULLPTR)					{}
-	
+
+	// UE_DEPRECATED(5.1, "Asset path FNames have been deprecated. This constructor should be used only temporarily to fix up old codepaths that produce an FName.")
+	explicit FSoftObjectPath(FName Path) 
+	{
+PRAGMA_DISABLE_DEPRECATION_WARNINGS;
+		SetPath(Path); 
+PRAGMA_ENABLE_DEPRECATION_WARNINGS;
+	}
+
 	/** Construct from an asset FName and subobject pair */
-	FSoftObjectPath(FName InAssetPathName, FString InSubPathString)
-		: AssetPathName(InAssetPathName)
-		, SubPathString(MoveTemp(InSubPathString))
-	{}
+	// UE_DEPRECATED(5.1, "Asset path FNames have been deprecated. This constructor should be used only temporarily to fix up old codepaths that produce an FName.")
+	FSoftObjectPath(FName InAssetPathName, FString InSubPathString);
 	
 	/** Construct from an existing object in memory */
 	FSoftObjectPath(const UObject* InObject);
@@ -82,43 +100,68 @@ struct COREUOBJECT_API FSoftObjectPath
 	{}
 
 	~FSoftObjectPath() {}
-	
+PRAGMA_DISABLE_DEPRECATION_WARNINGS	
 	FSoftObjectPath& operator=(const FSoftObjectPath& Path)	= default;
 	FSoftObjectPath& operator=(FSoftObjectPath&& Path) = default;
+PRAGMA_ENABLE_DEPRECATION_WARNINGS
+	FSoftObjectPath& operator=(const FTopLevelAssetPath Path)			{ SetPath(Path); return *this; }
 	FSoftObjectPath& operator=(const FString& Path)						{ SetPath(FStringView(Path)); return *this; }
 	FSoftObjectPath& operator=(FWideStringView Path)					{ SetPath(Path); return *this; }
 	FSoftObjectPath& operator=(FAnsiStringView Path)					{ SetPath(Path); return *this; }
-	FSoftObjectPath& operator=(FName Path)								{ SetPath(Path); return *this; }
 	FSoftObjectPath& operator=(const WIDECHAR* Path)					{ SetPath(FWideStringView(Path)); return *this; }
 	FSoftObjectPath& operator=(const ANSICHAR* Path)					{ SetPath(FAnsiStringView(Path)); return *this; }
 	FSoftObjectPath& operator=(TYPE_OF_NULLPTR)							{ Reset(); return *this; }
 
+	// UE_DEPRECATED(5.1, "Asset path FNames have been deprecated. This assignment operator should be used only temporarily to fix up old codepaths that produce an FName.")
+	FSoftObjectPath& operator=(FName Path) 
+	{
+PRAGMA_DISABLE_DEPRECATION_WARNINGS;
+		SetPath(Path); 
+PRAGMA_ENABLE_DEPRECATION_WARNINGS;
+		return *this;
+	}
+
 	/** Returns string representation of reference, in form /package/path.assetname[:subpath] */
 	FString ToString() const;
 
+	// UE_DEPRECATED(5.1, "Asset path FNames have been deprecated. This function should be used only temporarily to interface with old APIs that require an FName.")
+	FName ToFName() const { return *ToString(); }
+
 	/** Append string representation of reference, in form /package/path.assetname[:subpath] */
 	void ToString(FStringBuilderBase& Builder) const;
 
-	/** Returns the entire asset path as an FName, including both package and asset but not sub object */
-	FORCEINLINE FName GetAssetPathName() const
+	/** Append string representation of reference, in form /package/path.assetname[:subpath] */
+	void AppendString(FString& Builder) const;
+	void AppendString(FStringBuilderBase& Builder) const;
+
+	/** Returns the top-level asset part of this path, without the subobject path. */
+	FTopLevelAssetPath GetAssetPath() const
 	{
-		return AssetPathName;
+		return AssetPath;
 	}
 
-	FORCEINLINE void SetAssetPathName(FName InAssetPathName)
+	/** Returns this path without the SubPath component, restricting the result to a top level asset but keeping the type as FSoftObjectPath in contrast to GetAssetPath. */
+	FSoftObjectPath GetWithoutSubPath() const
 	{
-		AssetPathName = InAssetPathName;
+		return FSoftObjectPath(AssetPath, {});
 	}
+	
+	/** Returns the entire asset path as an FName, including both package and asset but not sub object */
+	// UE_DEPRECATED(5.1, "Asset path FNames have been deprecated. Use GetAssetPath instead.")
+	FName GetAssetPathName() const;
+	
+	// UE_DEPRECATED(5.1, "Asset path FNames have been deprecated. Use SetAssetPath instead.")
+	void SetAssetPathName(FName InAssetPathName);
 
 	/** Returns string version of asset path, including both package and asset but not sub object */
 	FORCEINLINE FString GetAssetPathString() const
 	{
-		if (AssetPathName.IsNone())
+		if (AssetPath.IsNull())
 		{
 			return FString();
 		}
 
-		return AssetPathName.ToString();
+		return AssetPath.ToString();
 	}
 
 	/** Returns the sub path, which is often empty */
@@ -133,18 +176,35 @@ struct COREUOBJECT_API FSoftObjectPath
 	}
 
 	/** Returns /package/path, leaving off the asset name and sub object */
-	FString GetLongPackageName() const;
+	FString GetLongPackageName() const
+	{
+		return GetAssetPath().GetPackageName().ToString();
+	}
 
 	/** Returns /package/path, leaving off the asset name and sub object */
-	FName GetLongPackageFName() const;
+	FName GetLongPackageFName() const
+	{
+		return GetAssetPath().GetPackageName();
+	}
+
+	/** Returns assetname string, leaving off the /package/path part and sub object */
+	FString GetAssetName() const
+	{
+		return GetAssetPath().GetAssetName().ToString();
+	}
 
 	/** Returns assetname string, leaving off the /package/path part and sub object */
-	FString GetAssetName() const;
+	FName GetAssetFName() const
+	{
+		return GetAssetPath().GetAssetName();
+	}
 
 	/** Sets asset path of this reference based on a string path */
+	void SetPath(const FTopLevelAssetPath& InAssetPath, FString InSubPathString = FString());
 	void SetPath(FWideStringView Path);
 	void SetPath(FAnsiStringView Path);
 	void SetPath(FUtf8StringView Path);
+	// UE_DEPRECATED(5.1, "Asset path FNames have been deprecated. This function should be used only temporarily to fix up old codepaths that produce an FName.")
 	void SetPath(FName Path);
 	void SetPath(const WIDECHAR* Path)			{ SetPath(FWideStringView(Path)); }
 	void SetPath(const ANSICHAR* Path)			{ SetPath(FAnsiStringView(Path)); }
@@ -167,41 +227,54 @@ struct COREUOBJECT_API FSoftObjectPath
 	/** Resets reference to point to null */
 	void Reset()
 	{		
-		AssetPathName = FName();
+		AssetPath.Reset();
 		SubPathString.Reset();
 	}
 	
 	/** Check if this could possibly refer to a real object, or was initialized to null */
 	FORCEINLINE bool IsValid() const
 	{
-		return !AssetPathName.IsNone();
+		return AssetPath.IsValid();
 	}
 
 	/** Checks to see if this is initialized to null */
 	FORCEINLINE bool IsNull() const
 	{
-		return AssetPathName.IsNone();
+		return AssetPath.IsNull();
 	}
 
 	/** Check if this represents an asset, meaning it is not null but does not have a sub path */
 	FORCEINLINE bool IsAsset() const
 	{
-		return !AssetPathName.IsNone() && SubPathString.IsEmpty();
+		return !AssetPath.IsNull() && SubPathString.IsEmpty();
 	}
 
 	/** Check if this represents a sub object, meaning it has a sub path */
 	FORCEINLINE bool IsSubobject() const
 	{
-		return !AssetPathName.IsNone() && !SubPathString.IsEmpty();
+		return !AssetPath.IsNull() && !SubPathString.IsEmpty();
 	}
 
+	/** Return true if this path appears before Other in lexical order */
 	bool LexicalLess(const FSoftObjectPath& Other) const
 	{
-		if (AssetPathName != Other.AssetPathName)
+		int32 PathCompare = AssetPath.Compare(Other.AssetPath);
+		if (PathCompare != 0)
+		{
+			return PathCompare < 0;
+		}
+		return SubPathString.Compare(Other.SubPathString) < 0; 
+	}
+
+	/** Return true if this path appears before Other using fast index-based fname order */
+	bool FastLess(const FSoftObjectPath& Other) const
+	{
+		int32 PathCompare = AssetPath.CompareFast(Other.AssetPath);
+		if (PathCompare != 0)
 		{
-			return AssetPathName.LexicalLess(Other.AssetPathName);
+			return PathCompare < 0;
 		}
-		return SubPathString.Compare(Other.SubPathString, ESearchCase::IgnoreCase) < 0;
+		return SubPathString.Compare(Other.SubPathString) < 0; 
 	}
 
 	/** Struct overrides */
@@ -242,7 +315,7 @@ struct COREUOBJECT_API FSoftObjectPath
 	{
 		uint32 Hash = 0;
 
-		Hash = HashCombine(Hash, GetTypeHash(This.AssetPathName));
+		Hash = HashCombine(Hash, GetTypeHash(This.AssetPath));
 		Hash = HashCombine(Hash, GetTypeHash(This.SubPathString));
 		return Hash;
 	}
@@ -266,7 +339,7 @@ struct COREUOBJECT_API FSoftObjectPath
 
 private:
 	/** Asset path, patch to a top level object in a package. This is /package/path.assetname */
-	FName AssetPathName;
+	FTopLevelAssetPath AssetPath;
 
 	/** Optional FString for subobject within an asset. This is the sub path after the : */
 	FString SubPathString;
@@ -288,12 +361,7 @@ struct FSoftObjectPathFastLess
 {
 	bool operator()(const FSoftObjectPath& Lhs, const FSoftObjectPath& Rhs) const
 	{
-		int32 Comp = Lhs.GetAssetPathName().CompareIndexes(Rhs.GetAssetPathName());
-		if (Comp < 0)
-		{
-			return true;
-		}
-		return Comp == 0 && Lhs.GetSubPathString() < Rhs.GetSubPathString();
+		return Lhs.FastLess(Rhs);
 	}
 };
 
@@ -302,12 +370,7 @@ struct FSoftObjectPathLexicalLess
 {
 	bool operator()(const FSoftObjectPath& Lhs, const FSoftObjectPath& Rhs) const
 	{
-		int32 Comp = Lhs.GetAssetPathName().Compare(Rhs.GetAssetPathName());
-		if (Comp < 0)
-		{
-			return true;
-		}
-		return Comp == 0 && Lhs.GetSubPathString() < Rhs.GetSubPathString();
+		return Lhs.LexicalLess(Rhs);
 	}
 };
 
@@ -491,7 +554,7 @@ struct FSoftObjectPathFixupArchive : public FArchiveUObject
 		{
 			if (!Value.IsNull() && Value.GetAssetPathString().Equals(OldAssetPathString, ESearchCase::IgnoreCase))
 			{
-				Value.SetAssetPathName(FName(*NewAssetPathString));
+				Value = FSoftObjectPath(NewAssetPathString);
 			}
 		})
 	{
@@ -511,3 +574,28 @@ struct FSoftObjectPathFixupArchive : public FArchiveUObject
 
 	TFunction<void(FSoftObjectPath&)> FixupFunction;
 };
+
+namespace UE::SoftObjectPath::Private
+{
+	// UE_DEPRECATED(5.1, "This function is only for use in fixing up deprecated APIs.")
+	inline TArray<FName> ConvertSoftObjectPaths(TConstArrayView<FSoftObjectPath> InPaths)
+	{
+		TArray<FName> Out;
+		Out.Reserve(InPaths.Num());
+PRAGMA_DISABLE_DEPRECATION_WARNINGS
+		Algo::Transform(InPaths, Out, [](const FSoftObjectPath& Path) { return Path.ToFName(); });
+PRAGMA_ENABLE_DEPRECATION_WARNINGS
+		return Out;
+	}
+
+	// UE_DEPRECATED(5.1, "This function is only for use in fixing up deprecated APIs.")
+	inline TArray<FSoftObjectPath> ConvertObjectPathNames(TConstArrayView<FName> InPaths)
+	{
+		TArray<FSoftObjectPath> Out;
+		Out.Reserve(InPaths.Num());
+PRAGMA_DISABLE_DEPRECATION_WARNINGS
+		Algo::Transform(InPaths, Out, [](FName Name) { return FSoftObjectPath(Name); });
+PRAGMA_ENABLE_DEPRECATION_WARNINGS
+		return Out;
+	}
+}
\ No newline at end of file
diff --git a/Engine/Source/Runtime/Engine/Classes/Kismet/KismetSystemLibrary.h b/Engine/Source/Runtime/Engine/Classes/Kismet/KismetSystemLibrary.h
--- a/Engine/Source/Runtime/Engine/Classes/Kismet/KismetSystemLibrary.h
+++ b/Engine/Source/Runtime/Engine/Classes/Kismet/KismetSystemLibrary.h
@@ -2083,8 +2083,21 @@ class ENGINE_API UKismetSystemLibrary : public UBlueprintFunctionLibrary
 	 * @param ClassNames [DEPRECATED] - Class names are now represented by path names. If non-empty, this input will result in a runtime warning. Please use the ClassPaths input instead.
 	 * @param RecursiveClassesExclusionSet [DEPRECATED] - Class names are now represented by path names. If non-empty, this input will result in a runtime warning. Please use the RecursiveClassPathsExclusionSet input instead.
 	 */
-	UFUNCTION(BlueprintPure, Category = "Utilities", meta = (Keywords = "construct build", NativeMakeFunc, BlueprintThreadSafe, AdvancedDisplay = "5", AutoCreateRefTerm = "PackageNames, PackagePaths, ObjectPaths, ClassNames, ClassPaths, RecursiveClassPathsExclusionSet, RecursiveClassesExclusionSet"))
-	static FARFilter MakeARFilter(const TArray<FName>& PackageNames, const TArray<FName>& PackagePaths, const TArray<FName>& ObjectPaths, const TArray<FTopLevelAssetPath>& ClassPaths, const TSet<FTopLevelAssetPath>& RecursiveClassPathsExclusionSet, const TArray<FName>& ClassNames, const TSet<FName>& RecursiveClassesExclusionSet, const bool bRecursivePaths = false, const bool bRecursiveClasses = false, const bool bIncludeOnlyOnDiskAssets = false);
+	UFUNCTION(BlueprintPure, Category = "Utilities", 
+		meta = (Keywords = "construct build", NativeMakeFunc, BlueprintThreadSafe, AdvancedDisplay = "5", 
+		AutoCreateRefTerm = "PackageNames, PackagePaths, SoftObjectPaths, ClassNames, ClassPaths, RecursiveClassPathsExclusionSet, RecursiveClassesExclusionSet"))
+	static FARFilter MakeARFilter(
+		const TArray<FName>& PackageNames, 
+		const TArray<FName>& PackagePaths, 
+		const TArray<FSoftObjectPath>& SoftObjectPaths, 
+		const TArray<FTopLevelAssetPath>& ClassPaths,
+		const TSet<FTopLevelAssetPath>& RecursiveClassPathsExclusionSet, 
+		const TArray<FName>& ClassNames, 
+		const TSet<FName>& RecursiveClassesExclusionSet, 
+		const bool bRecursivePaths = false, 
+		const bool bRecursiveClasses = false, 
+		const bool bIncludeOnlyOnDiskAssets = false
+		);
 
 	/**
 	 * Breaks an ARFilter struct into its component pieces. You should be using ClassPaths and RecursiveClassPathsExclusionSet from this node, ClassNames and RecursiveClassesExclusionSet are deprecated.
@@ -2093,7 +2106,19 @@ class ENGINE_API UKismetSystemLibrary : public UBlueprintFunctionLibrary
 	 * @param RecursiveClassesExclusionSet [DEPRECATED] - Class names are now represented by path names. Please use the RecursiveClassPathsExclusionSet output instead.
 	 */
 	UFUNCTION(BlueprintPure, Category = "Utilities", meta = (NativeBreakFunc, BlueprintThreadSafe, AdvancedDisplay = "6"))
-	static void BreakARFilter(FARFilter InARFilter, TArray<FName>& PackageNames, TArray<FName>& PackagePaths, TArray<FName>& ObjectPaths, TArray<FTopLevelAssetPath>& ClassPaths, TSet<FTopLevelAssetPath>& RecursiveClassPathsExclusionSet, TArray<FName>& ClassNames, TSet<FName>& RecursiveClassesExclusionSet, bool& bRecursivePaths, bool& bRecursiveClasses, bool& bIncludeOnlyOnDiskAssets);
+	static void BreakARFilter(
+		FARFilter InARFilter,
+		TArray<FName>& PackageNames,
+		TArray<FName>& PackagePaths,
+		TArray<FSoftObjectPath>& SoftObjectPaths,
+		TArray<FTopLevelAssetPath>& ClassPaths,
+		TSet<FTopLevelAssetPath>& RecursiveClassPathsExclusionSet,
+		TArray<FName>& ClassNames,
+		TSet<FName>& RecursiveClassesExclusionSet,
+		bool& bRecursivePaths,
+		bool& bRecursiveClasses,
+		bool& bIncludeOnlyOnDiskAssets
+		);
 
 };
 
diff --git a/Engine/Source/Runtime/Engine/Private/KismetSystemLibrary.cpp b/Engine/Source/Runtime/Engine/Private/KismetSystemLibrary.cpp
--- a/Engine/Source/Runtime/Engine/Private/KismetSystemLibrary.cpp
+++ b/Engine/Source/Runtime/Engine/Private/KismetSystemLibrary.cpp
@@ -3308,12 +3308,23 @@ void UKismetSystemLibrary::GetPrimaryAssetsWithBundleState(const TArray<FName>&
 	}
 }
 
-FARFilter UKismetSystemLibrary::MakeARFilter(const TArray<FName>& PackageNames, const TArray<FName>& PackagePaths, const TArray<FName>& ObjectPaths, const TArray<FTopLevelAssetPath>& ClassPaths, const TSet<FTopLevelAssetPath>& RecursiveClassPathsExclusionSet, const TArray<FName>& ClassNames, const TSet<FName>& RecursiveClassesExclusionSet, const bool bRecursivePaths, const bool bRecursiveClasses, const bool bIncludeOnlyOnDiskAssets)
+FARFilter UKismetSystemLibrary::MakeARFilter(
+	const TArray<FName>& PackageNames, 
+	const TArray<FName>& PackagePaths, 
+	const TArray<FSoftObjectPath>& SoftObjectPaths, 
+	const TArray<FTopLevelAssetPath>& ClassPaths,
+	const TSet<FTopLevelAssetPath>& RecursiveClassPathsExclusionSet, 
+	const TArray<FName>& ClassNames, 
+	const TSet<FName>& RecursiveClassesExclusionSet, 
+	const bool bRecursivePaths, 
+	const bool bRecursiveClasses, 
+	const bool bIncludeOnlyOnDiskAssets
+	)
 {
 	FARFilter NewFilter;
 	NewFilter.PackageNames = PackageNames;
 	NewFilter.PackagePaths = PackagePaths;
-	NewFilter.ObjectPaths = ObjectPaths;
+	NewFilter.SoftObjectPaths = SoftObjectPaths;
 	NewFilter.bRecursivePaths = bRecursivePaths;
 	NewFilter.bRecursiveClasses = bRecursiveClasses;
 	NewFilter.bIncludeOnlyOnDiskAssets = bIncludeOnlyOnDiskAssets;
@@ -3350,11 +3361,23 @@ FARFilter UKismetSystemLibrary::MakeARFilter(const TArray<FName>& PackageNames,
 	return NewFilter;
 }
 
-void UKismetSystemLibrary::BreakARFilter(FARFilter InARFilter, TArray<FName>& PackageNames, TArray<FName>& PackagePaths, TArray<FName>& ObjectPaths, TArray<FTopLevelAssetPath>& ClassPaths, TSet<FTopLevelAssetPath>& RecursiveClassPathsExclusionSet, TArray<FName>& ClassNames, TSet<FName>& RecursiveClassesExclusionSet, bool& bRecursivePaths, bool& bRecursiveClasses, bool& bIncludeOnlyOnDiskAssets)
+void UKismetSystemLibrary::BreakARFilter(
+	FARFilter InARFilter,
+	TArray<FName>& PackageNames,
+	TArray<FName>& PackagePaths,
+	TArray<FSoftObjectPath>& SoftObjectPaths,
+	TArray<FTopLevelAssetPath>& ClassPaths,
+	TSet<FTopLevelAssetPath>& RecursiveClassPathsExclusionSet,
+	TArray<FName>& ClassNames,
+	TSet<FName>& RecursiveClassesExclusionSet,
+	bool& bRecursivePaths,
+	bool& bRecursiveClasses,
+	bool& bIncludeOnlyOnDiskAssets
+	)
 {
 	PackageNames = InARFilter.PackageNames;
 	PackagePaths = InARFilter.PackagePaths;
-	ObjectPaths = InARFilter.ObjectPaths;
+	SoftObjectPaths = InARFilter.SoftObjectPaths;
 	ClassPaths = InARFilter.ClassPaths;
 	RecursiveClassPathsExclusionSet = InARFilter.RecursiveClassPathsExclusionSet;
 	bRecursivePaths = InARFilter.bRecursivePaths;
