Fix asset registry writer not propagating IsFilterEditorOnly properly from caller to FAssetData serialization.

diff --git a/Engine/Source/Runtime/AssetRegistry/Private/AssetRegistryArchive.cpp b/Engine/Source/Runtime/AssetRegistry/Private/AssetRegistryArchive.cpp
--- a/Engine/Source/Runtime/AssetRegistry/Private/AssetRegistryArchive.cpp
+++ b/Engine/Source/Runtime/AssetRegistry/Private/AssetRegistryArchive.cpp
@@ -176,6 +176,15 @@ FAssetRegistryWriter::FAssetRegistryWriter(const FAssetRegistryWriterOptions& Op
 , TargetAr(Out)
 {
 	check(!IsLoading());
+
+	// Copy requested serialization flags to intermediate archive. Technically the flags could change after this as TargetAr is a passed-in reference
+	SetArchiveState(TargetAr.GetArchiveState()); 
+
+	// The state copy explicitly clears this flag! 
+	SetFilterEditorOnly(TargetAr.IsFilterEditorOnly()); 
+
+	// The above function in FArchiveProxy seems broken - it only modifies the inner archive state, but this archive's state is what will be returned to serialization functions
+	FArchive::SetFilterEditorOnly(TargetAr.IsFilterEditorOnly()); 
 }
 
 static TArray<FDisplayNameEntryId> FlattenIndex(const TMap<FDisplayNameEntryId, uint32>& Names)
