Add header with state of editor-only-data-stripping to saved asset registry to allow properly loading cooked asset registries from a WITH_EDITORONLY_DATA binary.

diff --git a/Engine/Source/Runtime/AssetRegistry/Private/AssetDataGatherer.cpp b/Engine/Source/Runtime/AssetRegistry/Private/AssetDataGatherer.cpp
--- a/Engine/Source/Runtime/AssetRegistry/Private/AssetDataGatherer.cpp
+++ b/Engine/Source/Runtime/AssetRegistry/Private/AssetDataGatherer.cpp
@@ -3870,7 +3870,13 @@ FCachePayload LoadCacheFile(FStringView InCacheFilename)
 	{
 		int32 WorkerReduction = 2; // Current worker + preload task
 		int32 Parallelism = FMath::Max(FTaskGraphInterface::Get().GetNumWorkerThreads() - WorkerReduction, 0);
-		FAssetRegistryReader RegistryReader(ChecksummingReader, Parallelism);
+		
+		// The discovery cache is always serialized with a fixed format.
+		// We discard it before this point if it's not the latest version, and it always includes editor-only data.
+		FAssetRegistryHeader Header;
+		Header.Version = FAssetRegistryVersion::LatestVersion;
+		Header.bFilterEditorOnlyData = false;
+		FAssetRegistryReader RegistryReader(ChecksummingReader, Parallelism, Header);
 		return RegistryReader.IsError() ? FCachePayload() : SerializeCacheLoad(RegistryReader);
 	};
 
diff --git a/Engine/Source/Runtime/AssetRegistry/Private/AssetRegistryArchive.cpp b/Engine/Source/Runtime/AssetRegistry/Private/AssetRegistryArchive.cpp
--- a/Engine/Source/Runtime/AssetRegistry/Private/AssetRegistryArchive.cpp
+++ b/Engine/Source/Runtime/AssetRegistry/Private/AssetRegistryArchive.cpp
@@ -164,6 +164,19 @@ static TSharedPtr<FAssetBundleData, ESPMode::ThreadSafe> LoadBundlesOldVersion(F
 	return LoadBundlesInternal(Ar, Version);
 }
 
+void FAssetRegistryHeader::SerializeHeader(FArchive& Ar)
+{
+	FAssetRegistryVersion::SerializeVersion(Ar, Version);
+	if (Version >= FAssetRegistryVersion::AddedHeader)
+	{
+		Ar << bFilterEditorOnlyData;
+	}
+	else if(Ar.IsLoading())
+	{
+		bFilterEditorOnlyData = false;		
+	}
+}
+
 #if ALLOW_NAME_BATCH_SAVING
 
 FAssetRegistryWriterOptions::FAssetRegistryWriterOptions(const FAssetRegistrySerializationOptions& Options)
@@ -241,11 +254,14 @@ void FAssetRegistryWriter::SerializeTagsAndBundles(const FAssetData& Out)
 
 #endif
 
-FAssetRegistryReader::FAssetRegistryReader(FArchive& Inner, int32 NumWorkers, FAssetRegistryVersion::Type Version)
+FAssetRegistryReader::FAssetRegistryReader(FArchive& Inner, int32 NumWorkers, FAssetRegistryHeader Header)
 	: FArchiveProxy(Inner)
 {
 	check(IsLoading());
 
+	SetFilterEditorOnly(Header.bFilterEditorOnlyData);
+	FArchive::SetFilterEditorOnly(Header.bFilterEditorOnlyData); // Workaround for bug in FArchiveProxy
+
 	if (NumWorkers > 0)
 	{
 		TFunction<TArray<FDisplayNameEntryId> ()> GetFutureNames = LoadNameBatchAsync(*this, NumWorkers);
@@ -259,7 +275,7 @@ FAssetRegistryReader::FAssetRegistryReader(FArchive& Inner, int32 NumWorkers, FA
 	else
 	{
 		Names = LoadNameBatch(Inner);
-		Tags = FixedTagPrivate::LoadStore(*this, Version);
+		Tags = FixedTagPrivate::LoadStore(*this, Header.Version);
 	}
 }
 
@@ -386,7 +402,8 @@ bool FAssetRegistryTagSerializationTest::RunTest(const FString& Parameters)
 
 	{
 		FMemoryReader DataReader(Data);
-		FAssetRegistryReader RegistryReader(DataReader);
+		FAssetRegistryHeader Header;
+		FAssetRegistryReader RegistryReader(DataReader, 0, Header);
 		for (FAssetDataTagMapSharedView& FixedMap : FixedMaps)
 		{
 			FixedMap = LoadTags(RegistryReader);
@@ -398,7 +415,8 @@ bool FAssetRegistryTagSerializationTest::RunTest(const FString& Parameters)
 	// Re-create second fixed tag store to test operator==(FMapHandle, FMapHandle)
 	{
 		FMemoryReader DataReader(Data);
-		FAssetRegistryReader RegistryReader(DataReader);
+		FAssetRegistryHeader Header;
+		FAssetRegistryReader RegistryReader(DataReader, 0, Header);
 
 		for (const FAssetDataTagMapSharedView& FixedMap1 : FixedMaps)
 		{
diff --git a/Engine/Source/Runtime/AssetRegistry/Private/AssetRegistryArchive.h b/Engine/Source/Runtime/AssetRegistry/Private/AssetRegistryArchive.h
--- a/Engine/Source/Runtime/AssetRegistry/Private/AssetRegistryArchive.h
+++ b/Engine/Source/Runtime/AssetRegistry/Private/AssetRegistryArchive.h
@@ -9,12 +9,21 @@
 struct FAssetData;
 struct FAssetRegistrySerializationOptions;
 
+// Header containing versioning & data stripping information for a serialized asset registry state.
+struct FAssetRegistryHeader
+{
+	FAssetRegistryVersion::Type Version = FAssetRegistryVersion::LatestVersion;
+	bool bFilterEditorOnlyData = false; // True if editor only data has been removed/should be removed while saving 
+
+	void SerializeHeader(FArchive& Ar);
+};
+
 /** @see FAssetRegistryWriter */
 class FAssetRegistryReader : public FArchiveProxy
 {
 public:
 	/// @param NumWorkers > 0 for parallel loading
-	FAssetRegistryReader(FArchive& Inner, int32 NumWorkers = 0, FAssetRegistryVersion::Type Version = FAssetRegistryVersion::LatestVersion);
+	FAssetRegistryReader(FArchive& Inner, int32 NumWorkers, FAssetRegistryHeader Header);
 	~FAssetRegistryReader();
 
 	virtual FArchive& operator<<(FName& Value) override;
diff --git a/Engine/Source/Runtime/AssetRegistry/Private/AssetRegistryState.cpp b/Engine/Source/Runtime/AssetRegistry/Private/AssetRegistryState.cpp
--- a/Engine/Source/Runtime/AssetRegistry/Private/AssetRegistryState.cpp
+++ b/Engine/Source/Runtime/AssetRegistry/Private/AssetRegistryState.cpp
@@ -1043,8 +1043,10 @@ bool FAssetRegistryState::Save(FArchive& OriginalAr, const FAssetRegistrySeriali
 #else
 	check(CachedAssets.Num() == NumAssets);
 
-	FAssetRegistryVersion::Type Version = FAssetRegistryVersion::LatestVersion;
-	FAssetRegistryVersion::SerializeVersion(OriginalAr, Version);
+	FAssetRegistryHeader Header;
+	Header.Version = FAssetRegistryVersion::LatestVersion;
+	Header.bFilterEditorOnlyData = OriginalAr.IsFilterEditorOnly();
+	Header.SerializeHeader(OriginalAr);
 
 	// Set up fixed asset registry writer
 	FAssetRegistryWriter Ar(FAssetRegistryWriterOptions(Options), OriginalAr);
@@ -1163,28 +1165,28 @@ bool FAssetRegistryState::Load(FArchive& OriginalAr, const FAssetRegistryLoadOpt
 {
 	LLM_SCOPE(ELLMTag::AssetRegistry);
 
-	FAssetRegistryVersion::Type Version = FAssetRegistryVersion::LatestVersion;
-	FAssetRegistryVersion::SerializeVersion(OriginalAr, Version);
+	FAssetRegistryHeader Header;
+	Header.SerializeHeader(OriginalAr);
 	if (OutVersion != nullptr)
 	{
-		*OutVersion = Version;
+		*OutVersion = Header.Version;
 	}
 
 	FSoftObjectPathSerializationScope SerializationScope(NAME_None, NAME_None, ESoftObjectPathCollectType::NonPackage, ESoftObjectPathSerializeType::AlwaysSerialize);
 
-	if (Version < FAssetRegistryVersion::RemovedMD5Hash)
+	if (Header.Version < FAssetRegistryVersion::RemovedMD5Hash)
 	{
 		// Cannot read states before this version
 		return false;
 	}
-	else if (Version < FAssetRegistryVersion::FixedTags)
+	else if (Header.Version < FAssetRegistryVersion::FixedTags)
 	{
 		FNameTableArchiveReader NameTableReader(OriginalAr);
-		Load(NameTableReader, Version, Options);
+		Load(NameTableReader, Header, Options);
 	}
 	else
 	{
-		FAssetRegistryReader Reader(OriginalAr, Options.ParallelWorkers, Version);
+		FAssetRegistryReader Reader(OriginalAr, Options.ParallelWorkers, Header);
 
 		if (Reader.IsError())
 		{
@@ -1193,7 +1195,7 @@ bool FAssetRegistryState::Load(FArchive& OriginalAr, const FAssetRegistryLoadOpt
 
 		// Load won't resolve asset registry tag values loaded in parallel
 		// and can run before WaitForTasks
-		Load(Reader, Version, Options);
+		Load(Reader, Header, Options);
 
 		Reader.WaitForTasks();
 	}
@@ -1222,8 +1224,10 @@ bool FAssetRegistryState::Load(FArchive& OriginalAr, const FAssetRegistryLoadOpt
 }
 
 template<class Archive>
-void FAssetRegistryState::Load(Archive&& Ar, FAssetRegistryVersion::Type Version, const FAssetRegistryLoadOptions& Options)
+void FAssetRegistryState::Load(Archive&& Ar, const FAssetRegistryHeader& Header, const FAssetRegistryLoadOptions& Options)
 {
+	FAssetRegistryVersion::Type Version = Header.Version;
+
 	// serialize number of objects
 	int32 LocalNumAssets = 0;
 	Ar << LocalNumAssets;
diff --git a/Engine/Source/Runtime/AssetRegistry/Public/AssetRegistry/AssetRegistryState.h b/Engine/Source/Runtime/AssetRegistry/Public/AssetRegistry/AssetRegistryState.h
--- a/Engine/Source/Runtime/AssetRegistry/Public/AssetRegistry/AssetRegistryState.h
+++ b/Engine/Source/Runtime/AssetRegistry/Public/AssetRegistry/AssetRegistryState.h
@@ -33,6 +33,8 @@ namespace UE::AssetRegistry
 	class FAssetRegistryImpl;
 }
 
+struct FAssetRegistryHeader;
+
 /** Load/Save options used to modify how the cache is serialized. These are read out of the AssetRegistry section of Engine.ini and can be changed per platform. */
 struct FAssetRegistrySerializationOptions
 {
@@ -529,7 +531,7 @@ public:
 
 private:
 	template<class Archive>
-	void Load(Archive&& Ar, FAssetRegistryVersion::Type Version, const FAssetRegistryLoadOptions& Options);
+	void Load(Archive&& Ar, const FAssetRegistryHeader& Header, const FAssetRegistryLoadOptions& Options);
 
 	/** Initialize the lookup maps */
 	void SetAssetDatas(TArrayView<FAssetData> AssetDatas, const FAssetRegistryLoadOptions& Options);
diff --git a/Engine/Source/Runtime/CoreUObject/Public/AssetRegistry/AssetData.h b/Engine/Source/Runtime/CoreUObject/Public/AssetRegistry/AssetData.h
--- a/Engine/Source/Runtime/CoreUObject/Public/AssetRegistry/AssetData.h
+++ b/Engine/Source/Runtime/CoreUObject/Public/AssetRegistry/AssetData.h
@@ -93,6 +93,7 @@ struct COREUOBJECT_API FAssetRegistryVersion
 		AddedChunkHashes,					// Added FIoHash for each FIoChunkId in the package to the AssetPackageData.
 		ClassPaths,							// Classes are serialized as path names rather than short object names, e.g. /Script/Engine.StaticMesh
 		RemoveAssetPathFNames,				// Asset bundles are serialized as FTopLevelAssetPath instead of FSoftObjectPath, deprecated FAssetData::ObjectPath	
+		AddedHeader,						// Added header with bFilterEditorOnlyData flag 
 
 		// -----<new versions can be added above this line>-------------------------------------------------
 		VersionPlusOne,
