AssetRegistry - Fix GetAncestorClassNames and any other uses of CachedBPInheritanceMap to handle CoreRedirects.
Set TempCachedInheritanceBuffer.bDirty when native classes are modified by e.g. plugin loads to fix a possible hole in the AssetRegistry's map of native class inheritance.
#rnx

diff --git a/Engine/Source/Runtime/AssetRegistry/Private/AssetRegistry.cpp b/Engine/Source/Runtime/AssetRegistry/Private/AssetRegistry.cpp
--- a/Engine/Source/Runtime/AssetRegistry/Private/AssetRegistry.cpp
+++ b/Engine/Source/Runtime/AssetRegistry/Private/AssetRegistry.cpp
@@ -1298,7 +1298,7 @@ void FAssetRegistryImpl::CollectCodeGeneratorClasses()
 
 				if (!CachedBPInheritanceMap.Contains(GeneratedClassPathName))
 				{
-					CachedBPInheritanceMap.Add(GeneratedClassPathName, ParentClassPathName);
+					AddCachedBPClassParent(GeneratedClassPathName, ParentClassPathName);
 
 					// Invalidate caching because CachedBPInheritanceMap got modified
 					TempCachedInheritanceBuffer.bDirty = true;
@@ -1321,8 +1321,10 @@ namespace UE::AssetRegistry
 
 void FAssetRegistryImpl::RefreshNativeClasses()
 {
-	// Native classes have changed so reinitialize code generator and serialization options
+	// Native classes have changed so reinitialize code generator, class inheritance maps,
+	// and serialization options
 	CollectCodeGeneratorClasses();
+	TempCachedInheritanceBuffer.bDirty = true;
 
 	// Read default serialization options
 	Utils::InitializeSerializationOptionsFromIni(SerializationOptions, FString());
@@ -4002,7 +4004,7 @@ void FAssetRegistryImpl::CachePathsFromState(Impl::FEventContext& EventContext,
 
 				if (GeneratedClass && ParentClass)
 				{
-					CachedBPInheritanceMap.Add(GeneratedClass.ToTopLevelAssetPath(), ParentClass.ToTopLevelAssetPath());
+					AddCachedBPClassParent(GeneratedClass.ToTopLevelAssetPath(), ParentClass.ToTopLevelAssetPath());
 
 					// Invalidate caching because CachedBPInheritanceMap got modified
 					TempCachedInheritanceBuffer.bDirty = true;
@@ -4770,9 +4772,11 @@ void FAssetRegistryImpl::AddAssetData(Impl::FEventContext& EventContext, FAssetD
 		{
 			const FTopLevelAssetPath GeneratedClassPathName(GeneratedClass);
 			const FTopLevelAssetPath ParentClassPathName(ParentClass);
-			if (ensureAlwaysMsgf(!GeneratedClassPathName.IsNull() && !ParentClassPathName.IsNull(), TEXT("Short class names used in AddAssetData: GeneratedClass=%s, ParentClass=%s. Short class names in these tags on the Blueprint class should have been converted to path names."), *GeneratedClass, *ParentClass))
+			if (ensureAlwaysMsgf(!GeneratedClassPathName.IsNull() && !ParentClassPathName.IsNull(),
+				TEXT("Short class names used in AddAssetData: GeneratedClass=%s, ParentClass=%s. Short class names in these tags on the Blueprint class should have been converted to path names."),
+				*GeneratedClass, *ParentClass))
 			{
-				CachedBPInheritanceMap.Add(GeneratedClassPathName, ParentClassPathName);
+				AddCachedBPClassParent(GeneratedClassPathName, ParentClassPathName);
 
 				// Invalidate caching because CachedBPInheritanceMap got modified
 				TempCachedInheritanceBuffer.bDirty = true;
@@ -4795,7 +4799,9 @@ void FAssetRegistryImpl::UpdateAssetData(Impl::FEventContext& EventContext, FAss
 			if (!OldGeneratedClass.IsEmpty() && OldGeneratedClass != TEXTVIEW("None"))
 			{
 				const FTopLevelAssetPath OldGeneratedClassName(OldGeneratedClass);
-				if (ensureAlwaysMsgf(!OldGeneratedClassName.IsNull(), TEXT("Short class name used: OldGeneratedClass=%s. Short class names in tags on the Blueprint class should have been converted to path names."), *OldGeneratedClass))
+				if (ensureAlwaysMsgf(!OldGeneratedClassName.IsNull(),
+					TEXT("Short class name used: OldGeneratedClass=%s. Short class names in tags on the Blueprint class should have been converted to path names."),
+					*OldGeneratedClass))
 				{
 					CachedBPInheritanceMap.Remove(OldGeneratedClassName);
 
@@ -4808,9 +4814,11 @@ void FAssetRegistryImpl::UpdateAssetData(Impl::FEventContext& EventContext, FAss
 			{
 				const FTopLevelAssetPath NewGeneratedClassName(NewGeneratedClass);
 				const FTopLevelAssetPath NewParentClassName(NewParentClass);
-				if (ensureAlwaysMsgf(!NewGeneratedClassName.IsNull() && !NewParentClassName.IsNull(), TEXT("Short class names used in AddAssetData: GeneratedClass=%s, ParentClass=%s. Short class names in these tags on the Blueprint class should have been converted to path names."), *NewGeneratedClass, *NewParentClass))
+				if (ensureAlwaysMsgf(!NewGeneratedClassName.IsNull() && !NewParentClassName.IsNull(),
+					TEXT("Short class names used in AddAssetData: GeneratedClass=%s, ParentClass=%s. Short class names in these tags on the Blueprint class should have been converted to path names."),
+					*NewGeneratedClass, *NewParentClass))
 				{
-					CachedBPInheritanceMap.Add(NewGeneratedClassName, NewParentClassName);
+					AddCachedBPClassParent(NewGeneratedClassName, NewParentClassName);
 				}
 
 				// Invalidate caching because CachedBPInheritanceMap got modified
@@ -4846,7 +4854,8 @@ bool FAssetRegistryImpl::RemoveAssetData(Impl::FEventContext& EventContext, FAss
 			if (!OldGeneratedClass.IsEmpty() && OldGeneratedClass != TEXTVIEW("None"))
 			{
 				const FTopLevelAssetPath OldGeneratedClassPathName(FPackageName::ExportTextPathToObjectPath(OldGeneratedClass));
-				if (ensureAlwaysMsgf(!OldGeneratedClassPathName.IsNull(), TEXT("Short class name used: OldGeneratedClass=%s"), *OldGeneratedClass))
+				if (ensureAlwaysMsgf(!OldGeneratedClassPathName.IsNull(),
+					TEXT("Short class name used: OldGeneratedClass=%s"), *OldGeneratedClass))
 				{
 					CachedBPInheritanceMap.Remove(OldGeneratedClassPathName);
 
@@ -5560,6 +5569,31 @@ bool UAssetRegistryImpl::GetTemporaryCachingMode() const
 namespace UE::AssetRegistry
 {
 
+void FAssetRegistryImpl::AddCachedBPClassParent(const FTopLevelAssetPath& ClassPath, const FTopLevelAssetPath& NotYetRedirectedParentPath)
+{
+	// We do not check for CoreRedirects for ClassPath, because this function is only called on behalf of ClassPath being loaded,
+	// and the code author would have changed the package containing ClassPath to match the redirect they added.
+	// But we do need to check for CoreRedirects in the ParentPath, because when a parent class is renamed, we do not resave
+	// all packages containing subclasses to update their FBlueprintTags::ParentClassPath AssetData tags.
+	FTopLevelAssetPath ParentPath = NotYetRedirectedParentPath;
+#if WITH_EDITOR
+	FCoreRedirectObjectName RedirectedParentObjectName = FCoreRedirects::GetRedirectedName(ECoreRedirectFlags::Type_Class,
+		FCoreRedirectObjectName(NotYetRedirectedParentPath.GetAssetName(), NAME_None, NotYetRedirectedParentPath.GetPackageName()));
+	if (!RedirectedParentObjectName.OuterName.IsNone())
+	{
+		UE_LOG(LogAssetRegistry, Error,
+			TEXT("Class redirect exists from %s -> %s, which is invalid because ClassNames must be TopLevelAssetPaths. ")
+			TEXT("Redirect will be ignored in AssetRegistry queries."),
+			*NotYetRedirectedParentPath.ToString(), *RedirectedParentObjectName.ToString());
+	}
+	else
+	{
+		ParentPath = FTopLevelAssetPath(RedirectedParentObjectName.PackageName, RedirectedParentObjectName.ObjectName);
+	}
+#endif
+	CachedBPInheritanceMap.Add(ClassPath, ParentPath);
+}
+
 void FAssetRegistryImpl::UpdateInheritanceBuffer(Impl::FClassInheritanceBuffer& OutBuffer) const
 {
 	TRACE_CPUPROFILER_EVENT_SCOPE(UAssetRegistryImpl::UpdateTemporaryCaches)
@@ -5794,24 +5828,6 @@ void FAssetRegistryImpl::GetSubClasses_Recursive(Impl::FClassInheritanceContext&
 
 		// Add Subclasses of the given classname
 		AddSubClasses(InClassName);
-#if WITH_EDITOR
-		// Also add subclasses of any old names (from CoreRedirects) in case subclasses have not been resaved and have
-		// registered their ParentClass as the old name
-		TArray<FString> OldNames = FLinkerLoad::FindPreviousPathNamesForClass(InClassName.ToString(), false, true);
-		for (const FString& OldNameString : OldNames)
-		{
-			if (FPackageName::IsShortPackageName(OldNameString))
-			{
-				// TODO: We do not have a way to handle these, because we don't allow searching subclasses for short path
-				// names. We should instead change the construction of CachedBPInheritanceMap and ReverseInheritanceMap 
-				// to modify the ParentClassName with CoreRedirects
-				UE_LOG(LogAssetRegistry, Verbose, TEXT("GetDerivedClassNames found CoreRedirect from ShortPathName %s; it will not be able to find old blueprint classes that were saved with that parent class before the Redirect was added."),
-					*OldNameString);
-				continue;
-			}
-			AddSubClasses(FTopLevelAssetPath(OldNameString));
-		}
-#endif
 	}
 }
 
diff --git a/Engine/Source/Runtime/AssetRegistry/Private/AssetRegistryImpl.h b/Engine/Source/Runtime/AssetRegistry/Private/AssetRegistryImpl.h
--- a/Engine/Source/Runtime/AssetRegistry/Private/AssetRegistryImpl.h
+++ b/Engine/Source/Runtime/AssetRegistry/Private/AssetRegistryImpl.h
@@ -304,6 +304,9 @@ private:
 	/** Internal helper which processes a given state and adds its contents to the current registry */
 	void CachePathsFromState(Impl::FEventContext& EventContext, const FAssetRegistryState& InState);
 
+	/** Adds ClassPath-> ParentPath to CachedBPInheritanceMap. Applies CoreRedirects to NotYetRedirectedParentPath. */
+	void AddCachedBPClassParent(const FTopLevelAssetPath& ClassPath, const FTopLevelAssetPath& NotYetRedirectedParentPath);
+
 	/** Updates OutBuffer from loaded classes and registered-in-CachedBPInheritanceMap blueprint classes */
 	void UpdateInheritanceBuffer(Impl::FClassInheritanceBuffer& OutBuffer) const;
 
