Removing irrelevant comment (I cleaned it up eventually and forgot to remove the note to clean it up).

diff --git a/Engine/Source/Runtime/AssetRegistry/Private/AssetRegistryState.cpp b/Engine/Source/Runtime/AssetRegistry/Private/AssetRegistryState.cpp
--- a/Engine/Source/Runtime/AssetRegistry/Private/AssetRegistryState.cpp
+++ b/Engine/Source/Runtime/AssetRegistry/Private/AssetRegistryState.cpp
@@ -1899,7 +1899,6 @@ void FAssetRegistryState::UpdateAssetData(FAssetData* AssetData, FAssetData&& Ne
 		// so we can skip the more expensive compares if the inequality is already known
 		// This is not possible for keyfields - we have to take action on those even if inequality is already known -
 		// so we start with whether bKeyFieldIsModified
-		// TODO: this is horrifying, don't submit like this
 		*bOutModified = bKeyFieldIsModified ||
 			AssetData->AssetName != NewAssetData.AssetName ||
 			AssetData->PackageFlags != NewAssetData.PackageFlags ||
