Any thread can now safely register to AssetRegistry events

diff --git a/Engine/Source/Runtime/AssetRegistry/Private/AssetRegistry.cpp b/Engine/Source/Runtime/AssetRegistry/Private/AssetRegistry.cpp
--- a/Engine/Source/Runtime/AssetRegistry/Private/AssetRegistry.cpp
+++ b/Engine/Source/Runtime/AssetRegistry/Private/AssetRegistry.cpp
@@ -6315,67 +6315,56 @@ void UAssetRegistryImpl::Broadcast(UE::AssetRegistry::Impl::FEventContext& Event
 
 UAssetRegistryImpl::FPathAddedEvent& UAssetRegistryImpl::OnPathAdded()
 {
-	checkf(IsInGameThread(), TEXT("Registering to AssetRegistry events is not supported from multiple threads."));
 	return PathAddedEvent;
 }
 
 UAssetRegistryImpl::FPathRemovedEvent& UAssetRegistryImpl::OnPathRemoved()
 {
-	checkf(IsInGameThread(), TEXT("Registering to AssetRegistry events is not supported from multiple threads."));
 	return PathRemovedEvent;
 }
 
 UAssetRegistryImpl::FAssetAddedEvent& UAssetRegistryImpl::OnAssetAdded()
 {
-	checkf(IsInGameThread(), TEXT("Registering to AssetRegistry events is not supported from multiple threads."));
 	return AssetAddedEvent;
 }
 
 UAssetRegistryImpl::FAssetRemovedEvent& UAssetRegistryImpl::OnAssetRemoved()
 {
-	checkf(IsInGameThread(), TEXT("Registering to AssetRegistry events is not supported from multiple threads."));
 	return AssetRemovedEvent;
 }
 
 UAssetRegistryImpl::FAssetRenamedEvent& UAssetRegistryImpl::OnAssetRenamed()
 {
-	checkf(IsInGameThread(), TEXT("Registering to AssetRegistry events is not supported from multiple threads."));
 	return AssetRenamedEvent;
 }
 
 UAssetRegistryImpl::FAssetUpdatedEvent& UAssetRegistryImpl::OnAssetUpdated()
 {
-	checkf(IsInGameThread(), TEXT("Registering to AssetRegistry events is not supported from multiple threads."));
 	return AssetUpdatedEvent;
 }
 
 UAssetRegistryImpl::FAssetUpdatedEvent& UAssetRegistryImpl::OnAssetUpdatedOnDisk()
 {
-	checkf(IsInGameThread(), TEXT("Registering to AssetRegistry events is not supported from multiple threads."));
 	return AssetUpdatedOnDiskEvent;
 }
 
 UAssetRegistryImpl::FInMemoryAssetCreatedEvent& UAssetRegistryImpl::OnInMemoryAssetCreated()
 {
-	checkf(IsInGameThread(), TEXT("Registering to AssetRegistry events is not supported from multiple threads."));
 	return InMemoryAssetCreatedEvent;
 }
 
 UAssetRegistryImpl::FInMemoryAssetDeletedEvent& UAssetRegistryImpl::OnInMemoryAssetDeleted()
 {
-	checkf(IsInGameThread(), TEXT("Registering to AssetRegistry events is not supported from multiple threads."));
 	return InMemoryAssetDeletedEvent;
 }
 
 UAssetRegistryImpl::FFilesLoadedEvent& UAssetRegistryImpl::OnFilesLoaded()
 {
-	checkf(IsInGameThread(), TEXT("Registering to AssetRegistry events is not supported from multiple threads."));
 	return FileLoadedEvent;
 }
 
 UAssetRegistryImpl::FFileLoadProgressUpdatedEvent& UAssetRegistryImpl::OnFileLoadProgressUpdated()
 {
-	checkf(IsInGameThread(), TEXT("Registering to AssetRegistry events is not supported from multiple threads."));
 	return FileLoadProgressUpdatedEvent;
 }
 
diff --git a/Engine/Source/Runtime/AssetRegistry/Public/AssetRegistry/IAssetRegistry.h b/Engine/Source/Runtime/AssetRegistry/Public/AssetRegistry/IAssetRegistry.h
--- a/Engine/Source/Runtime/AssetRegistry/Public/AssetRegistry/IAssetRegistry.h
+++ b/Engine/Source/Runtime/AssetRegistry/Public/AssetRegistry/IAssetRegistry.h
@@ -518,7 +518,7 @@ public:
 	virtual void ScanModifiedAssetFiles(const TArray<FString>& InFilePaths) = 0;
 
 	/** Event for when paths are added to the registry */
-	DECLARE_EVENT_OneParam( IAssetRegistry, FPathAddedEvent, const FString& /*Path*/ );
+	DECLARE_TS_MULTICAST_DELEGATE_OneParam( FPathAddedEvent, const FString& /*Path*/ );
 	virtual FPathAddedEvent& OnPathAdded() = 0;
 
 	/** Event for when paths are removed from the registry */
@@ -544,34 +544,34 @@ public:
 	virtual void AssetTagsFinalized(const UObject& FinalizedAsset) = 0;
 
 	/** Event for when assets are added to the registry */
-	DECLARE_EVENT_OneParam( IAssetRegistry, FAssetAddedEvent, const FAssetData& );
+	DECLARE_TS_MULTICAST_DELEGATE_OneParam( FAssetAddedEvent, const FAssetData& );
 	virtual FAssetAddedEvent& OnAssetAdded() = 0;
 
 	/** Event for when assets are removed from the registry */
-	DECLARE_EVENT_OneParam( IAssetRegistry, FAssetRemovedEvent, const FAssetData& );
+	DECLARE_TS_MULTICAST_DELEGATE_OneParam( FAssetRemovedEvent, const FAssetData& );
 	virtual FAssetRemovedEvent& OnAssetRemoved() = 0;
 
 	/** Event for when assets are renamed in the registry */
-	DECLARE_EVENT_TwoParams( IAssetRegistry, FAssetRenamedEvent, const FAssetData&, const FString& );
+	DECLARE_TS_MULTICAST_DELEGATE_TwoParams( FAssetRenamedEvent, const FAssetData&, const FString& );
 	virtual FAssetRenamedEvent& OnAssetRenamed() = 0;
 
 	/** Event for when assets are updated in the registry */
-	DECLARE_EVENT_OneParam(IAssetRegistry, FAssetUpdatedEvent, const FAssetData&);
+	DECLARE_TS_MULTICAST_DELEGATE_OneParam( FAssetUpdatedEvent, const FAssetData&);
 	virtual FAssetUpdatedEvent& OnAssetUpdated() = 0;
 
 	/** Event for when assets are updated on disk and have been refreshed in the assetregistry */
 	virtual FAssetUpdatedEvent& OnAssetUpdatedOnDisk() = 0;
 
 	/** Event for when in-memory assets are created */
-	DECLARE_EVENT_OneParam( IAssetRegistry, FInMemoryAssetCreatedEvent, UObject* );
+	DECLARE_TS_MULTICAST_DELEGATE_OneParam( FInMemoryAssetCreatedEvent, UObject* );
 	virtual FInMemoryAssetCreatedEvent& OnInMemoryAssetCreated() = 0;
 
 	/** Event for when assets are deleted */
-	DECLARE_EVENT_OneParam( IAssetRegistry, FInMemoryAssetDeletedEvent, UObject* );
+	DECLARE_TS_MULTICAST_DELEGATE_OneParam( FInMemoryAssetDeletedEvent, UObject* );
 	virtual FInMemoryAssetDeletedEvent& OnInMemoryAssetDeleted() = 0;
 
 	/** Event for when the asset registry is done loading files */
-	DECLARE_EVENT( IAssetRegistry, FFilesLoadedEvent );
+	DECLARE_TS_MULTICAST_DELEGATE( FFilesLoadedEvent );
 	virtual FFilesLoadedEvent& OnFilesLoaded() = 0;
 
 	/** Payload data for a file progress update */
@@ -592,7 +592,7 @@ public:
 	};
 
 	/** Event to update the progress of the background file load */
-	DECLARE_EVENT_OneParam( IAssetRegistry, FFileLoadProgressUpdatedEvent, const FFileLoadProgressUpdateData& /*ProgressUpdateData*/ );
+	DECLARE_TS_MULTICAST_DELEGATE_OneParam( FFileLoadProgressUpdatedEvent, const FFileLoadProgressUpdateData& /*ProgressUpdateData*/ );
 	virtual FFileLoadProgressUpdatedEvent& OnFileLoadProgressUpdated() = 0;
 
 	/** Returns true if the asset registry is currently loading files and does not yet know about all assets */
