[Engine]
* Added UE_SCOPED_ENGINE_ACTIVITY in a couple places that takes noticable amount of time when starting monolith editor

diff --git a/Engine/Source/Developer/DerivedDataCache/Private/FileSystemCacheStore.cpp b/Engine/Source/Developer/DerivedDataCache/Private/FileSystemCacheStore.cpp
--- a/Engine/Source/Developer/DerivedDataCache/Private/FileSystemCacheStore.cpp
+++ b/Engine/Source/Developer/DerivedDataCache/Private/FileSystemCacheStore.cpp
@@ -1128,6 +1128,7 @@ bool FFileSystemCacheStore::RunSpeedTest(
 	double& OutWriteSpeedMBs)
 {
 	SCOPED_BOOT_TIMING("RunSpeedTest");
+	UE_SCOPED_ENGINE_ACTIVITY("Running IO speed test");
 
 	//  files of increasing size. Most DDC data falls within this range so we don't want to skew by reading
 	// large amounts of data. Ultimately we care most about latency anyway.
diff --git a/Engine/Source/Runtime/AssetRegistry/Private/AssetRegistry.cpp b/Engine/Source/Runtime/AssetRegistry/Private/AssetRegistry.cpp
--- a/Engine/Source/Runtime/AssetRegistry/Private/AssetRegistry.cpp
+++ b/Engine/Source/Runtime/AssetRegistry/Private/AssetRegistry.cpp
@@ -671,6 +671,8 @@ void FAssetRegistryImpl::LoadPremadeAssetRegistry(Impl::FEventContext& EventCont
 	Premade::ELoadResult LoadResult, FAssetRegistryState&& ARState, FAssetRegistryState::EInitializationMode Mode)
 {
 	SCOPED_BOOT_TIMING("LoadPremadeAssetRegistry");
+	UE_SCOPED_ENGINE_ACTIVITY("Loading premade asset registry");
+
 	if (SerializationOptions.bSerializeAssetRegistry)
 	{
 		if (LoadResult == Premade::ELoadResult::Succeeded)
diff --git a/Engine/Source/Runtime/CoreUObject/Private/UObject/UObjectGlobals.cpp b/Engine/Source/Runtime/CoreUObject/Private/UObject/UObjectGlobals.cpp
--- a/Engine/Source/Runtime/CoreUObject/Private/UObject/UObjectGlobals.cpp
+++ b/Engine/Source/Runtime/CoreUObject/Private/UObject/UObjectGlobals.cpp
@@ -1464,6 +1464,7 @@ UPackage* LoadPackageInternal(UPackage* InOuter, const FPackagePath& PackagePath
 
 		if (RequestID != INDEX_NONE)
 		{
+			UE_SCOPED_IO_ACTIVITY(*WriteToString<512>(TEXT("Sync "), PackagePath.GetDebugName()));
 			FlushAsyncLoading(RequestID);
 		}
 		GSyncLoadUsingAsyncLoaderCount--;
diff --git a/Engine/Source/Runtime/Engine/Private/DeviceProfiles/DeviceProfileManager.cpp b/Engine/Source/Runtime/Engine/Private/DeviceProfiles/DeviceProfileManager.cpp
--- a/Engine/Source/Runtime/Engine/Private/DeviceProfiles/DeviceProfileManager.cpp
+++ b/Engine/Source/Runtime/Engine/Private/DeviceProfiles/DeviceProfileManager.cpp
@@ -59,6 +59,7 @@ UDeviceProfileManager& UDeviceProfileManager::Get(bool bFromPostCDOContruct)
 {
 	if (DeviceProfileManagerSingleton == nullptr)
 	{
+		UE_SCOPED_ENGINE_ACTIVITY("Loading Device Profiles");
 		static bool bEntered = false;
 		if (bEntered && bFromPostCDOContruct)
 		{
