AssetRegistry - Provide FClassInheritanceContext to most writelock functions, so that it is available for the calculation of recursive class filters when calling LoadCalculatedDependencies.

diff --git a/Engine/Source/Runtime/AssetRegistry/Private/AssetRegistry.cpp b/Engine/Source/Runtime/AssetRegistry/Private/AssetRegistry.cpp
--- a/Engine/Source/Runtime/AssetRegistry/Private/AssetRegistry.cpp
+++ b/Engine/Source/Runtime/AssetRegistry/Private/AssetRegistry.cpp
@@ -632,6 +632,8 @@ struct FInitializeContext
 {
 	UAssetRegistryImpl& UARI;
 	FEventContext Events;
+	FClassInheritanceContext InheritanceContext;
+	FClassInheritanceBuffer InheritanceBuffer;
 	TArray<FString> RootContentPaths;
 	bool bRedirectorsNeedSubscribe = false;
 	bool bUpdateDiskCacheAfterLoad = false;
@@ -649,6 +651,9 @@ UAssetRegistryImpl::UAssetRegistryImpl(const FObjectInitializer& ObjectInitializ
 
 	{
 		FWriteScopeLock InterfaceScopeLock(InterfaceLock);
+		GetInheritanceContextWithRequiredLock(InterfaceScopeLock, Context.InheritanceContext,
+			Context.InheritanceBuffer);
+
 		GuardedData.Initialize(Context);
 		InitializeEvents(Context);
 	}
@@ -776,7 +781,7 @@ void FAssetRegistryImpl::Initialize(Impl::FInitializeContext& Context)
 		ConstructGatherer();
 		if (!GlobalGatherer->IsSynchronous())
 		{
-			SearchAllAssetsInitialAsync(Context.Events);
+			SearchAllAssetsInitialAsync(Context.Events, Context.InheritanceContext);
 		}
 		else
 		{
@@ -805,7 +810,7 @@ void FAssetRegistryImpl::Initialize(Impl::FInitializeContext& Context)
 		AddPath(Context.Events, AssetPath);
 	}
 
-	InitRedirectors(Context.Events, Context.bRedirectorsNeedSubscribe);
+	InitRedirectors(Context.Events, Context.InheritanceContext, Context.bRedirectorsNeedSubscribe);
 
 #if WITH_EDITOR
 	if (!UE::AssetDependencyGatherer::Private::FRegisteredAssetDependencyGatherer::IsEmpty())
@@ -948,7 +953,8 @@ bool FAssetRegistryImpl::ResolveRedirect(const FString& InPackageName, FString&
 	return false;
 }
 
-void FAssetRegistryImpl::InitRedirectors(Impl::FEventContext& EventContext, bool& bOutRedirectorsNeedSubscribe)
+void FAssetRegistryImpl::InitRedirectors(Impl::FEventContext& EventContext,
+	Impl::FClassInheritanceContext& InheritanceContext, bool& bOutRedirectorsNeedSubscribe)
 {
 	bOutRedirectorsNeedSubscribe = false;
 
@@ -983,7 +989,7 @@ void FAssetRegistryImpl::InitRedirectors(Impl::FEventContext& EventContext, bool
 			FString RootPackageName = FString::Printf(TEXT("/%s/"), *Plugin->GetName());
 			PathsToSearch.Add(RootPackageName);
 
-			Impl::FScanPathContext Context(EventContext, PathsToSearch, TArray<FString>());
+			Impl::FScanPathContext Context(EventContext, InheritanceContext, PathsToSearch, TArray<FString>());
 			ScanPathsSynchronous(Context);
 		}
 
@@ -1464,12 +1470,13 @@ void FAssetRegistryImpl::ConstructGatherer()
 	}
 }
 
-void FAssetRegistryImpl::SearchAllAssetsInitialAsync(Impl::FEventContext& EventContext)
+void FAssetRegistryImpl::SearchAllAssetsInitialAsync(Impl::FEventContext& EventContext,
+	Impl::FClassInheritanceContext& InheritanceContext)
 {
 	bInitialSearchStarted = true;
 	bInitialSearchCompleted = false;
 	FullSearchStartTime = FPlatformTime::Seconds();
-	SearchAllAssets(EventContext, false /* bSynchronousSearch */);
+	SearchAllAssets(EventContext, InheritanceContext, false /* bSynchronousSearch */);
 }
 
 }
@@ -1481,12 +1488,15 @@ void UAssetRegistryImpl::SearchAllAssets(bool bSynchronousSearch)
 	FEventContext EventContext;
 	{
 		FWriteScopeLock InterfaceScopeLock(InterfaceLock);
+		FClassInheritanceContext InheritanceContext;
+		FClassInheritanceBuffer InheritanceBuffer;
+		GetInheritanceContextWithRequiredLock(InterfaceScopeLock, InheritanceContext, InheritanceBuffer);
 		if (bSynchronousSearch)
 		{
 			// make sure any outstanding async preload is complete
 			GuardedData.ConditionalLoadPremadeAssetRegistry(*this, EventContext, InterfaceScopeLock);
 		}
-		GuardedData.SearchAllAssets(EventContext, bSynchronousSearch);
+		GuardedData.SearchAllAssets(EventContext, InheritanceContext, bSynchronousSearch);
 	}
 #if WITH_EDITOR
 	if (bSynchronousSearch)
@@ -1512,7 +1522,8 @@ bool UAssetRegistryImpl::IsSearchAsync() const
 namespace UE::AssetRegistry
 {
 
-void FAssetRegistryImpl::SearchAllAssets(Impl::FEventContext& EventContext, bool bSynchronousSearch)
+void FAssetRegistryImpl::SearchAllAssets(Impl::FEventContext& EventContext,
+	Impl::FClassInheritanceContext& InheritanceContext, bool bSynchronousSearch)
 {
 	LLM_SCOPE(ELLMTag::AssetRegistry);
 
@@ -1542,7 +1553,7 @@ void FAssetRegistryImpl::SearchAllAssets(Impl::FEventContext& EventContext, bool
 	{
 		Gatherer.WaitForIdle();
 		bool bUnusedInterrupted;
-		Impl::EGatherStatus UnusedStatus = TickGatherer(EventContext, -1., bUnusedInterrupted);
+		Impl::EGatherStatus UnusedStatus = TickGatherer(EventContext, InheritanceContext, -1., bUnusedInterrupted);
 #if WITH_EDITOR
 		if (!bInitialSearchStarted)
 		{
@@ -1571,8 +1582,12 @@ void UAssetRegistryImpl::WaitForCompletion()
 		EGatherStatus Status;
 		{
 			FWriteScopeLock InterfaceScopeLock(InterfaceLock);
+			FClassInheritanceContext InheritanceContext;
+			FClassInheritanceBuffer InheritanceBuffer;
+			GetInheritanceContextWithRequiredLock(InterfaceScopeLock, InheritanceContext, InheritanceBuffer);
+
 			bool bUnusedInterrupted;
-			Status = GuardedData.TickGatherer(EventContext, -1., bUnusedInterrupted);
+			Status = GuardedData.TickGatherer(EventContext, InheritanceContext, -1., bUnusedInterrupted);
 		}
 #if WITH_EDITOR
 		ProcessLoadedAssetsToUpdateCache(EventContext, -1., Status);
@@ -3256,11 +3271,15 @@ void UAssetRegistryImpl::ScanPathsSynchronousInternal(const TArray<FString>& InD
 	const double SearchStartTime = FPlatformTime::Seconds();
 
 	UE::AssetRegistry::Impl::FEventContext EventContext;
-	UE::AssetRegistry::Impl::FScanPathContext Context(EventContext, InDirs, InFiles,
+	UE::AssetRegistry::Impl::FClassInheritanceContext InheritanceContext;
+	UE::AssetRegistry::Impl::FClassInheritanceBuffer InheritanceBuffer;
+	UE::AssetRegistry::Impl::FScanPathContext Context(EventContext, InheritanceContext, InDirs, InFiles,
 		bInForceRescan, bInIgnoreDenyListScanFilters, nullptr /* OutFindAssets */);
 
 	{
 		FWriteScopeLock InterfaceScopeLock(InterfaceLock);
+		GetInheritanceContextWithRequiredLock(InterfaceScopeLock, InheritanceContext, InheritanceBuffer);
+
 		// make sure any outstanding async preload is complete
 		GuardedData.ConditionalLoadPremadeAssetRegistry(*this, EventContext, InterfaceScopeLock);
 		GuardedData.ScanPathsSynchronous(Context);
@@ -3543,6 +3562,9 @@ void UAssetRegistryImpl::Tick(float DeltaTime)
 		UE::AssetRegistry::Impl::FEventContext EventContext;
 		{
 			FWriteScopeLock InterfaceScopeLock(InterfaceLock);
+			UE::AssetRegistry::Impl::FClassInheritanceContext InheritanceContext;
+			UE::AssetRegistry::Impl::FClassInheritanceBuffer InheritanceBuffer;
+			GetInheritanceContextWithRequiredLock(InterfaceScopeLock, InheritanceContext, InheritanceBuffer);
 
 			// Process any deferred events and deletes
 			PRAGMA_DISABLE_DEPRECATION_WARNINGS
@@ -3553,7 +3575,7 @@ void UAssetRegistryImpl::Tick(float DeltaTime)
 			if (EventContext.IsEmpty())
 			{
 				// Tick the Gatherer
-				Status = GuardedData.TickGatherer(EventContext, TickStartTime, bInterrupted);
+				Status = GuardedData.TickGatherer(EventContext, InheritanceContext, TickStartTime, bInterrupted);
 			}
 			else
 			{
@@ -3587,7 +3609,8 @@ void FAssetRegistryImpl::TickDeletes()
 }
 
 Impl::EGatherStatus FAssetRegistryImpl::TickGatherer(Impl::FEventContext& EventContext,
-	const double TickStartTime, bool& bOutInterrupted, TOptional<FAssetsFoundCallback> AssetsFoundCallback)
+	Impl::FClassInheritanceContext& InheritanceContext, const double TickStartTime, bool& bOutInterrupted,
+	TOptional<FAssetsFoundCallback> AssetsFoundCallback)
 {
 	using namespace UE::AssetRegistry::Impl;
 
@@ -3707,7 +3730,7 @@ Impl::EGatherStatus FAssetRegistryImpl::TickGatherer(Impl::FEventContext& EventC
 	bool bDiskGatherComplete = !bIsSearching && NumGatherFromDiskPending == 0;
 	if (bDiskGatherComplete && PackagesNeedingDependencyCalculation.Num())
 	{
-		LoadCalculatedDependencies(nullptr, TickStartTime, bOutInterrupted);
+		LoadCalculatedDependencies(nullptr, TickStartTime, InheritanceContext, bOutInterrupted);
 		if (bOutInterrupted)
 		{
 			UpdateStatus(NumGatherFromDiskPending, true /* bInterrupted */);
@@ -3787,7 +3810,8 @@ void FAssetRegistryImpl::TickGatherPackage(Impl::FEventContext& EventContext, co
 }
 
 #if WITH_EDITOR
-void FAssetRegistryImpl::LoadCalculatedDependencies(TArray<FName>* AssetPackageNamesToCalculate, double TickStartTime, bool& bOutInterrupted)
+void FAssetRegistryImpl::LoadCalculatedDependencies(TArray<FName>* AssetPackageNamesToCalculate, double TickStartTime,
+	Impl::FClassInheritanceContext& InheritanceContext, bool& bOutInterrupted)
 {
 	auto CheckForTimeUp = [&TickStartTime, &bOutInterrupted](bool bHadActivity)
 	{
@@ -3814,7 +3838,7 @@ void FAssetRegistryImpl::LoadCalculatedDependencies(TArray<FName>* AssetPackageN
 			if (PackagesNeedingDependencyCalculation.Contains(PackageName))
 			{
 				bool bHadActivity;
-				LoadCalculatedDependencies(PackageName, bHadActivity);
+				LoadCalculatedDependencies(PackageName, InheritanceContext, bHadActivity);
 				if (CheckForTimeUp(bHadActivity))
 				{
 					return;
@@ -3827,7 +3851,7 @@ void FAssetRegistryImpl::LoadCalculatedDependencies(TArray<FName>* AssetPackageN
 		for (TSet<FName>::TIterator It = PackagesNeedingDependencyCalculation.CreateIterator(); It; ++It)
 		{
 			bool bHadActivity;
-			LoadCalculatedDependencies(*It, bHadActivity);
+			LoadCalculatedDependencies(*It, InheritanceContext, bHadActivity);
 			It.RemoveCurrent();
 			if (CheckForTimeUp(bHadActivity))
 			{
@@ -3838,14 +3862,13 @@ void FAssetRegistryImpl::LoadCalculatedDependencies(TArray<FName>* AssetPackageN
 	}
 }
 
-void FAssetRegistryImpl::LoadCalculatedDependencies(FName PackageName, bool& bOutHadActivity)
+void FAssetRegistryImpl::LoadCalculatedDependencies(FName PackageName,
+	Impl::FClassInheritanceContext& InheritanceContext, bool& bOutHadActivity)
 {
 	bOutHadActivity = false;
 	
-	auto GetCompiledFilter = [this](const FARFilter& InFilter) -> FARCompiledFilter
+	auto GetCompiledFilter = [this, &InheritanceContext](const FARFilter& InFilter) -> FARCompiledFilter
 	{
-		check(!InFilter.bRecursiveClasses);
-		UE::AssetRegistry::Impl::FClassInheritanceContext InheritanceContext;
 		FARCompiledFilter CompiledFilter;
 		CompileFilter(InheritanceContext, InFilter, CompiledFilter);
 		return CompiledFilter;
@@ -4190,9 +4213,11 @@ namespace UE::AssetRegistry
 namespace Impl
 {
 
-FScanPathContext::FScanPathContext(FEventContext& InEventContext, const TArray<FString>& InDirs, const TArray<FString>& InFiles,
-	bool bInForceRescan, bool bInIgnoreDenyListScanFilters, TArray<FSoftObjectPath>* FoundAssets)
+FScanPathContext::FScanPathContext(FEventContext& InEventContext, FClassInheritanceContext& InInheritanceContext, 
+	const TArray<FString>& InDirs, const TArray<FString>& InFiles, bool bInForceRescan,
+	bool bInIgnoreDenyListScanFilters, TArray<FSoftObjectPath>* FoundAssets)
 	: EventContext(InEventContext)
+	, InheritanceContext(InInheritanceContext)
 	, OutFoundAssets(FoundAssets)
 	, bForceRescan(bInForceRescan)
 	, bIgnoreDenyListScanFilters(bInIgnoreDenyListScanFilters)
@@ -4347,9 +4372,10 @@ void FAssetRegistryImpl::ScanPathsSynchronous(Impl::FScanPathContext& Context)
 	};
 
 	bool bUnusedInterrupted;
-	Context.Status = TickGatherer(Context.EventContext, -1., bUnusedInterrupted, FAssetsFoundCallback(AssetsFoundCallback));
+	Context.Status = TickGatherer(Context.EventContext, Context.InheritanceContext, -1., bUnusedInterrupted,
+		FAssetsFoundCallback(AssetsFoundCallback));
 #if WITH_EDITOR
-	LoadCalculatedDependencies(&FoundAssetPackageNames, -1., bUnusedInterrupted);
+	LoadCalculatedDependencies(&FoundAssetPackageNames, -1., Context.InheritanceContext, bUnusedInterrupted);
 #endif
 }
 
@@ -4955,7 +4981,10 @@ void UAssetRegistryImpl::OnDirectoryChanged(const TArray<FFileChangeData>& FileC
 	UE::AssetRegistry::Impl::FEventContext EventContext;
 	{
 		FWriteScopeLock InterfaceScopeLock(InterfaceLock);
-		GuardedData.OnDirectoryChanged(EventContext, FileChangesProcessed);
+		UE::AssetRegistry::Impl::FClassInheritanceContext InheritanceContext;
+		UE::AssetRegistry::Impl::FClassInheritanceBuffer InheritanceBuffer;
+		GetInheritanceContextWithRequiredLock(InterfaceScopeLock, InheritanceContext, InheritanceBuffer);
+		GuardedData.OnDirectoryChanged(EventContext, InheritanceContext, FileChangesProcessed);
 	}
 	Broadcast(EventContext);
 }
@@ -4963,7 +4992,8 @@ void UAssetRegistryImpl::OnDirectoryChanged(const TArray<FFileChangeData>& FileC
 namespace UE::AssetRegistry
 {
 
-void FAssetRegistryImpl::OnDirectoryChanged(Impl::FEventContext& EventContext, TArray<FFileChangeData>& FileChangesProcessed)
+void FAssetRegistryImpl::OnDirectoryChanged(Impl::FEventContext& EventContext,
+	Impl::FClassInheritanceContext& InheritanceContext, TArray<FFileChangeData>& FileChangesProcessed)
 {
 	TArray<FString> NewDirs;
 	TArray<FString> NewFiles;
@@ -5073,13 +5103,13 @@ void FAssetRegistryImpl::OnDirectoryChanged(Impl::FEventContext& EventContext, T
 			GlobalGatherer->OnFilesCreated(NewFiles);
 			if (GlobalGatherer->IsSynchronous())
 			{
-				Impl::FScanPathContext Context(EventContext, NewDirs, NewFiles,
+				Impl::FScanPathContext Context(EventContext, InheritanceContext, NewDirs, NewFiles,
 					false /* bForceRescan */, false /* bIgnoreDenyListScanFilters */, nullptr /* OutFoundAssets */);
 				ScanPathsSynchronous(Context);
 			}
 		}
 	}
-	ScanModifiedAssetFiles(EventContext, ModifiedFiles);
+	ScanModifiedAssetFiles(EventContext, InheritanceContext, ModifiedFiles);
 }
 
 }
@@ -5268,7 +5298,10 @@ void UAssetRegistryImpl::ScanModifiedAssetFiles(const TArray<FString>& InFilePat
 	UE::AssetRegistry::Impl::FEventContext EventContext;
 	{
 		FWriteScopeLock InterfaceScopeLock(InterfaceLock);
-		GuardedData.ScanModifiedAssetFiles(EventContext, InFilePaths);
+		UE::AssetRegistry::Impl::FClassInheritanceContext InheritanceContext;
+		UE::AssetRegistry::Impl::FClassInheritanceBuffer InheritanceBuffer;
+		GetInheritanceContextWithRequiredLock(InterfaceScopeLock, InheritanceContext, InheritanceBuffer);
+		GuardedData.ScanModifiedAssetFiles(EventContext, InheritanceContext, InFilePaths);
 	}
 
 #if WITH_EDITOR
@@ -5285,7 +5318,8 @@ void UAssetRegistryImpl::ScanModifiedAssetFiles(const TArray<FString>& InFilePat
 namespace UE::AssetRegistry
 {
 
-void FAssetRegistryImpl::ScanModifiedAssetFiles(Impl::FEventContext& EventContext, const TArray<FString>& InFilePaths)
+void FAssetRegistryImpl::ScanModifiedAssetFiles(Impl::FEventContext& EventContext,
+	Impl::FClassInheritanceContext& InheritanceContext, const TArray<FString>& InFilePaths)
 {
 	if (InFilePaths.Num() > 0)
 	{
@@ -5315,7 +5349,7 @@ void FAssetRegistryImpl::ScanModifiedAssetFiles(Impl::FEventContext& EventContex
 
 		// Re-scan and update the asset registry with the new asset data
 		TArray<FSoftObjectPath> FoundAssets;
-		Impl::FScanPathContext Context(EventContext, TArray<FString>(), InFilePaths,
+		Impl::FScanPathContext Context(EventContext, InheritanceContext, TArray<FString>(), InFilePaths,
 			true /* bForceRescan */, false /* bIgnoreDenyListScanFilters */, &FoundAssets);
 		ScanPathsSynchronous(Context);
 
@@ -5377,7 +5411,11 @@ void UAssetRegistryImpl::OnContentPathMounted(const FString& InAssetPath, const
 	UE::AssetRegistry::Impl::FEventContext EventContext;
 	{
 		FWriteScopeLock InterfaceScopeLock(InterfaceLock);
-		GuardedData.OnContentPathMounted(EventContext, InAssetPath, AssetPathWithTrailingSlash, FileSystemPath);
+		UE::AssetRegistry::Impl::FClassInheritanceContext InheritanceContext;
+		UE::AssetRegistry::Impl::FClassInheritanceBuffer InheritanceBuffer;
+		GetInheritanceContextWithRequiredLock(InterfaceScopeLock, InheritanceContext, InheritanceBuffer);
+		GuardedData.OnContentPathMounted(EventContext, InheritanceContext, InAssetPath, AssetPathWithTrailingSlash,
+			FileSystemPath);
 
 		// Listen for directory changes in this content path
 #if WITH_EDITOR
@@ -5405,7 +5443,9 @@ void UAssetRegistryImpl::OnContentPathMounted(const FString& InAssetPath, const
 namespace UE::AssetRegistry
 {
 
-void FAssetRegistryImpl::OnContentPathMounted(Impl::FEventContext& EventContext, const FString& InAssetPath, const FString& AssetPathWithTrailingSlash, const FString& FileSystemPath)
+void FAssetRegistryImpl::OnContentPathMounted(Impl::FEventContext& EventContext,
+	Impl::FClassInheritanceContext& InheritanceContext, const FString& InAssetPath,
+	const FString& AssetPathWithTrailingSlash, const FString& FileSystemPath)
 {
 	// Content roots always exist
 	AddPath(EventContext, AssetPathWithTrailingSlash);
@@ -5414,7 +5454,7 @@ void FAssetRegistryImpl::OnContentPathMounted(Impl::FEventContext& EventContext,
 	{
 		if (GlobalGatherer->IsSynchronous())
 		{
-			Impl::FScanPathContext Context(EventContext, { FileSystemPath }, TArray<FString>());
+			Impl::FScanPathContext Context(EventContext, InheritanceContext, { FileSystemPath }, TArray<FString>());
 			ScanPathsSynchronous(Context);
 		}
 		else
@@ -5675,10 +5715,6 @@ void UAssetRegistryImpl::GetInheritanceContextWithRequiredLock(FRWScopeLock& InO
 	UE::AssetRegistry::Impl::FClassInheritanceBuffer& StackBuffer)
 {
 	uint64 CurrentClassesVersionNumber = GetRegisteredClassesVersionNumber();
-	// If bIsTempCachingAlwaysEnabled, then we are guaranteed that bIsTempCachingEnabled=true.
-	// We rely on this to simplify logic and only check bIsTempCachingEnabled
-	check(!GuardedData.IsTempCachingAlwaysEnabled() || GuardedData.IsTempCachingEnabled());
-
 	bool bNeedsWriteLock = false;
 	if (GuardedData.GetClassGeneratorNamesRegisteredClassesVersionNumber() != CurrentClassesVersionNumber)
 	{
@@ -5697,6 +5733,25 @@ void UAssetRegistryImpl::GetInheritanceContextWithRequiredLock(FRWScopeLock& InO
 	}
 
 	// Note that we have to reread all data since we may have dropped the lock
+	GetInheritanceContextAfterVerifyingLock(CurrentClassesVersionNumber, InheritanceContext, StackBuffer);
+}
+
+void UAssetRegistryImpl::GetInheritanceContextWithRequiredLock(FWriteScopeLock& InOutScopeLock,
+	UE::AssetRegistry::Impl::FClassInheritanceContext& InheritanceContext,
+	UE::AssetRegistry::Impl::FClassInheritanceBuffer& StackBuffer)
+{
+	uint64 CurrentClassesVersionNumber = GetRegisteredClassesVersionNumber();
+	GetInheritanceContextAfterVerifyingLock(CurrentClassesVersionNumber, InheritanceContext, StackBuffer);
+}
+
+void UAssetRegistryImpl::GetInheritanceContextAfterVerifyingLock(uint64 CurrentClassesVersionNumber,
+	UE::AssetRegistry::Impl::FClassInheritanceContext& InheritanceContext,
+	UE::AssetRegistry::Impl::FClassInheritanceBuffer& StackBuffer)
+{
+	// If bIsTempCachingAlwaysEnabled, then we are guaranteed that bIsTempCachingEnabled=true.
+	// We rely on this to simplify logic and only check bIsTempCachingEnabled
+	check(!GuardedData.IsTempCachingAlwaysEnabled() || GuardedData.IsTempCachingEnabled());
+
 	bool bCodeGeneratorClassesUpToDate = GuardedData.GetClassGeneratorNamesRegisteredClassesVersionNumber() == CurrentClassesVersionNumber;
 	if (GuardedData.IsTempCachingEnabled())
 	{
diff --git a/Engine/Source/Runtime/AssetRegistry/Private/AssetRegistry.h b/Engine/Source/Runtime/AssetRegistry/Private/AssetRegistry.h
--- a/Engine/Source/Runtime/AssetRegistry/Private/AssetRegistry.h
+++ b/Engine/Source/Runtime/AssetRegistry/Private/AssetRegistry.h
@@ -239,6 +239,13 @@ private:
 	void GetInheritanceContextWithRequiredLock(FRWScopeLock& InOutScopeLock,
 		UE::AssetRegistry::Impl::FClassInheritanceContext& InheritanceContext,
 		UE::AssetRegistry::Impl::FClassInheritanceBuffer& StackBuffer);
+	void GetInheritanceContextWithRequiredLock(FWriteScopeLock& InOutScopeLock,
+		UE::AssetRegistry::Impl::FClassInheritanceContext& InheritanceContext,
+		UE::AssetRegistry::Impl::FClassInheritanceBuffer& StackBuffer);
+	void GetInheritanceContextAfterVerifyingLock(uint64 CurrentClassesVersionNumber,
+		UE::AssetRegistry::Impl::FClassInheritanceContext& InheritanceContext,
+		UE::AssetRegistry::Impl::FClassInheritanceBuffer& StackBuffer);
+
 
 #if WITH_EDITOR
 	/**
diff --git a/Engine/Source/Runtime/AssetRegistry/Private/AssetRegistryImpl.h b/Engine/Source/Runtime/AssetRegistry/Private/AssetRegistryImpl.h
--- a/Engine/Source/Runtime/AssetRegistry/Private/AssetRegistryImpl.h
+++ b/Engine/Source/Runtime/AssetRegistry/Private/AssetRegistryImpl.h
@@ -134,10 +134,12 @@ public:
 	void SetTemporaryCachingMode(bool bEnable);
 	void SetTemporaryCachingModeInvalidated();
 	bool AddPath(Impl::FEventContext& EventContext, const FString& PathToAdd);
-	void SearchAllAssets(Impl::FEventContext& EventContext, bool bSynchronousSearch);
+	void SearchAllAssets(Impl::FEventContext& EventContext, Impl::FClassInheritanceContext& InheritanceContext,
+		bool bSynchronousSearch);
 	void ScanPathsSynchronous(Impl::FScanPathContext& Context);
 	void PrioritizeSearchPath(const FString& PathToPrioritize);
-	void ScanModifiedAssetFiles(Impl::FEventContext& EventContext, const TArray<FString>& InFilePaths);
+	void ScanModifiedAssetFiles(Impl::FEventContext& EventContext, Impl::FClassInheritanceContext& InheritanceContext,
+		const TArray<FString>& InFilePaths);
 	void Serialize(FArchive& Ar, Impl::FEventContext& EventContext);
 	void AppendState(Impl::FEventContext& EventContext, const FAssetRegistryState& InState);
 	void GetAllocatedSize(bool bLogDetailed, SIZE_T& StateSize, SIZE_T& StaticSize, SIZE_T& SearchSize) const;
@@ -148,11 +150,12 @@ public:
 	bool SetPrimaryAssetIdForObjectPath(Impl::FEventContext& EventContext, const FSoftObjectPath& ObjectPath, FPrimaryAssetId PrimaryAssetId);
 	bool ResolveRedirect(const FString& InPackageName, FString& OutPackageName) const;
 #if WITH_EDITOR
-	void OnDirectoryChanged(Impl::FEventContext& EventContext, TArray<FFileChangeData>& FileChangesProcessed);
+	void OnDirectoryChanged(Impl::FEventContext& EventContext, Impl::FClassInheritanceContext& InheritanceContext,
+		TArray<FFileChangeData>& FileChangesProcessed);
 	void AddLoadedAssetToProcess(const UObject& AssetLoaded);
 #endif
-	void OnContentPathMounted(Impl::FEventContext& EventContext, const FString& InAssetPath,
-		const FString& AssetPathWithTrailingSlash, const FString& FileSystemPath);
+	void OnContentPathMounted(Impl::FEventContext& EventContext, Impl::FClassInheritanceContext& InheritanceContext,
+		const FString& InAssetPath, const FString& AssetPathWithTrailingSlash, const FString& FileSystemPath);
 	void OnContentPathDismounted(Impl::FEventContext& EventContext, const FString& InAssetPath,
 		const FString& AssetPathNoTrailingSlash, const FString& FileSystemPath);
 
@@ -177,7 +180,8 @@ public:
 	/** Callback type for TickGatherer */
 	typedef TFunctionRef<void(const TRingBuffer<FAssetData*>&)> FAssetsFoundCallback;
 	/** Consume any results from the gatherer and return its status */
-	Impl::EGatherStatus TickGatherer(Impl::FEventContext& EventContext, const double TickStartTime, bool& bOutInterrupted,
+	Impl::EGatherStatus TickGatherer(Impl::FEventContext& EventContext,
+		Impl::FClassInheritanceContext& InheritanceContext, const double TickStartTime, bool& bOutInterrupted,
 		TOptional<FAssetsFoundCallback> AssetsFoundCallback = TOptional<FAssetsFoundCallback>());
 	/** Look for and load a single AssetData result from the gatherer. */
 	void TickGatherPackage(Impl::FEventContext& EventContext, const FString& PackageName, const FString& LocalPath);
@@ -186,13 +190,15 @@ public:
 	void PushProcessLoadedAssetsBatch(Impl::FEventContext& EventContext,
 		TArrayView<FAssetData> LoadedAssetDatas, TArrayView<const UObject*> UnprocessedFromBatch);
 	/** Call LoadCalculatedDependencies on each Package updated after the last LoadCalculatedDependencies. */
-	void LoadCalculatedDependencies(TArray<FName>* AssetPackageNamessToCalculate, double TickStartTime, bool& bOutInterrupted);
+	void LoadCalculatedDependencies(TArray<FName>* AssetPackageNamessToCalculate, double TickStartTime, 
+		Impl::FClassInheritanceContext& InheritanceContext, bool& bOutInterrupted);
 	/**
 	 * Look for a CalculatedDependencies function registered for the asset(s) in the given package
 	 * and call that function to add calculated dependencies. Calculated dependencies are added only after
 	 * all normal dependencies gathered from the AssetRegistry data stored in the package have been loaded.
 	 */
-	void LoadCalculatedDependencies(FName PackageName, bool& bOutHadActivity);
+	void LoadCalculatedDependencies(FName PackageName, Impl::FClassInheritanceContext& InheritanceContext,
+		bool& bOutHadActivity);
 	/** Add a watch on a directory that to modify data for a package whenever packages in the directory are modified. */
 	void AddDirectoryReferencer(FName PackageName, const FString& DirectoryLocalPathOrLongPackageName);
 	/** Remove all directory watches for PackageName. */
@@ -253,13 +259,14 @@ private:
 	void EnumerateAssetsByPathNoTags(FName PackagePath,
 		TFunctionRef<bool(const FAssetData&)> Callback, bool bRecursive, bool bIncludeOnlyOnDiskAssets) const;
 	/** Create redirectors read by this->GetRedirectedObjectPath, based on plugin settings */
-	void InitRedirectors(Impl::FEventContext& EventContext, bool& bOutRedirectorsNeedSubscribe);
+	void InitRedirectors(Impl::FEventContext& EventContext, Impl::FClassInheritanceContext& InheritanceContext,
+		bool& bOutRedirectorsNeedSubscribe);
 	/** If collecting dependencies, create an FAssetPackageData for every script package, to make sure dependencies can find a data for them. */
 	void ReadScriptPackages();
 	/* Construct the gatherer if it does not already exist */
 	void ConstructGatherer();
 	/**  Called to set up timing variables and launch the on-constructor SearchAllAssets async call */
-	void SearchAllAssetsInitialAsync(Impl::FEventContext& EventContext);
+	void SearchAllAssetsInitialAsync(Impl::FEventContext& EventContext, Impl::FClassInheritanceContext& InheritanceContext);
 	/**
 	 * Called every tick to when data is retrieved by the background asset search.
 	 * If TickStartTime is < 0, the entire list of gathered assets will be cached. Also used in sychronous searches
@@ -489,8 +496,9 @@ struct FClassInheritanceContext
 /** Context to transform, collect, and pass along ScanPathsSynchronous arguments and results */
 struct FScanPathContext
 {
-	FScanPathContext(FEventContext& InEventContext, const TArray<FString>& InDirs, const TArray<FString>& InFiles,
-		bool bInForceRescan = false, bool bInIgnoreDenyListScanFilters = false, TArray<FSoftObjectPath>* FoundAssets = nullptr);
+	FScanPathContext(FEventContext& InEventContext, FClassInheritanceContext& InInheritanceContext,
+		const TArray<FString>& InDirs, const TArray<FString>& InFiles, bool bInForceRescan = false,
+		bool bInIgnoreDenyListScanFilters = false, TArray<FSoftObjectPath>* FoundAssets = nullptr);
 
 	TArray<FString> PackageDirs;
 	TArray<FString> LocalDirs;
@@ -498,6 +506,7 @@ struct FScanPathContext
 	TArray<FString> LocalFiles;
 	TArray<FString> LocalPaths;
 	FEventContext& EventContext;
+	FClassInheritanceContext& InheritanceContext;
 	TArray<FSoftObjectPath>* OutFoundAssets = nullptr;
 	int32 NumFoundAssets = 0;
 	bool bForceRescan = false;
