[AssetRegistry] quick fix for merge conflict of new code and deprecated code

diff --git a/Engine/Source/Runtime/AssetRegistry/Private/AssetRegistry.cpp b/Engine/Source/Runtime/AssetRegistry/Private/AssetRegistry.cpp
--- a/Engine/Source/Runtime/AssetRegistry/Private/AssetRegistry.cpp
+++ b/Engine/Source/Runtime/AssetRegistry/Private/AssetRegistry.cpp
@@ -2086,7 +2086,9 @@ UE::AssetRegistry::EExists UAssetRegistryImpl::TryGetAssetByObjectPath(const FNa
 	FReadScopeLock InterfaceScopeLock(InterfaceLock);
 	bool bAssetRegistryReady = GuardedData.IsInitialSearchStarted() && GuardedData.IsInitialSearchCompleted();
 	const FAssetRegistryState& State = GuardedData.GetState();
-	const FAssetData* FoundData = State.GetAssetByObjectPath(ObjectPath);
+	PRAGMA_DISABLE_DEPRECATION_WARNINGS;
+	const FAssetData* FoundData = State.GetAssetByObjectPath(FSoftObjectPath(ObjectPath.ToString()));
+	PRAGMA_ENABLE_DEPRECATION_WARNINGS;
 	if (!FoundData)
 	{
 		if (!bAssetRegistryReady)
