AssetRegistry: PackageReader: Downgrade !IsFileVersionValid case from Error to Warning so that if the package is unreferenced a cook can still succeed without errors.
#rnx

diff --git a/Engine/Source/Runtime/AssetRegistry/Private/PackageReader.cpp b/Engine/Source/Runtime/AssetRegistry/Private/PackageReader.cpp
--- a/Engine/Source/Runtime/AssetRegistry/Private/PackageReader.cpp
+++ b/Engine/Source/Runtime/AssetRegistry/Private/PackageReader.cpp
@@ -91,7 +91,8 @@ bool FPackageReader::OpenPackageFile(EOpenPackageResult* OutErrorCode)
 
 	if (!PackageFileSummary.IsFileVersionValid())
 	{
-		UE_LOG(LogAssetRegistry, Error, TEXT("Package %s is unversioned which cannot be opened by the current process"), *PackageFilename);
+		// Log a warning rather than an error. Linkerload gracefully handles this case.
+		UE_LOG(LogAssetRegistry, Warning, TEXT("Package %s is unversioned which cannot be opened by the current process"), *PackageFilename);
 		SetPackageErrorCode(EOpenPackageResult::Unversioned);
 		return false;
 	}
@@ -99,17 +100,19 @@ bool FPackageReader::OpenPackageFile(EOpenPackageResult* OutErrorCode)
 	// Don't read packages that are too old
 	if (PackageFileSummary.IsFileVersionTooOld())
 	{
-		UE_LOG(	LogAssetRegistry, Error, TEXT("Package %s is too old. Min Version: %i  Package Version: %i"), 
+		// Log a warning rather than an error. Linkerload gracefully handles this case.
+		UE_LOG(	LogAssetRegistry, Warning, TEXT("Package %s is too old. Min Version: %i  Package Version: %i"),
 				*PackageFilename, (int32)VER_UE4_OLDEST_LOADABLE_PACKAGE, PackageFileSummary.GetFileVersionUE().FileVersionUE4);
 
 		SetPackageErrorCode(EOpenPackageResult::Unversioned);
 		return false;
 	}
 
-	// Don't read packages that were saved with an package version newer than the current one.
+	// Don't read packages that were saved with a package version newer than the current one.
 	if (PackageFileSummary.IsFileVersionTooNew())
 	{
-		UE_LOG(	LogAssetRegistry, Error, TEXT("Package %s is too new. Engine Version: %i  Package Version: %i"), 
+		// Log a warning rather than an error. Linkerload gracefully handles this case.
+		UE_LOG(	LogAssetRegistry, Warning, TEXT("Package %s is too new. Engine Version: %i  Package Version: %i"),
 				*PackageFilename, GPackageFileUEVersion.ToValue(), PackageFileSummary.GetFileVersionUE().ToValue());
 
 		SetPackageErrorCode(EOpenPackageResult::VersionTooNew);
@@ -118,7 +121,8 @@ bool FPackageReader::OpenPackageFile(EOpenPackageResult* OutErrorCode)
 
 	if (PackageFileSummary.GetFileVersionLicenseeUE() > GPackageFileLicenseeUEVersion)
 	{
-		UE_LOG(LogAssetRegistry, Error, TEXT("Package %s is too new. Licensee Version: %i Package Licensee Version: %i"),
+		// Log a warning rather than an error. Linkerload gracefully handles this case.
+		UE_LOG(LogAssetRegistry, Warning, TEXT("Package %s is too new. Licensee Version: %i Package Licensee Version: %i"),
 			*PackageFilename, GPackageFileLicenseeUEVersion, PackageFileSummary.GetFileVersionLicenseeUE());
 
 		SetPackageErrorCode(EOpenPackageResult::VersionTooNew);
