Asset registry manages source files

#rnx

diff --git a/Engine/Source/Runtime/AssetRegistry/Private/AssetDataGatherer.cpp b/Engine/Source/Runtime/AssetRegistry/Private/AssetDataGatherer.cpp
--- a/Engine/Source/Runtime/AssetRegistry/Private/AssetDataGatherer.cpp
+++ b/Engine/Source/Runtime/AssetRegistry/Private/AssetDataGatherer.cpp
@@ -82,31 +82,34 @@ void AssignStringWithoutShrinking(FString& InOutResult, FStringView Value)
 	}
 }
 
-FDiscoveredPathData::FDiscoveredPathData(FStringView InLocalAbsPath, FStringView InLongPackageName, FStringView InRelPath, const FDateTime& InPackageTimestamp)
+FDiscoveredPathData::FDiscoveredPathData(FStringView InLocalAbsPath, FStringView InLongPackageName, FStringView InRelPath, const FDateTime& InPackageTimestamp, EGatherableFileType InType)
 	: LocalAbsPath(InLocalAbsPath)
 	, LongPackageName(InLongPackageName)
 	, RelPath(InRelPath)
 	, PackageTimestamp(InPackageTimestamp)
+	, Type(InType)
 {
 }
 
-FDiscoveredPathData::FDiscoveredPathData(FStringView InLocalAbsPath, FStringView InLongPackageName, FStringView InRelPath)
+FDiscoveredPathData::FDiscoveredPathData(FStringView InLocalAbsPath, FStringView InLongPackageName, FStringView InRelPath, EGatherableFileType InType)
 	: LocalAbsPath(InLocalAbsPath)
 	, LongPackageName(InLongPackageName)
 	, RelPath(InRelPath)
+	, Type(InType)
 {
 }
 
-void FDiscoveredPathData::Assign(FStringView InLocalAbsPath, FStringView InLongPackageName, FStringView InRelPath)
+void FDiscoveredPathData::Assign(FStringView InLocalAbsPath, FStringView InLongPackageName, FStringView InRelPath, EGatherableFileType InType)
 {
 	AssignStringWithoutShrinking(LocalAbsPath, InLocalAbsPath);
 	AssignStringWithoutShrinking(LongPackageName, InLongPackageName);
 	AssignStringWithoutShrinking(RelPath, InRelPath);
+	Type = InType;
 }
 
-void FDiscoveredPathData::Assign(FStringView InLocalAbsPath, FStringView InLongPackageName, FStringView InRelPath, const FDateTime& InPackageTimestamp)
+void FDiscoveredPathData::Assign(FStringView InLocalAbsPath, FStringView InLongPackageName, FStringView InRelPath, const FDateTime& InPackageTimestamp, EGatherableFileType InType)
 {
-	Assign(InLocalAbsPath, InLongPackageName, InRelPath);
+	Assign(InLocalAbsPath, InLongPackageName, InRelPath, InType);
 	PackageTimestamp = InPackageTimestamp;
 }
 
@@ -115,15 +118,16 @@ SIZE_T FDiscoveredPathData::GetAllocatedSize() const
 	return LocalAbsPath.GetAllocatedSize() + LongPackageName.GetAllocatedSize() + RelPath.GetAllocatedSize();
 }
 
-FGatheredPathData::FGatheredPathData(FStringView InLocalAbsPath, FStringView InLongPackageName, const FDateTime& InPackageTimestamp)
+FGatheredPathData::FGatheredPathData(FStringView InLocalAbsPath, FStringView InLongPackageName, const FDateTime& InPackageTimestamp, EGatherableFileType InType)
 	: LocalAbsPath(InLocalAbsPath)
 	, LongPackageName(InLongPackageName)
 	, PackageTimestamp(InPackageTimestamp)
+	, Type(InType)
 {
 }
 
 FGatheredPathData::FGatheredPathData(const FDiscoveredPathData& DiscoveredData)
-	:FGatheredPathData(DiscoveredData.LocalAbsPath, DiscoveredData.LongPackageName, DiscoveredData.PackageTimestamp)
+	: FGatheredPathData(DiscoveredData.LocalAbsPath, DiscoveredData.LongPackageName, DiscoveredData.PackageTimestamp, DiscoveredData.Type)
 {
 }
 
@@ -131,19 +135,21 @@ FGatheredPathData::FGatheredPathData(FDiscoveredPathData&& DiscoveredData)
 	: LocalAbsPath(MoveTemp(DiscoveredData.LocalAbsPath))
 	, LongPackageName(MoveTemp(DiscoveredData.LongPackageName))
 	, PackageTimestamp(MoveTemp(DiscoveredData.PackageTimestamp))
+	, Type(DiscoveredData.Type)
 {
 }
 
-void FGatheredPathData::Assign(FStringView InLocalAbsPath, FStringView InLongPackageName, const FDateTime& InPackageTimestamp)
+void FGatheredPathData::Assign(FStringView InLocalAbsPath, FStringView InLongPackageName, const FDateTime& InPackageTimestamp, EGatherableFileType InType)
 {
 	AssignStringWithoutShrinking(LocalAbsPath, InLocalAbsPath);
 	AssignStringWithoutShrinking(LongPackageName, InLongPackageName);
 	PackageTimestamp = InPackageTimestamp;
+	Type = InType;
 }
 
 void FGatheredPathData::Assign(const FDiscoveredPathData& DiscoveredData)
 {
-	Assign(DiscoveredData.LocalAbsPath, DiscoveredData.LongPackageName, DiscoveredData.PackageTimestamp);
+	Assign(DiscoveredData.LocalAbsPath, DiscoveredData.LongPackageName, DiscoveredData.PackageTimestamp, DiscoveredData.Type);
 }
 
 SIZE_T FGatheredPathData::GetAllocatedSize() const
@@ -1548,22 +1554,26 @@ void FAssetDataDiscovery::TickInternal()
 						check(IteratedSubDirs.Num() == NumIteratedDirs);
 						IteratedSubDirs.Emplace();
 					}
-					IteratedSubDirs[NumIteratedDirs++].Assign(LocalAbsPath, DirLongPackageName, RelPath);
+					IteratedSubDirs[NumIteratedDirs++].Assign(LocalAbsPath, DirLongPackageName, RelPath, EGatherableFileType::Directory);
 				}
 			}
 			else
 			{
-				FStringView BaseName = FPathViews::GetBaseFilename(RelPath);
-				FPathViews::AppendPath(DirLongPackageName, BaseName);
+				EGatherableFileType FileType = GetFileType(RelPath);
 				// Don't record files that contain invalid packagepath characters (not counting their extension) or that do not end with a recognized extension
-				if (!FPackageName::DoesPackageNameContainInvalidCharacters(BaseName) && FPackageName::IsPackageFilename(RelPath))
+				if (FileType != EGatherableFileType::Invalid)
 				{
-					if (IteratedFiles.Num() < NumIteratedFiles + 1)
+					FStringView BaseName = FPathViews::GetBaseFilename(RelPath);
+					if (!FPackageName::DoesPackageNameContainInvalidCharacters(BaseName))
 					{
-						check(IteratedFiles.Num() == NumIteratedFiles);
-						IteratedFiles.Emplace();
+						if (IteratedFiles.Num() < NumIteratedFiles + 1)
+						{
+							check(IteratedFiles.Num() == NumIteratedFiles);
+							IteratedFiles.Emplace();
+						}
+						FPathViews::AppendPath(DirLongPackageName, BaseName);
+						IteratedFiles[NumIteratedFiles++].Assign(LocalAbsPath, DirLongPackageName, RelPath, InPackageStatData.ModificationTime, FileType);
 					}
-					IteratedFiles[NumIteratedFiles++].Assign(LocalAbsPath, DirLongPackageName, RelPath, InPackageStatData.ModificationTime);
 				}
 			}
 			return true;
@@ -1895,18 +1905,22 @@ void FAssetDataDiscovery::SetPropertiesAndWait(FPathExistence& QueryPath, bool b
 			if (!bAlreadyScanned || bForceRescan)
 			{
 				FStringView RelPathFromParentDir = FPathViews::GetCleanFilename(RelPath);
-				FStringView FileRelPathNoExt = FPathViews::GetBaseFilenameWithPath(RelPath);
-				if (!FPackageName::DoesPackageNameContainInvalidCharacters(FileRelPathNoExt) && FPackageName::IsPackageFilename(RelPathFromParentDir))
+				EGatherableFileType FileType = GetFileType(RelPathFromParentDir);
+				if (FileType != EGatherableFileType::Invalid)
 				{
-					TStringBuilder<256> LongPackageName;
-					LongPackageName << MountDir->GetLongPackageName();
-					FPathViews::AppendPath(LongPackageName, ScanDir->GetMountRelPath());
-					FPathViews::AppendPath(LongPackageName, FileRelPathNoExt);
-					AddDiscoveredFile(FDiscoveredPathData(SearchPath, LongPackageName, RelPathFromParentDir, QueryPath.GetModificationTime()));
-					if (FPathViews::IsPathLeaf(RelPath) && !ScanDir->HasScanned())
+					FStringView FileRelPathNoExt = FPathViews::GetBaseFilenameWithPath(RelPath);
+					if (!FPackageName::DoesPackageNameContainInvalidCharacters(FileRelPathNoExt))
 					{
-						SetIsIdle(false);
-						ScanDir->MarkFileAlreadyScanned(RelPath);
+						TStringBuilder<256> LongPackageName;
+						LongPackageName << MountDir->GetLongPackageName();
+						FPathViews::AppendPath(LongPackageName, ScanDir->GetMountRelPath());
+						FPathViews::AppendPath(LongPackageName, FileRelPathNoExt);
+						AddDiscoveredFile(FDiscoveredPathData(SearchPath, LongPackageName, RelPathFromParentDir, QueryPath.GetModificationTime(), FileType));
+						if (FPathViews::IsPathLeaf(RelPath) && !ScanDir->HasScanned())
+						{
+							SetIsIdle(false);
+							ScanDir->MarkFileAlreadyScanned(RelPath);
+						}
 					}
 				}
 			}
@@ -2226,17 +2240,21 @@ void FAssetDataDiscovery::OnFileCreated(const FString& LocalAbsPath)
 	}
 
 	FStringView RelPathFromParentDir = FPathViews::GetCleanFilename(FileRelPath);
-	FStringView FileRelPathNoExt = FPathViews::GetBaseFilenameWithPath(FileRelPath);
-	if (!FPackageName::DoesPackageNameContainInvalidCharacters(FileRelPathNoExt) && FPackageName::IsPackageFilename(RelPathFromParentDir))
-	{
-		TStringBuilder<256> LongPackageName;
-		LongPackageName << MountDir->GetLongPackageName();
-		FPathViews::AppendPath(LongPackageName, ScanDir->GetMountRelPath());
-		FPathViews::AppendPath(LongPackageName, FileRelPathNoExt);
-		AddDiscoveredFile(FDiscoveredPathData(LocalAbsPath, LongPackageName, RelPathFromParentDir, StatData.ModificationTime));
-		if (FPathViews::IsPathLeaf(FileRelPath))
-		{
-			ScanDir->MarkFileAlreadyScanned(FileRelPath);
+	EGatherableFileType FileType = GetFileType(RelPathFromParentDir);
+	if (FileType != EGatherableFileType::Invalid)
+	{
+		FStringView FileRelPathNoExt = FPathViews::GetBaseFilenameWithPath(FileRelPath);
+		if (!FPackageName::DoesPackageNameContainInvalidCharacters(FileRelPathNoExt))
+		{
+			TStringBuilder<256> LongPackageName;
+			LongPackageName << MountDir->GetLongPackageName();
+			FPathViews::AppendPath(LongPackageName, ScanDir->GetMountRelPath());
+			FPathViews::AppendPath(LongPackageName, FileRelPathNoExt);
+			AddDiscoveredFile(FDiscoveredPathData(LocalAbsPath, LongPackageName, RelPathFromParentDir, StatData.ModificationTime, FileType));
+			if (FPathViews::IsPathLeaf(FileRelPath))
+			{
+				ScanDir->MarkFileAlreadyScanned(FileRelPath);
+			}
 		}
 	}
 }
@@ -2330,6 +2348,13 @@ void FAssetDataDiscovery::AddDiscoveredFile(FDiscoveredPathData&& File)
 	NumDiscoveredFiles++;
 }
 
+EGatherableFileType FAssetDataDiscovery::GetFileType(FStringView FilePath)
+{
+	return FPackageName::IsPackageFilename(FilePath)
+		? EGatherableFileType::PackageFile
+		: (FAssetDataGatherer::IsVerseFile(FilePath) ? EGatherableFileType::VerseFile : EGatherableFileType::Invalid);
+}
+
 FAssetDataDiscovery::FDirectoryResult::FDirectoryResult(FStringView InDirAbsPath, TConstArrayView<FDiscoveredPathData> InFiles)
 	: DirAbsPath(InDirAbsPath)
 {
@@ -2866,8 +2891,8 @@ FAssetDataGatherer::FAssetDataGatherer(const TArray<FString>& InLongPackageNames
 	, bFirstTickAfterIdle(true)
 	, bFinishedInitialDiscovery(false)
 	, WaitBatchCount(0)
-	, NumCachedFiles(0)
-	, NumUncachedFiles(0)
+	, NumCachedAssetFiles(0)
+	, NumUncachedAssetFiles(0)
 	, bIsSavingAsyncCache(false)
 {
 	using namespace UE::AssetDataGather::Private;
@@ -3059,6 +3084,7 @@ void FAssetDataGatherer::TickInternal(bool& bOutIsTickInterrupt)
 	TArray<FAssetData*, FBatchInlineAllocator> LocalAssetResults;
 	TArray<FPackageDependencyData, FBatchInlineAllocator> LocalDependencyResults;
 	TArray<FString, FBatchInlineAllocator> LocalCookedPackageNamesWithoutAssetDataResults;
+	TArray<FName, FBatchInlineAllocator> LocalVerseResults;
 	bool bLoadMonolithicCache = false;
 	double LocalLastCacheWriteTime = 0.0;
 	bool bWaitBatchCountDecremented = false;
@@ -3155,6 +3181,21 @@ void FAssetDataGatherer::TickInternal(bool& bOutIsTickInterrupt)
 	TArray<FReadContext> ReadContexts;
 	for (FGatheredPathData& AssetFileData : LocalFilesToSearch)
 	{
+		// If this a Verse source file, just directly add its file name to the Verse results
+		if (AssetFileData.Type == EGatherableFileType::VerseFile)
+		{
+			// Store Verse results in a hybrid format using the LongPackageName but keeping the extension
+			WriteToString<256> MappedPath(AssetFileData.LongPackageName, FPathViews::GetExtension(AssetFileData.LocalAbsPath, true));
+			LocalVerseResults.Add(*MappedPath);
+			continue;
+		}
+
+		if (AssetFileData.Type != EGatherableFileType::PackageFile)
+		{
+			ensureMsgf(false, TEXT("Encountered unrecognized gathered asset %s!"), *AssetFileData.LongPackageName);
+			continue;
+		}
+
 		const FName PackageName = *AssetFileData.LongPackageName;
 		const FName Extension = FName(*FPaths::GetExtension(AssetFileData.LocalAbsPath));
 
@@ -3179,7 +3220,7 @@ void FAssetDataGatherer::TickInternal(bool& bOutIsTickInterrupt)
 		if (DiskCachedAssetData)
 		{
 			// Add the valid cached data to our results, and to the map of data we keep to write out the new version of the cache file
-			++NumCachedFiles;
+			++NumCachedAssetFiles;
 
 			// Set the transient flags based on whether our current cache has dependency data.
 			// Note that in editor, bGatherAssetPackageData is always true, no way to turn it off,
@@ -3199,6 +3240,7 @@ void FAssetDataGatherer::TickInternal(bool& bOutIsTickInterrupt)
 		}
 		else
 		{
+			// Not found in cache (or stale) - schedule to be read from disk
 			ReadContexts.Emplace(PackageName, Extension, AssetFileData);
 		}
 	}
@@ -3229,12 +3271,14 @@ void FAssetDataGatherer::TickInternal(bool& bOutIsTickInterrupt)
 		}
 		else if (ReadContext.bResult)
 		{
-			++NumUncachedFiles;
+			++NumUncachedAssetFiles;
 
 			// Add the results from a cooked package into our results on cooked package
 			LocalCookedPackageNamesWithoutAssetDataResults.Append(MoveTemp(ReadContext.CookedPackageNamesWithoutAssetData));
 			// Do not add the results from a cooked package into the map of data we keep to write out the new version of the cache file
-			bool bCachePackage = bCacheEnabled && LocalCookedPackageNamesWithoutAssetDataResults.Num() == 0;
+			bool bCachePackage = bCacheEnabled
+				&& LocalCookedPackageNamesWithoutAssetDataResults.Num() == 0
+				&& ensure(ReadContext.AssetFileData.Type == EGatherableFileType::PackageFile);
 			if (bCachePackage)
 			{
 				for (const FAssetData* AssetData : ReadContext.AssetDataFromFile)
@@ -3284,6 +3328,7 @@ void FAssetDataGatherer::TickInternal(bool& bOutIsTickInterrupt)
 		AssetResults.Append(MoveTemp(LocalAssetResults));
 		DependencyResults.Append(MoveTemp(LocalDependencyResults));
 		CookedPackageNamesWithoutAssetDataResults.Append(MoveTemp(LocalCookedPackageNamesWithoutAssetDataResults));
+		VerseResults.Append(MoveTemp(LocalVerseResults));
 
 		if (bHasCancelation)
 		{
@@ -3419,33 +3464,32 @@ void FAssetDataGatherer::AddToCache(FName PackageName, FDiskCachedAssetData* Dis
 	}
 }
 
-void FAssetDataGatherer::GetAndTrimSearchResults(bool& bOutIsSearching, bool& bOutAbleToProgress,
-	TRingBuffer<FAssetData*>& OutAssetResults, TRingBuffer<FString>& OutPathResults, TRingBuffer<FPackageDependencyData>& OutDependencyResults,
-	TRingBuffer<FString>& OutCookedPackageNamesWithoutAssetDataResults, TArray<double>& OutSearchTimes, int32& OutNumFilesToSearch, int32& OutNumPathsToSearch, bool& OutIsDiscoveringFiles)
+void FAssetDataGatherer::GetAndTrimSearchResults(FResults& InOutResults, FResultContext& OutContext)
 {
 	FGathererScopeLock ResultsScopeLock(&ResultsLock);
 
-	auto MoveAppendRangeToRingBuffer = [](auto& OutRingBuffer, auto& InArray)
+	auto MoveAppendRangeToRingBuffer = [](auto& InOutRingBuffer, auto& InArray)
 	{
-		OutRingBuffer.Reserve(OutRingBuffer.Num() + InArray.Num());
+		InOutRingBuffer.Reserve(InOutRingBuffer.Num() + InArray.Num());
 		for (auto& Element : InArray)
 		{
-			OutRingBuffer.Add(MoveTemp(Element));
+			InOutRingBuffer.Add(MoveTemp(Element));
 		}
 		InArray.Reset();
 	};
 
-	MoveAppendRangeToRingBuffer(OutAssetResults, AssetResults);
-	MoveAppendRangeToRingBuffer(OutPathResults, DiscoveredPaths);
-	MoveAppendRangeToRingBuffer(OutDependencyResults, DependencyResults);
-	MoveAppendRangeToRingBuffer(OutCookedPackageNamesWithoutAssetDataResults, CookedPackageNamesWithoutAssetDataResults);
+	MoveAppendRangeToRingBuffer(InOutResults.Assets, AssetResults);
+	MoveAppendRangeToRingBuffer(InOutResults.Paths, DiscoveredPaths);
+	MoveAppendRangeToRingBuffer(InOutResults.Dependencies, DependencyResults);
+	MoveAppendRangeToRingBuffer(InOutResults.CookedPackageNamesWithoutAssetData, CookedPackageNamesWithoutAssetDataResults);
+	MoveAppendRangeToRingBuffer(InOutResults.VerseFiles, VerseResults);
 
-	OutSearchTimes.Append(MoveTemp(SearchTimes));
+	OutContext.SearchTimes.Append(MoveTemp(SearchTimes));
 	SearchTimes.Reset();
 
-	OutNumFilesToSearch = FilesToSearch->Num();
-	OutNumPathsToSearch = NumPathsToSearchAtLastSyncPoint;
-	OutIsDiscoveringFiles = !bDiscoveryIsComplete;
+	OutContext.NumFilesToSearch = FilesToSearch->Num();
+	OutContext.NumPathsToSearch = NumPathsToSearchAtLastSyncPoint;
+	OutContext.bIsDiscoveringFiles = !bDiscoveryIsComplete;
 
 	// Idle means no more work OR we are blocked on external events, but complete means no more work period.
 	bool bLocalIsComplete = bIsIdle && FilesToSearch->Num() == 0;
@@ -3454,8 +3498,8 @@ void FAssetDataGatherer::GetAndTrimSearchResults(bool& bOutIsSearching, bool& bO
 		bIsComplete = true;
 		Shrink();
 	}
-	bOutIsSearching = !bLocalIsComplete;
-	bOutAbleToProgress = !bIsIdle;
+	OutContext.bIsSearching = !bLocalIsComplete;
+	OutContext.bAbleToProgress = !bIsIdle;
 }
 
 void FAssetDataGatherer::GetPackageResults(TRingBuffer<FAssetData*>& OutAssetResults, TRingBuffer<FPackageDependencyData>& OutDependencyResults)
@@ -3960,6 +4004,7 @@ SIZE_T FAssetDataGatherer::GetAllocatedSize() const
 
 	Result += GetArrayRecursiveAllocatedSize(DependencyResults);
 	Result += GetArrayRecursiveAllocatedSize(CookedPackageNamesWithoutAssetDataResults);
+	Result += VerseResults.GetAllocatedSize();
 	Result += SearchTimes.GetAllocatedSize();
 	Result += GetArrayRecursiveAllocatedSize(DiscoveredPaths);
 	Result += GPreloadSettings.GetMonolithicCacheFilename().GetAllocatedSize();
@@ -3988,6 +4033,7 @@ void FAssetDataGatherer::Shrink()
 	AssetResults.Shrink();
 	DependencyResults.Shrink();
 	CookedPackageNamesWithoutAssetDataResults.Shrink();
+	VerseResults.Shrink();
 	SearchTimes.Shrink();
 	DiscoveredPaths.Shrink();
 }
@@ -4139,6 +4185,11 @@ bool FAssetDataGatherer::IsMonitored(FStringView LocalPath) const
 	return Discovery->IsMonitored(NormalizeLocalPath(LocalPath));
 }
 
+bool FAssetDataGatherer::IsVerseFile(FStringView FilePath)
+{
+	return FilePath.EndsWith(TEXT(".verse")) || FilePath.EndsWith(TEXT(".vmodule"));
+}
+
 void FAssetDataGatherer::SetIsIdle(bool bInIsIdle)
 {
 	CHECK_IS_LOCKED_CURRENT_THREAD(ResultsLock);
@@ -4156,7 +4207,7 @@ void FAssetDataGatherer::SetIsIdle(bool bInIsIdle)
 		{
 			bFinishedInitialDiscovery = true;
 
-			UE_LOG(LogAssetRegistry, Verbose, TEXT("Initial scan took %0.6f seconds (found %d cached assets, and loaded %d)"), (float)SearchTime, NumCachedFiles, NumUncachedFiles);
+			UE_LOG(LogAssetRegistry, Verbose, TEXT("Initial scan took %0.6f seconds (found %d cached assets, and loaded %d)"), (float)SearchTime, NumCachedAssetFiles, NumUncachedAssetFiles);
 		}
 		SearchTimes.Add(SearchTime);
 	}
diff --git a/Engine/Source/Runtime/AssetRegistry/Private/AssetDataGatherer.h b/Engine/Source/Runtime/AssetRegistry/Private/AssetDataGatherer.h
--- a/Engine/Source/Runtime/AssetRegistry/Private/AssetDataGatherer.h
+++ b/Engine/Source/Runtime/AssetRegistry/Private/AssetDataGatherer.h
@@ -85,12 +85,27 @@ public:
 
 
 	// Receiving Results (possibly while tick is running)
+	struct FResults
+	{
+		TRingBuffer<FAssetData*> Assets;
+		TRingBuffer<FString> Paths;
+		TRingBuffer<FPackageDependencyData> Dependencies;
+		TRingBuffer<FString> CookedPackageNamesWithoutAssetData;
+		TRingBuffer<FName> VerseFiles;
 
+		SIZE_T GetAllocatedSize() const { return Assets.GetAllocatedSize() + Paths.GetAllocatedSize() + Dependencies.GetAllocatedSize() + CookedPackageNamesWithoutAssetData.GetAllocatedSize() + VerseFiles.GetAllocatedSize(); }
+	};
+	struct FResultContext
+	{
+		bool bIsSearching = false;
+		bool bAbleToProgress = false;
+		TArray<double> SearchTimes;
+		int32 NumFilesToSearch = 0;
+		int32 NumPathsToSearch = 0;
+		bool bIsDiscoveringFiles = false;
+	};
 	/** Gets search results from the data gatherer. */
-	void GetAndTrimSearchResults(bool& bOutIsSearching, bool& bOutAbleToProgress,
-		TRingBuffer<FAssetData*>& OutAssetResults, TRingBuffer<FString>& OutPathResults,
-		TRingBuffer<FPackageDependencyData>& OutDependencyResults, TRingBuffer<FString>& OutCookedPackageNamesWithoutAssetDataResults,
-		TArray<double>& OutSearchTimes, int32& OutNumFilesToSearch, int32& OutNumPathsToSearch, bool& OutIsDiscoveringFiles);
+	void GetAndTrimSearchResults(FResults& InOutResults, FResultContext& OutContext);
 	/** Gets just the AssetResults and DependencyResults from the data gatherer. */
 	void GetPackageResults(TRingBuffer<FAssetData*>& OutAssetResults, TRingBuffer<FPackageDependencyData>& OutDependencyResults);
 	/** Wait for all monitored assets under the given path to be added to search results. Returns immediately if the given path is not monitored. */
@@ -164,6 +179,9 @@ public:
 	/** Report whether the path is both in the allow list and not in the deny list. */
 	bool IsMonitored(FStringView LocalPath) const;
 
+	/** Determine, based on the file extension, if the given file path is a Verse file */
+	static bool IsVerseFile(FStringView FilePath);
+
 	/**
 	 * Reads FAssetData information out of a previously initialized package reader
 	 *
@@ -336,6 +354,8 @@ private:
 	 * These assets may still contain assets (if they were older for example). Read/writable only within ResultsLock.
 	 */
 	TArray<FString> CookedPackageNamesWithoutAssetDataResults;
+	/** File paths (in UE LongPackagePath notation) of the Verse source code gathered from the searched files. */
+	TArray<FName> VerseResults;
 
 	/** All the search times since the last call to GetAndTrimSearchResults. Read/writable only within ResultsLock. */
 	TArray<double> SearchTimes;
@@ -396,9 +416,9 @@ private:
 	 */
 	int32 WaitBatchCount;
 	/** How many files in the search results were read from the cache. Read/writable only within TickLock. */
-	int32 NumCachedFiles;
+	int32 NumCachedAssetFiles;
 	/** How many files in the search results were not in the cache and were read by parsing the file. Read/writable only within TickLock. */
-	int32 NumUncachedFiles;
+	int32 NumUncachedAssetFiles;
 	/**
 	 * True if the current TickInternal is synchronous, which may be because bIsSynchronous or because the game thread has taken over the tick for a synchronous function.
 	 * Read/writable only within TickLock.
diff --git a/Engine/Source/Runtime/AssetRegistry/Private/AssetDataGathererPrivate.h b/Engine/Source/Runtime/AssetRegistry/Private/AssetDataGathererPrivate.h
--- a/Engine/Source/Runtime/AssetRegistry/Private/AssetDataGathererPrivate.h
+++ b/Engine/Source/Runtime/AssetRegistry/Private/AssetDataGathererPrivate.h
@@ -65,6 +65,15 @@ struct FSetPathProperties
 	}
 };
 
+/** To distinguish between assets and Verse source code */
+enum class EGatherableFileType : uint8
+{
+	Invalid,
+	Directory,
+	PackageFile,
+	VerseFile,
+};
+
 /** Information needed about a discovered asset file or path that is needed by the Discoverer */
 struct FDiscoveredPathData
 {
@@ -76,14 +85,18 @@ struct FDiscoveredPathData
 	FString RelPath;
 	/** If the path is a file, the modification timestamp of the package file (that it had when it was discovered). */
 	FDateTime PackageTimestamp;
+	/** The type of file that was found */
+	EGatherableFileType Type = EGatherableFileType::Invalid;
 
 	FDiscoveredPathData() = default;
 	FDiscoveredPathData(FStringView InLocalAbsPath, FStringView InLongPackageName, FStringView InRelPath,
-		const FDateTime& InPackageTimestamp);
-	FDiscoveredPathData(FStringView InLocalAbsPath, FStringView InLongPackageName, FStringView InRelPath);
+		const FDateTime& InPackageTimestamp, EGatherableFileType InType);
+	FDiscoveredPathData(FStringView InLocalAbsPath, FStringView InLongPackageName, FStringView InRelPath,
+		EGatherableFileType InType);
+	void Assign(FStringView InLocalAbsPath, FStringView InLongPackageName, FStringView InRelPath,
+		const FDateTime& InPackageTimestamp, EGatherableFileType InType);
 	void Assign(FStringView InLocalAbsPath, FStringView InLongPackageName, FStringView InRelPath,
-		const FDateTime& InPackageTimestamp);
-	void Assign(FStringView InLocalAbsPath, FStringView InLongPackageName, FStringView InRelPath);
+		EGatherableFileType InType);
 
 	/** Return the total amount of heap memory used by the gatherer (including not-yet-claimed search results). */
 	SIZE_T GetAllocatedSize() const;
@@ -98,12 +111,14 @@ struct FGatheredPathData
 	FString LongPackageName;
 	/** The modification timestamp of the package file (that it had when it was discovered) */
 	FDateTime PackageTimestamp;
+	/** The type of file that was found */
+	EGatherableFileType Type = EGatherableFileType::Invalid;
 
 	FGatheredPathData() = default;
-	FGatheredPathData(FStringView InLocalAbsPath, FStringView InLongPackageName, const FDateTime& InPackageTimestamp);
+	FGatheredPathData(FStringView InLocalAbsPath, FStringView InLongPackageName, const FDateTime& InPackageTimestamp, EGatherableFileType InType);
 	explicit FGatheredPathData(const FDiscoveredPathData& DiscoveredData);
 	explicit FGatheredPathData(FDiscoveredPathData&& DiscoveredData);
-	void Assign(FStringView InLocalAbsPath, FStringView InLongPackageName, const FDateTime& InPackageTimestamp);
+	void Assign(FStringView InLocalAbsPath, FStringView InLongPackageName, const FDateTime& InPackageTimestamp, EGatherableFileType InType);
 	void Assign(const FDiscoveredPathData& DiscoveredData);
 
 	/**
@@ -648,6 +663,9 @@ private:
 	/** Store the given specially reported single file in the results. */
 	void AddDiscoveredFile(FDiscoveredPathData&& File);
 
+	/** For a given file path, determine the type it should be gathered as. */
+	static EGatherableFileType GetFileType(FStringView FilePath);
+
 	/**
 	 * Return whether a directory with the given LongPackageName should be reported to the AssetRegistry
 	 * We do not report some directories because they are paths that should not enter the AssetRegistry list of paths if empty,
diff --git a/Engine/Source/Runtime/AssetRegistry/Private/AssetRegistry.cpp b/Engine/Source/Runtime/AssetRegistry/Private/AssetRegistry.cpp
--- a/Engine/Source/Runtime/AssetRegistry/Private/AssetRegistry.cpp
+++ b/Engine/Source/Runtime/AssetRegistry/Private/AssetRegistry.cpp
@@ -3252,6 +3252,40 @@ void PrioritizeAssetInstall(const FAssetData& AssetData)
 
 }
 
+bool UAssetRegistryImpl::GetVerseFilesByPath(FName PackagePath, TArray<FName>& OutFilePaths, bool bRecursive /*= false*/) const
+{
+	FReadScopeLock InterfaceScopeLock(InterfaceLock);
+	return GuardedData.GetVerseFilesByPath(PackagePath, OutFilePaths, bRecursive);
+}
+
+namespace UE::AssetRegistry
+{
+
+bool FAssetRegistryImpl::GetVerseFilesByPath(FName PackagePath, TArray<FName>& OutFilePaths, bool bRecursive /*= false*/) const
+{
+	bool bFoundAnything = false;
+	TSet<FName> PathList;
+	PathList.Reserve(32);
+	if (bRecursive)
+	{
+		CachedPathTree.GetSubPaths(PackagePath, PathList, true);
+	}
+	PathList.Add(PackagePath);
+	for (const FName& PathName : PathList)
+	{
+		const TArray<FName>* FilePaths = CachedVerseFilesByPath.Find(PathName);
+		if (FilePaths)
+		{
+			OutFilePaths.Append(*FilePaths);
+			bFoundAnything = true;
+		}
+	}
+
+	return bFoundAnything;
+}
+
+}
+
 bool UAssetRegistryImpl::AddPath(const FString& PathToAdd)
 {
 	UE::AssetRegistry::Impl::FEventContext EventContext;
@@ -3390,21 +3424,21 @@ void FAssetRegistryImpl::PrioritizeSearchPath(const FString& PathToPrioritize)
 
 	// Also prioritize the queue of background search results
 	int32 FirstNonPriorityIndex = 0;
-	for (int Index = 0; Index < BackgroundAssetResults.Num(); ++Index)
+	for (int Index = 0; Index < BackgroundResults.Assets.Num(); ++Index)
 	{
-		FAssetData* PriorityElement = BackgroundAssetResults[Index];
+		FAssetData* PriorityElement = BackgroundResults.Assets[Index];
 		if (PriorityElement && PriorityElement->PackagePath.ToString().StartsWith(PathToPrioritize))
 		{
-			Swap(BackgroundAssetResults[FirstNonPriorityIndex++], BackgroundAssetResults[Index]);
+			Swap(BackgroundResults.Assets[FirstNonPriorityIndex++], BackgroundResults.Assets[Index]);
 		}
 	}
 	FirstNonPriorityIndex = 0;
-	for (int Index = 0; Index < BackgroundPathResults.Num(); ++Index)
+	for (int Index = 0; Index < BackgroundResults.Paths.Num(); ++Index)
 	{
-		FString& PriorityElement = BackgroundPathResults[Index];
+		FString& PriorityElement = BackgroundResults.Paths[Index];
 		if (PriorityElement.StartsWith(PathToPrioritize))
 		{
-			Swap(BackgroundPathResults[FirstNonPriorityIndex++], BackgroundPathResults[Index]);
+			Swap(BackgroundResults.Paths[FirstNonPriorityIndex++], BackgroundResults.Paths[Index]);
 		}
 	}
 }
@@ -3684,27 +3718,20 @@ Impl::EGatherStatus FAssetRegistryImpl::TickGatherer(Impl::FEventContext& EventC
 	}
 
 	// Gather results from the background search
-	bool bIsSearching = false;
-	bool bAbleToProgress = false;
-	TArray<double> SearchTimes;
-	int32 NumFilesToSearch = 0;
-	int32 NumPathsToSearch = 0;
-	bool bIsDiscoveringFiles = false;
-	GlobalGatherer->GetAndTrimSearchResults(bIsSearching, bAbleToProgress, BackgroundAssetResults, BackgroundPathResults, BackgroundDependencyResults,
-		BackgroundCookedPackageNamesWithoutAssetDataResults, SearchTimes, NumFilesToSearch, NumPathsToSearch, bIsDiscoveringFiles);
+	FAssetDataGatherer::FResultContext ResultContext;
+	GlobalGatherer->GetAndTrimSearchResults(BackgroundResults, ResultContext);	
 	// Report the search times
-	for (int32 SearchTimeIdx = 0; SearchTimeIdx < SearchTimes.Num(); ++SearchTimeIdx)
+	for (double SearchTime : ResultContext.SearchTimes)
 	{
-		UE_LOG(LogAssetRegistry, Verbose, TEXT("### Background search completed in %0.4f seconds"), SearchTimes[SearchTimeIdx]);
+		UE_LOG(LogAssetRegistry, Verbose, TEXT("### Background search completed in %0.4f seconds"), SearchTime);
 	}
-	const bool bHadAssetsToProcess = BackgroundAssetResults.Num() > 0 || BackgroundDependencyResults.Num() > 0;
-	auto GetNumGatherFromDiskPending = [&NumFilesToSearch, &NumPathsToSearch, this]()
+	const bool bHadAssetsToProcess = BackgroundResults.Assets.Num() > 0 || BackgroundResults.Dependencies.Num() > 0;
+	auto GetNumGatherFromDiskPending = [&ResultContext, this]()
 	{
-		return NumFilesToSearch + NumPathsToSearch + BackgroundPathResults.Num() + BackgroundAssetResults.Num()
-			+ BackgroundDependencyResults.Num() + BackgroundCookedPackageNamesWithoutAssetDataResults.Num();
+		return ResultContext.NumFilesToSearch + ResultContext.NumPathsToSearch + BackgroundResults.Paths.Num() + BackgroundResults.Assets.Num()
+			+ BackgroundResults.Dependencies.Num() + BackgroundResults.CookedPackageNamesWithoutAssetData.Num();
 	};
-	auto UpdateStatus = [bIsSearching, bAbleToProgress, bHadAssetsToProcess, bIsDiscoveringFiles, &EventContext,
-		&OutStatus, this](int32 NumGatherPending, bool bInterrupted)
+	auto UpdateStatus = [bHadAssetsToProcess, &ResultContext, &EventContext, &OutStatus, this](int32 NumGatherPending, bool bInterrupted)
 	{
 		// Compute total pending, plus highest pending for this run so we can show a good progress bar
 		int32 NumPending = NumGatherPending
@@ -3714,11 +3741,11 @@ Impl::EGatherStatus FAssetRegistryImpl::TickGatherer(Impl::FEventContext& EventC
 			;
 		HighestPending = FMath::Max(this->HighestPending, NumPending);
 
-		if (!bInterrupted && !bIsSearching && NumPending == 0)
+		if (!bInterrupted && !ResultContext.bIsSearching && NumPending == 0)
 		{
 			OutStatus = EGatherStatus::Complete;
 		}
-		else if (!bInterrupted && !bAbleToProgress)
+		else if (!bInterrupted && !ResultContext.bAbleToProgress)
 		{
 			OutStatus = EGatherStatus::UnableToProgress;
 		}
@@ -3727,14 +3754,14 @@ Impl::EGatherStatus FAssetRegistryImpl::TickGatherer(Impl::FEventContext& EventC
 			OutStatus = EGatherStatus::Active;
 		}
 		// Notify the status change, only when something changed, or when sending the final result before going idle
-		if (bIsSearching || bHadAssetsToProcess ||
+		if (ResultContext.bIsSearching || bHadAssetsToProcess ||
 			(OutStatus == EGatherStatus::Complete && this->GatherStatus != EGatherStatus::Complete))
 		{
 			EventContext.ProgressUpdateData.Emplace(
 				HighestPending,					// NumTotalAssets
 				HighestPending - NumPending,	// NumAssetsProcessedByAssetRegistry
 				NumPending / 2,					// NumAssetsPendingDataLoad, divided by 2 because assets are double counted due to dependencies
-				bIsDiscoveringFiles				// bIsDiscoveringAssetFiles
+				ResultContext.bIsDiscoveringFiles // bIsDiscoveringAssetFiles
 			);
 		}
 		this->GatherStatus = OutStatus;
@@ -3742,13 +3769,13 @@ Impl::EGatherStatus FAssetRegistryImpl::TickGatherer(Impl::FEventContext& EventC
 
 
 	// Add discovered paths
-	if (BackgroundPathResults.Num())
+	if (BackgroundResults.Paths.Num())
 	{
-		PathDataGathered(EventContext, TickStartTime, BackgroundPathResults);
+		PathDataGathered(EventContext, TickStartTime, BackgroundResults.Paths);
 	}
 
 	// Process the asset results
-	if (BackgroundAssetResults.Num())
+	if (BackgroundResults.Assets.Num())
 	{
 		// Mark the first amortize time
 		if (AmortizeStartTime == 0)
@@ -3757,12 +3784,12 @@ Impl::EGatherStatus FAssetRegistryImpl::TickGatherer(Impl::FEventContext& EventC
 		}
 		if (AssetsFoundCallback.IsSet())
 		{
-			AssetsFoundCallback.GetValue()(BackgroundAssetResults);
+			AssetsFoundCallback.GetValue()(BackgroundResults.Assets);
 		}
 
-		AssetSearchDataGathered(EventContext, TickStartTime, BackgroundAssetResults);
+		AssetSearchDataGathered(EventContext, TickStartTime, BackgroundResults.Assets);
 
-		if (BackgroundAssetResults.Num() == 0)
+		if (BackgroundResults.Assets.Num() == 0)
 		{
 			TotalAmortizeTime += FPlatformTime::Seconds() - AmortizeStartTime;
 			AmortizeStartTime = 0;
@@ -3770,15 +3797,15 @@ Impl::EGatherStatus FAssetRegistryImpl::TickGatherer(Impl::FEventContext& EventC
 	}
 
 	// Add dependencies
-	if (BackgroundDependencyResults.Num())
+	if (BackgroundResults.Dependencies.Num())
 	{
-		DependencyDataGathered(TickStartTime, BackgroundDependencyResults);
+		DependencyDataGathered(TickStartTime, BackgroundResults.Dependencies);
 	}
 
 	// Load cooked packages that do not have asset data
-	if (BackgroundCookedPackageNamesWithoutAssetDataResults.Num())
+	if (BackgroundResults.CookedPackageNamesWithoutAssetData.Num())
 	{
-		CookedPackageNamesWithoutAssetDataGathered(EventContext, TickStartTime, BackgroundCookedPackageNamesWithoutAssetDataResults, bOutInterrupted);
+		CookedPackageNamesWithoutAssetDataGathered(EventContext, TickStartTime, BackgroundResults.CookedPackageNamesWithoutAssetData, bOutInterrupted);
 		if (bOutInterrupted)
 		{
 			UpdateStatus(GetNumGatherFromDiskPending(), true /* bInterrupted */);
@@ -3786,10 +3813,16 @@ Impl::EGatherStatus FAssetRegistryImpl::TickGatherer(Impl::FEventContext& EventC
 		}
 	}
 
+	// Add Verse files
+	if (BackgroundResults.VerseFiles.Num())
+	{
+		VerseFilesGathered(TickStartTime, BackgroundResults.VerseFiles);
+	}
+
 	// Load Calculated Dependencies when the gather from disk is complete; the full gather is not complete until after this is done
 	int32 NumGatherFromDiskPending = GetNumGatherFromDiskPending();
 #if WITH_EDITOR
-	bool bDiskGatherComplete = !bIsSearching && NumGatherFromDiskPending == 0;
+	bool bDiskGatherComplete = !ResultContext.bIsSearching && NumGatherFromDiskPending == 0;
 	if (bDiskGatherComplete && PackagesNeedingDependencyCalculation.Num())
 	{
 		LoadCalculatedDependencies(nullptr, TickStartTime, InheritanceContext, bOutInterrupted);
@@ -3837,27 +3870,27 @@ void FAssetRegistryImpl::TickGatherPackage(Impl::FEventContext& EventContext, co
 	FName PackageFName(PackageName);
 
 	// Gather results from the background search
-	GlobalGatherer->GetPackageResults(BackgroundAssetResults, BackgroundDependencyResults);
+	GlobalGatherer->GetPackageResults(BackgroundResults.Assets, BackgroundResults.Dependencies);
 
 	TRingBuffer<FAssetData*> PackageAssets;
 	TRingBuffer<FPackageDependencyData> PackageDependencyDatas;
-	for (int n = 0; n < BackgroundAssetResults.Num(); ++n)
+	for (int n = 0; n < BackgroundResults.Assets.Num(); ++n)
 	{
-		FAssetData* Asset = BackgroundAssetResults[n];
+		FAssetData* Asset = BackgroundResults.Assets[n];
 		if (Asset->PackageName == PackageFName)
 		{
-			Swap(BackgroundAssetResults[n], BackgroundAssetResults.Last());
-			PackageAssets.Add(BackgroundAssetResults.PopValue());
+			Swap(BackgroundResults.Assets[n], BackgroundResults.Assets.Last());
+			PackageAssets.Add(BackgroundResults.Assets.PopValue());
 			--n;
 		}
 	}
-	for (int n = 0; n < BackgroundDependencyResults.Num(); ++n)
+	for (int n = 0; n < BackgroundResults.Dependencies.Num(); ++n)
 	{
-		FPackageDependencyData& DependencyData = BackgroundDependencyResults[n];
+		FPackageDependencyData& DependencyData = BackgroundResults.Dependencies[n];
 		if (DependencyData.PackageName == PackageFName)
 		{
-			Swap(BackgroundDependencyResults[n], BackgroundDependencyResults.Last());
-			PackageDependencyDatas.Add(BackgroundDependencyResults.PopValue());
+			Swap(BackgroundResults.Dependencies[n], BackgroundResults.Dependencies.Last());
+			PackageDependencyDatas.Add(BackgroundResults.Dependencies.PopValue());
 			--n;
 		}
 	}
@@ -4135,7 +4168,7 @@ void FAssetRegistryImpl::GetAllocatedSize(bool bLogDetailed, SIZE_T& StateSize,
 	StateSize = State.GetAllocatedSize(bLogDetailed);
 
 	StaticSize = CachedEmptyPackages.GetAllocatedSize() + CachedBPInheritanceMap.GetAllocatedSize() + ClassGeneratorNames.GetAllocatedSize();
-	SearchSize = BackgroundAssetResults.GetAllocatedSize() + BackgroundPathResults.GetAllocatedSize() + BackgroundDependencyResults.GetAllocatedSize() + BackgroundCookedPackageNamesWithoutAssetDataResults.GetAllocatedSize() + CachedPathTree.GetAllocatedSize();
+	SearchSize = BackgroundResults.GetAllocatedSize() + CachedPathTree.GetAllocatedSize();
 
 	if (bIsTempCachingEnabled && !bIsTempCachingAlwaysEnabled)
 	{
@@ -4761,7 +4794,6 @@ void FAssetRegistryImpl::CookedPackageNamesWithoutAssetDataGathered(Impl::FEvent
 	const double TickStartTime, TRingBuffer<FString>& CookedPackageNamesWithoutAssetDataResults, bool& bOutInterrupted)
 {
 	bOutInterrupted = false;
-	const bool bFlushFullBuffer = TickStartTime < 0;
 
 	struct FConfigValue
 	{
@@ -4802,6 +4834,30 @@ void FAssetRegistryImpl::CookedPackageNamesWithoutAssetDataGathered(Impl::FEvent
 	}
 }
 
+void FAssetRegistryImpl::VerseFilesGathered(const double TickStartTime, TRingBuffer<FName>& VerseResults)
+{
+	while (VerseResults.Num() > 0)
+	{
+		FName VerseFilePath = VerseResults.PopFrontValue();
+
+		bool bAlreadyExists = false;
+		CachedVerseFiles.Add(VerseFilePath, &bAlreadyExists);
+		if (!bAlreadyExists)
+		{
+			WriteToString<256> VerseFilePathString(VerseFilePath);
+			FName VerseDirectoryPath(FPathViews::GetPath(VerseFilePathString.ToView()));
+			TArray<FName>& FilePathsArray = CachedVerseFilesByPath.FindOrAdd(VerseDirectoryPath);
+			FilePathsArray.Add(VerseFilePath);
+		}
+
+		// Check to see if we have run out of time in this tick
+		if (TickStartTime >= 0 && (FPlatformTime::Seconds() - TickStartTime) > Impl::MaxSecondsPerFrame)
+		{
+			return;
+		}
+	}
+}
+
 void FAssetRegistryImpl::AddEmptyPackage(FName PackageName)
 {
 	CachedEmptyPackages.Add(PackageName);
@@ -5005,6 +5061,24 @@ void FAssetRegistryImpl::RemovePackageData(Impl::FEventContext& EventContext, co
 	}
 }
 
+void FAssetRegistryImpl::RemoveVerseFile(FName VerseFilePathToRemove)
+{
+	if (CachedVerseFiles.Remove(VerseFilePathToRemove))
+	{
+		WriteToString<256> VerseFilePathString(VerseFilePathToRemove);
+		FName VerseDirectoryPath(FPathViews::GetPath(VerseFilePathString.ToView()));
+		TArray<FName>* FilePathsArray = CachedVerseFilesByPath.Find(VerseDirectoryPath);
+		if (ensure(FilePathsArray)) // We found it in CachedVerseFiles, so we must also find it here
+		{
+			FilePathsArray->Remove(VerseFilePathToRemove);
+			if (FilePathsArray->IsEmpty())
+			{
+				CachedVerseFilesByPath.Remove(VerseDirectoryPath);
+			}
+		}
+	}
+}
+
 }
 
 #if WITH_EDITOR
@@ -5067,10 +5141,11 @@ void FAssetRegistryImpl::OnDirectoryChanged(Impl::FEventContext& EventContext,
 		const bool bIsPackageFile = FPackageName::IsPackageExtension(*FPaths::GetExtension(File, true));
 		const bool bIsValidPackageName = FPackageName::TryConvertFilenameToLongPackageName(File, LongPackageName);
 		const bool bIsValidPackage = bIsPackageFile && bIsValidPackageName;
-		FName LongPackageFName(*LongPackageName);
 
 		if (bIsValidPackage)
 		{
+			FName LongPackageFName(*LongPackageName);
+
 			bool bAddedOrCreated = false;
 			switch (FileChangesProcessed[FileIdx].Action)
 			{
@@ -5106,27 +5181,52 @@ void FAssetRegistryImpl::OnDirectoryChanged(Impl::FEventContext& EventContext,
 		}
 		else if (bIsValidPackageName)
 		{
-			// This could be a directory or possibly a file with no extension or a wrong extension.
-			// No guaranteed way to know at this point since it may have been deleted.
-			switch (FileChangesProcessed[FileIdx].Action)
+			// Is this a Verse file?
+			if (FAssetDataGatherer::IsVerseFile(File))
 			{
-			case FFileChangeData::FCA_Added:
-			{
-				if (FPaths::DirectoryExists(File))
+				switch (FileChangesProcessed[FileIdx].Action)
 				{
-					NewDirs.Add(File);
-					UE_LOG(LogAssetRegistry, Verbose, TEXT("Directory was added to content directory: %s"), *File);
+				case FFileChangeData::FCA_Added:
+					// This is a Verse file that was created on disk.
+					NewFiles.AddUnique(File);
+					UE_LOG(LogAssetRegistry, Verbose, TEXT("Verse file was added to content directory: %s"), *File);
+					break;
+
+				case FFileChangeData::FCA_Modified:
+					// Note: Since content of Verse files is not scanned, no need to handle FCA_Modified
+					break;
+
+				case FFileChangeData::FCA_Removed:
+					WriteToString<256> VerseFilePathString(LongPackageName, FPathViews::GetExtension(File, true));
+					RemoveVerseFile(*VerseFilePathString);
+					UE_LOG(LogAssetRegistry, Verbose, TEXT("Verse file was removed from content directory: %s"), *File);
+					break;
 				}
-				break;
 			}
-			case FFileChangeData::FCA_Removed:
+			else
 			{
-				RemoveAssetPath(EventContext, *LongPackageName);
-				UE_LOG(LogAssetRegistry, Verbose, TEXT("Directory was removed from content directory: %s"), *File);
-				break;
-			}
-			default:
-				break;
+				// This could be a directory or possibly a file with no extension or a wrong extension.
+				// No guaranteed way to know at this point since it may have been deleted.
+				switch (FileChangesProcessed[FileIdx].Action)
+				{
+				case FFileChangeData::FCA_Added:
+				{
+					if (FPaths::DirectoryExists(File))
+					{
+						NewDirs.Add(File);
+						UE_LOG(LogAssetRegistry, Verbose, TEXT("Directory was added to content directory: %s"), *File);
+					}
+					break;
+				}
+				case FFileChangeData::FCA_Removed:
+				{
+					RemoveAssetPath(EventContext, *LongPackageName);
+					UE_LOG(LogAssetRegistry, Verbose, TEXT("Directory was removed from content directory: %s"), *File);
+					break;
+				}
+				default:
+					break;
+				}
 			}
 		}
 
diff --git a/Engine/Source/Runtime/AssetRegistry/Private/AssetRegistry.h b/Engine/Source/Runtime/AssetRegistry/Private/AssetRegistry.h
--- a/Engine/Source/Runtime/AssetRegistry/Private/AssetRegistry.h
+++ b/Engine/Source/Runtime/AssetRegistry/Private/AssetRegistry.h
@@ -97,6 +97,7 @@ public:
 	virtual float GetAssetAvailabilityProgress(const FAssetData& AssetData, EAssetAvailabilityProgressReportingType::Type ReportType) const override;
 	virtual bool GetAssetAvailabilityProgressTypeSupported(EAssetAvailabilityProgressReportingType::Type ReportType) const override;
 	virtual void PrioritizeAssetInstall(const FAssetData& AssetData) const override;
+	virtual bool GetVerseFilesByPath(FName PackagePath, TArray<FName>& OutFilePaths, bool bRecursive = false) const override;
 	virtual bool AddPath(const FString& PathToAdd) override;
 	virtual bool RemovePath(const FString& PathToRemove) override;
 	virtual bool PathExists(const FString& PathToTest) const override;
diff --git a/Engine/Source/Runtime/AssetRegistry/Private/AssetRegistryImpl.h b/Engine/Source/Runtime/AssetRegistry/Private/AssetRegistryImpl.h
--- a/Engine/Source/Runtime/AssetRegistry/Private/AssetRegistryImpl.h
+++ b/Engine/Source/Runtime/AssetRegistry/Private/AssetRegistryImpl.h
@@ -2,6 +2,7 @@
 
 #pragma once
 
+#include "AssetDataGatherer.h"
 #include "AssetRegistry/AssetData.h"
 #include "AssetRegistry/AssetRegistryState.h"
 #include "AssetRegistry/IAssetRegistry.h"
@@ -136,6 +137,7 @@ public:
 	bool AddPath(Impl::FEventContext& EventContext, const FString& PathToAdd);
 	void SearchAllAssets(Impl::FEventContext& EventContext, Impl::FClassInheritanceContext& InheritanceContext,
 		bool bSynchronousSearch);
+	bool GetVerseFilesByPath(FName PackagePath, TArray<FName>& OutFilePaths, bool bRecursive = false) const;
 	void ScanPathsSynchronous(Impl::FScanPathContext& Context);
 	void PrioritizeSearchPath(const FString& PathToPrioritize);
 	void ScanModifiedAssetFiles(Impl::FEventContext& EventContext, Impl::FClassInheritanceContext& InheritanceContext,
@@ -215,6 +217,8 @@ public:
 	bool RemoveAssetPath(Impl::FEventContext& EventContext, FName PathToRemove, bool bEvenIfAssetsStillExist = false);
 	/** Removes the asset data associated with this package from the look-up maps */
 	void RemovePackageData(Impl::FEventContext& EventContext, const FName PackageName);
+	/** Removes the Verse file from the look-up maps */
+	void RemoveVerseFile(FName VerseFilePathToRemove);
 
 	/** Returns the names of all subclasses of the class whose name is ClassName */
 	void GetSubClasses(Impl::FClassInheritanceContext& InheritanceContext, const TArray<FTopLevelAssetPath>& InClassNames,
@@ -286,6 +290,9 @@ private:
 	void CookedPackageNamesWithoutAssetDataGathered(Impl::FEventContext& EventContext, const double TickStartTime,
 		TRingBuffer<FString>& CookedPackageNamesWithoutAssetDataResults, bool& bOutInterrupted);
 
+	/** Called every tick when data is retrieved by the background dependency search */
+	void VerseFilesGathered(const double TickStartTime, TRingBuffer<FName>& VerseResults);
+
 	/** Adds the asset data to the lookup maps */
 	void AddAssetData(Impl::FEventContext& EventContext, FAssetData* AssetData);
 
@@ -328,6 +335,10 @@ private:
 	/** Internal state of the cached asset registry */
 	FAssetRegistryState State;
 
+	/** Database of known Verse files */
+	TSet<FName> CachedVerseFiles;
+	TMap<FName, TArray<FName>> CachedVerseFilesByPath;
+
 	/** Default options used for serialization */
 	FAssetRegistrySerializationOptions SerializationOptions;
 	FAssetRegistrySerializationOptions DevelopmentSerializationOptions;
@@ -358,11 +369,8 @@ private:
 	/** Async task that gathers asset information from disk */
 	TUniquePtr<FAssetDataGatherer> GlobalGatherer;
 
-	/** A list of results from the background thread that are waiting to get processed by the main thread */
-	TRingBuffer<FAssetData*> BackgroundAssetResults;
-	TRingBuffer<FString> BackgroundPathResults;
-	TRingBuffer<FPackageDependencyData> BackgroundDependencyResults;
-	TRingBuffer<FString> BackgroundCookedPackageNamesWithoutAssetDataResults;
+	/** Lists of results from the background thread that are waiting to get processed by the main thread */
+	FAssetDataGatherer::FResults BackgroundResults;
 
 	UE_DEPRECATED(5.0, "DelayDelete is only intended to support deprecated functions.")
 	TArray<TUniqueFunction<void()>> DeleteActions;
diff --git a/Engine/Source/Runtime/AssetRegistry/Public/AssetRegistry/IAssetRegistry.h b/Engine/Source/Runtime/AssetRegistry/Public/AssetRegistry/IAssetRegistry.h
--- a/Engine/Source/Runtime/AssetRegistry/Public/AssetRegistry/IAssetRegistry.h
+++ b/Engine/Source/Runtime/AssetRegistry/Public/AssetRegistry/IAssetRegistry.h
@@ -483,6 +483,15 @@ public:
 	 */
 	virtual void PrioritizeAssetInstall(const FAssetData& AssetData) const = 0;
 
+	/**
+	 * Gets paths for all Verse files in the supplied folder path
+	 *
+	 * @param PackagePath the path to query asset data in (e.g. /Game/MyFolder)
+	 * @param OutFilePaths the list of Verse files in this path, as pseudo UE LongPackagePaths with extension (e.g. /Game/MyFolder/MyVerseFile.verse)
+	 * @param bRecursive if true, all supplied paths will be searched recursively
+	 */
+	virtual bool GetVerseFilesByPath(FName PackagePath, TArray<FName>& OutFilePaths, bool bRecursive = false) const = 0;
+
 	/** Adds the specified path to the set of cached paths. These will be returned by GetAllCachedPaths(). Returns true if the path was actually added and false if it already existed. */
 	virtual bool AddPath(const FString& PathToAdd) = 0;
 
