Add version of GetSubPaths that works on just FNames.

#tests Used when implementing source control menu refactor  https://p4-swarm.epicgames.net/reviews/22528870

diff --git a/Engine/Source/Runtime/AssetRegistry/Private/AssetRegistry.cpp b/Engine/Source/Runtime/AssetRegistry/Private/AssetRegistry.cpp
--- a/Engine/Source/Runtime/AssetRegistry/Private/AssetRegistry.cpp
+++ b/Engine/Source/Runtime/AssetRegistry/Private/AssetRegistry.cpp
@@ -2820,6 +2820,17 @@ void UAssetRegistryImpl::GetSubPaths(const FString& InBasePath, TArray<FString>&
 	}, bInRecurse);
 }
 
+void UAssetRegistryImpl::GetSubPaths(const FName& InBasePath, TArray<FName>& OutPathList, bool bInRecurse) const
+{
+	FReadScopeLock InterfaceScopeLock(InterfaceLock);
+	const FPathTree& CachedPathTree = GuardedData.GetCachedPathTree();
+	CachedPathTree.EnumerateSubPaths(InBasePath, [&OutPathList](FName Path)
+	{
+		OutPathList.Emplace(Path);
+		return true;
+	}, bInRecurse);
+}
+
 void UAssetRegistryImpl::EnumerateSubPaths(const FString& InBasePath, TFunctionRef<bool(FString)> Callback, bool bInRecurse) const
 {
 	TArray<FName, TInlineAllocator<64>> SubPaths;
diff --git a/Engine/Source/Runtime/AssetRegistry/Private/AssetRegistry.h b/Engine/Source/Runtime/AssetRegistry/Private/AssetRegistry.h
--- a/Engine/Source/Runtime/AssetRegistry/Private/AssetRegistry.h
+++ b/Engine/Source/Runtime/AssetRegistry/Private/AssetRegistry.h
@@ -82,6 +82,7 @@ public:
 	virtual void EnumerateAllCachedPaths(TFunctionRef<bool(FString)> Callback) const override;
 	virtual void EnumerateAllCachedPaths(TFunctionRef<bool(FName)> Callback) const override;
 	virtual void GetSubPaths(const FString& InBasePath, TArray<FString>& OutPathList, bool bInRecurse) const override;
+	virtual void GetSubPaths(const FName& InBasePath, TArray<FName>& OutPathList, bool bInRecurse) const override;
 	virtual void EnumerateSubPaths(const FString& InBasePath, TFunctionRef<bool(FString)> Callback, bool bInRecurse) const override;
 	virtual void EnumerateSubPaths(const FName InBasePath, TFunctionRef<bool(FName)> Callback, bool bInRecurse) const override;
 	virtual void RunAssetsThroughFilter (TArray<FAssetData>& AssetDataList, const FARFilter& Filter) const override;
diff --git a/Engine/Source/Runtime/AssetRegistry/Public/AssetRegistry/IAssetRegistry.h b/Engine/Source/Runtime/AssetRegistry/Public/AssetRegistry/IAssetRegistry.h
--- a/Engine/Source/Runtime/AssetRegistry/Public/AssetRegistry/IAssetRegistry.h
+++ b/Engine/Source/Runtime/AssetRegistry/Public/AssetRegistry/IAssetRegistry.h
@@ -417,6 +417,9 @@ public:
 	UFUNCTION(BlueprintCallable, BlueprintPure=false, Category = "AssetRegistry")
 	virtual void GetSubPaths(const FString& InBasePath, TArray<FString>& OutPathList, bool bInRecurse) const = 0;
 
+	/** Gets a list of all paths by name that are currently cached below the passed-in base path */
+	virtual void GetSubPaths(const FName& InBasePath, TArray<FName>& OutPathList, bool bInRecurse) const = 0;
+
 	/** Enumerate the all paths that are currently cached below the passed-in base path */
 	virtual void EnumerateSubPaths(const FString& InBasePath, TFunctionRef<bool(FString)> Callback, bool bInRecurse) const = 0;
 
