AssetRegistry: Fixed WaitForCompletion to wait for the gatherer to be idle when running with -NoThreading

#rnx

diff --git a/Engine/Source/Runtime/AssetRegistry/Private/AssetRegistry.cpp b/Engine/Source/Runtime/AssetRegistry/Private/AssetRegistry.cpp
--- a/Engine/Source/Runtime/AssetRegistry/Private/AssetRegistry.cpp
+++ b/Engine/Source/Runtime/AssetRegistry/Private/AssetRegistry.cpp
@@ -1608,6 +1608,8 @@ void UAssetRegistryImpl::WaitForCompletion()
 			FClassInheritanceBuffer InheritanceBuffer;
 			GetInheritanceContextWithRequiredLock(InterfaceScopeLock, InheritanceContext, InheritanceBuffer);
 
+			GuardedData.WaitForGathererIdleIfSynchronous();
+
 			bool bUnusedInterrupted;
 			Status = GuardedData.TickGatherer(EventContext, InheritanceContext, -1., bUnusedInterrupted);
 		}
@@ -3715,6 +3717,14 @@ void FAssetRegistryImpl::TickDeletes()
 	PRAGMA_ENABLE_DEPRECATION_WARNINGS
 }
 
+void FAssetRegistryImpl::WaitForGathererIdleIfSynchronous()
+{
+	if (GlobalGatherer && GlobalGatherer->IsSynchronous())
+	{
+		GlobalGatherer->WaitForIdle();
+	}
+}
+
 Impl::EGatherStatus FAssetRegistryImpl::TickGatherer(Impl::FEventContext& EventContext,
 	Impl::FClassInheritanceContext& InheritanceContext, const double TickStartTime, bool& bOutInterrupted,
 	TOptional<FAssetsFoundCallback> AssetsFoundCallback)
diff --git a/Engine/Source/Runtime/AssetRegistry/Private/AssetRegistryImpl.h b/Engine/Source/Runtime/AssetRegistry/Private/AssetRegistryImpl.h
--- a/Engine/Source/Runtime/AssetRegistry/Private/AssetRegistryImpl.h
+++ b/Engine/Source/Runtime/AssetRegistry/Private/AssetRegistryImpl.h
@@ -179,6 +179,8 @@ public:
 	/** Delete any temporary allocated objects that were passed out to external callers and hopefully have been deleted by now */
 	UE_DEPRECATED(5.0, "Supports deprecated functions returning pointers")
 	void TickDeletes();
+	/** Waits for the gatherer to be idle if it is operating synchronously. */
+	void WaitForGathererIdleIfSynchronous();
 	/** Callback type for TickGatherer */
 	typedef TFunctionRef<void(const TRingBuffer<FAssetData*>&)> FAssetsFoundCallback;
 	/** Consume any results from the gatherer and return its status */
