Cached Redirected package names. Saved on average 1.5s when calling DependencyDataGathered. Saved around 10s when starting the editor.

diff --git a/Engine/Source/Runtime/AssetRegistry/Private/AssetRegistry.cpp b/Engine/Source/Runtime/AssetRegistry/Private/AssetRegistry.cpp
--- a/Engine/Source/Runtime/AssetRegistry/Private/AssetRegistry.cpp
+++ b/Engine/Source/Runtime/AssetRegistry/Private/AssetRegistry.cpp
@@ -4698,6 +4698,7 @@ void FAssetRegistryImpl::DependencyDataGathered(const double TickStartTime, TRin
 	using namespace UE::AssetRegistry;
 	const bool bFlushFullBuffer = TickStartTime < 0;
 
+	TMap<FName, FName> CachedDepToRedirect;
 	while (DependsResults.Num() > 0)
 	{
 		FPackageDependencyData Result = DependsResults.PopFrontValue();
@@ -4742,9 +4743,14 @@ void FAssetRegistryImpl::DependencyDataGathered(const double TickStartTime, TRin
 				{
 					continue;
 				}
-				FName Redirected = FCoreRedirects::GetRedirectedName(ECoreRedirectFlags::Type_Package,
-					FCoreRedirectObjectName(NAME_None, NAME_None, DependencyPackageName)).PackageName;
-				DependencyPackageName = Redirected;
+				
+				FName& RedirectedName = CachedDepToRedirect.FindOrAdd(DependencyPackageName, NAME_None);
+				if (RedirectedName.IsNone())
+				{
+					RedirectedName = FCoreRedirects::GetRedirectedName(ECoreRedirectFlags::Type_Package,
+						FCoreRedirectObjectName(NAME_None, NAME_None, DependencyPackageName)).PackageName;
+				}
+				DependencyPackageName = RedirectedName;
 
 				FDependsNode::FPackageFlagSet& PackageFlagSet = PackageDependencies.FindOrAdd(DependencyPackageName);
 				PackageFlagSet.Add(FDependsNode::PackagePropertiesToByte(DependencyData.Property));
