AssetRegistry: Always create Content directories on disk for all mount points to be able to register to and detect and file changes in these directories.
This was already done later from UUnrealEdEngine::VerifyMountPointWritePermission, so this change only moves directory creation to an earlier point when the asset registry is being used, and makes sure that the monitoring behavior works correctly also when booting the editor the first time.
Fix multiple "LogDirectoryWatcher" errors in the output log after first launching a cooked editor running from pak and container files.

#rnx

diff --git a/Engine/Source/Runtime/AssetRegistry/Private/AssetRegistry.cpp b/Engine/Source/Runtime/AssetRegistry/Private/AssetRegistry.cpp
--- a/Engine/Source/Runtime/AssetRegistry/Private/AssetRegistry.cpp
+++ b/Engine/Source/Runtime/AssetRegistry/Private/AssetRegistry.cpp
@@ -12,6 +12,7 @@
 #include "Blueprint/BlueprintSupport.h"
 #include "DependsNode.h"
 #include "GenericPlatform/GenericPlatformChunkInstall.h"
+#include "GenericPlatform/GenericPlatformFile.h"
 #include "HAL/PlatformMisc.h"
 #include "HAL/ThreadHeartBeat.h"
 #include "Interfaces/IPluginManager.h"
@@ -883,24 +884,24 @@ void UAssetRegistryImpl::InitializeEvents(UE::AssetRegistry::Impl::FInitializeCo
 
 		if (DirectoryWatcher)
 		{
+			FString ContentFolder;
 			for (TArray<FString>::TConstIterator RootPathIt(Context.RootContentPaths); RootPathIt; ++RootPathIt)
 			{
 				const FString& RootPath = *RootPathIt;
-				const FString& ContentFolder = FPackageName::LongPackageNameToFilename(RootPath);
+				ContentFolder = FPackageName::LongPackageNameToFilename(RootPath);
 
-				// A missing directory here could be due to a plugin that specifies it contains content, yet has no content yet. PluginManager
-				// Mounts these folders anyway which results in them being returned from QueryRootContentPaths
-				if (IFileManager::Get().DirectoryExists(*ContentFolder))
-				{
-					FDelegateHandle NewHandle;
-					DirectoryWatcher->RegisterDirectoryChangedCallback_Handle(
+				// A missing directory here could be due to a plugin that specifies it contains content, yet has no content yet.
+				// PluginManager mounts these folders anyway which results in them being returned from QueryRootContentPaths.
+				// Make sure the directory exits on disk so that the OS level DirectoryWatcher can be used to monitor it.
+				IPlatformFile::GetPlatformPhysical().CreateDirectoryTree(*ContentFolder);
+				FDelegateHandle NewHandle;
+				DirectoryWatcher->RegisterDirectoryChangedCallback_Handle(
 						ContentFolder,
 						IDirectoryWatcher::FDirectoryChanged::CreateUObject(this, &UAssetRegistryImpl::OnDirectoryChanged),
 						NewHandle,
 						IDirectoryWatcher::WatchOptions::IncludeDirectoryChanges);
 
-					OnDirectoryChangedDelegateHandles.Add(RootPath, NewHandle);
-				}
+				OnDirectoryChangedDelegateHandles.Add(RootPath, NewHandle);
 			}
 		}
 	}
@@ -5590,8 +5591,8 @@ void UAssetRegistryImpl::OnContentPathMounted(const FString& InAssetPath, const
 		DirectoryWatcher = DirectoryWatcherModule.Get();
 		if (DirectoryWatcher)
 		{
-			// If the path doesn't exist on disk, make it so the watcher will work.
-			IFileManager::Get().MakeDirectory(*FileSystemPath, /*Tree=*/true);
+			// Make sure the directory exits on disk so that the OS level DirectoryWatcher can be used to monitor it.
+			IPlatformFile::GetPlatformPhysical().CreateDirectoryTree(*FileSystemPath);
 		}
 	}
 		
