Fix Asset Registry skip devonly tags.
FString::Split retuns false if no delim substring found.
This leads to complete ignore of KeepInDevOnly tag when configured with intended format:
"(Class=/Script/MyClass,Tag=MyTag,KeepInDevOnly)"

#tests preflight+replayrun

diff --git a/Engine/Source/Runtime/AssetRegistry/Private/AssetRegistry.cpp b/Engine/Source/Runtime/AssetRegistry/Private/AssetRegistry.cpp
--- a/Engine/Source/Runtime/AssetRegistry/Private/AssetRegistry.cpp
+++ b/Engine/Source/Runtime/AssetRegistry/Private/AssetRegistry.cpp
@@ -1250,15 +1250,19 @@ void InitializeSerializationOptionsFromIni(FAssetRegistrySerializationOptions& O
 			{
 				KeyString.TrimStartAndEndInline();
 				ValueString.TrimStartAndEndInline();
-				if (KeyString == TEXT("Class"))
+				if (KeyString.Equals(TEXT("Class"), ESearchCase::IgnoreCase))
 				{
 					ClassName = ValueString;
 				}
-				else if (KeyString == TEXT("Tag"))
+				else if (KeyString.Equals(TEXT("Tag"), ESearchCase::IgnoreCase))
 				{
 					TagName = ValueString;
 				}
-				else if (KeyString == TEXT("KeepInDevOnly"))
+			}
+			else
+			{
+				KeyString = Token.TrimStartAndEnd();
+				if (KeyString.Equals(TEXT("KeepInDevOnly"), ESearchCase::IgnoreCase))
 				{
 					bKeepInDevOnly = true;
 				}
