Updated to use FindOrAdd instead of finding and adding separately. Saved 120ms during editor startup.

diff --git a/Engine/Source/Runtime/AssetRegistry/Private/AssetRegistryState.cpp b/Engine/Source/Runtime/AssetRegistry/Private/AssetRegistryState.cpp
--- a/Engine/Source/Runtime/AssetRegistry/Private/AssetRegistryState.cpp
+++ b/Engine/Source/Runtime/AssetRegistry/Private/AssetRegistryState.cpp
@@ -2019,17 +2019,13 @@ FDependsNode* FAssetRegistryState::FindDependsNode(const FAssetIdentifier& Ident
 
 FDependsNode* FAssetRegistryState::CreateOrFindDependsNode(const FAssetIdentifier& Identifier)
 {
-	FDependsNode* FoundNode = FindDependsNode(Identifier);
-	if (FoundNode)
+	FDependsNode*& Node = CachedDependsNodes.FindOrAdd(Identifier);
+	if (Node == nullptr)
 	{
-		return FoundNode;
+		Node = new FDependsNode(Identifier);
+		NumDependsNodes++;
 	}
-
-	FDependsNode* NewNode = new FDependsNode(Identifier);
-	NumDependsNodes++;
-	CachedDependsNodes.Add(Identifier, NewNode);
-
-	return NewNode;
+	return Node;
 }
 
 bool FAssetRegistryState::RemoveDependsNode(const FAssetIdentifier& Identifier)
@@ -2106,17 +2102,13 @@ const FAssetPackageData* FAssetRegistryState::GetAssetPackageData(FName PackageN
 
 FAssetPackageData* FAssetRegistryState::CreateOrGetAssetPackageData(FName PackageName)
 {
-	FAssetPackageData** FoundData = CachedPackageData.Find(PackageName);
-	if (FoundData)
+	FAssetPackageData*& Data = CachedPackageData.FindOrAdd(PackageName);
+	if (Data == nullptr)
 	{
-		return *FoundData;
+		Data = new FAssetPackageData();
+		NumPackageData++;
 	}
-
-	FAssetPackageData* NewData = new FAssetPackageData();
-	NumPackageData++;
-	CachedPackageData.Add(PackageName, NewData);
-
-	return NewData;
+	return Data;
 }
 
 bool FAssetRegistryState::RemovePackageData(FName PackageName)
