Removed unnecessary remove before add.

diff --git a/Engine/Source/Runtime/AssetRegistry/Private/AssetRegistry.cpp b/Engine/Source/Runtime/AssetRegistry/Private/AssetRegistry.cpp
--- a/Engine/Source/Runtime/AssetRegistry/Private/AssetRegistry.cpp
+++ b/Engine/Source/Runtime/AssetRegistry/Private/AssetRegistry.cpp
@@ -4598,8 +4598,8 @@ void FAssetRegistryImpl::PostLoadAssetRegistryTags(FAssetData* AssetData)
 					FAssetDataTagMap TagsAndValues = AssetData->TagsAndValues.CopyMap();
 					for (const UObject::FAssetRegistryTag& Tag : TagsToModify)
 					{
-						TagsAndValues.Remove(Tag.Name);
-						TagsAndValues.Add(Tag.Name, Tag.Value);
+						FString& Value = TagsAndValues.FindOrAdd(Tag.Name);
+						Value = Tag.Value;
 					}
 					AssetData->TagsAndValues = FAssetDataTagMapSharedView(MoveTemp(TagsAndValues));
 				}
