Add Insight trace when waiting on asset registry completion

#rnx

diff --git a/Engine/Source/Runtime/AssetRegistry/Private/AssetRegistry.cpp b/Engine/Source/Runtime/AssetRegistry/Private/AssetRegistry.cpp
--- a/Engine/Source/Runtime/AssetRegistry/Private/AssetRegistry.cpp
+++ b/Engine/Source/Runtime/AssetRegistry/Private/AssetRegistry.cpp
@@ -1634,6 +1634,8 @@ void FAssetRegistryImpl::SearchAllAssets(Impl::FEventContext& EventContext,
 
 void UAssetRegistryImpl::WaitForCompletion()
 {
+	TRACE_CPUPROFILER_EVENT_SCOPE(UAssetRegistryImpl::WaitForCompletion);
+
 	using namespace UE::AssetRegistry::Impl;
 	LLM_SCOPE(ELLMTag::AssetRegistry);
 
@@ -1668,6 +1670,8 @@ void UAssetRegistryImpl::WaitForCompletion()
 
 void UAssetRegistryImpl::WaitForPackage(const FString& PackageName)
 {
+	TRACE_CPUPROFILER_EVENT_SCOPE(UAssetRegistryImpl::WaitForPackage);
+
 	UE::AssetRegistry::Impl::FEventContext EventContext;
 	{
 		FWriteScopeLock InterfaceScopeLock(InterfaceLock);
