AssetRegistry: Add ClearGathererCache function to save memory when cooking. Shrink the background results when the gatherer goes idle to save additional memory.
#rnx

diff --git a/Engine/Source/Editor/UnrealEd/Private/CookOnTheFlyServer.cpp b/Engine/Source/Editor/UnrealEd/Private/CookOnTheFlyServer.cpp
--- a/Engine/Source/Editor/UnrealEd/Private/CookOnTheFlyServer.cpp
+++ b/Engine/Source/Editor/UnrealEd/Private/CookOnTheFlyServer.cpp
@@ -7306,6 +7306,7 @@ void UCookOnTheFlyServer::BlockOnAssetRegistry()
 #endif
 
 	FAssetRegistryGenerator::UpdateAssetManagerDatabase();
+	AssetRegistry->ClearGathererCache();
 }
 
 void UCookOnTheFlyServer::RefreshPlatformAssetRegistries(const TArrayView<const ITargetPlatform* const>& TargetPlatforms)
diff --git a/Engine/Source/Runtime/AssetRegistry/Private/AssetDataGatherer.cpp b/Engine/Source/Runtime/AssetRegistry/Private/AssetDataGatherer.cpp
--- a/Engine/Source/Runtime/AssetRegistry/Private/AssetDataGatherer.cpp
+++ b/Engine/Source/Runtime/AssetRegistry/Private/AssetDataGatherer.cpp
@@ -3086,6 +3086,7 @@ void FAssetDataGatherer::TickInternal(bool& bOutIsTickInterrupt)
 	TArray<FString, FBatchInlineAllocator> LocalCookedPackageNamesWithoutAssetDataResults;
 	TArray<FName, FBatchInlineAllocator> LocalVerseResults;
 	bool bLoadMonolithicCache = false;
+	bool bLocalIsCacheEnabled = false;
 	double LocalLastCacheWriteTime = 0.0;
 	bool bWaitBatchCountDecremented = false;
 	bOutIsTickInterrupt = false;
@@ -3136,6 +3137,7 @@ void FAssetDataGatherer::TickInternal(bool& bOutIsTickInterrupt)
 		}
 		LocalLastCacheWriteTime = LastCacheWriteTime;
 	}
+	bLocalIsCacheEnabled = bCacheEnabled;
 
 	// Load the async cache if not yet loaded
 	if (bLoadMonolithicCache)
@@ -3276,7 +3278,7 @@ void FAssetDataGatherer::TickInternal(bool& bOutIsTickInterrupt)
 			// Add the results from a cooked package into our results on cooked package
 			LocalCookedPackageNamesWithoutAssetDataResults.Append(MoveTemp(ReadContext.CookedPackageNamesWithoutAssetData));
 			// Do not add the results from a cooked package into the map of data we keep to write out the new version of the cache file
-			bool bCachePackage = bCacheEnabled
+			bool bCachePackage = bLocalIsCacheEnabled
 				&& LocalCookedPackageNamesWithoutAssetDataResults.Num() == 0
 				&& ensure(ReadContext.AssetFileData.Type == EGatherableFileType::PackageFile);
 			if (bCachePackage)
@@ -3535,6 +3537,35 @@ void FAssetDataGatherer::WaitOnPath(FStringView InPath)
 	WaitOnPathsInternal(TArrayView<UE::AssetDataGather::Private::FPathExistence>(&QueryPath, 1), FString(), TArray<FString>());
 }
 
+void FAssetDataGatherer::ClearCache()
+{
+	FGathererScopeLock TickScopeLock(&TickLock);
+	bool bWasCacheEnabled = bCacheEnabled;
+	bCacheEnabled = false;
+	{
+		FGathererScopeLock ResultsScopeLock(&ResultsLock);
+		bUseMonolithicCache = false;
+	}
+
+	if (bWasCacheEnabled)
+	{
+		NewCachedAssetDataMap.Empty();
+		DiskCachedAssetDataMap.Empty();
+
+		for (FDiskCachedAssetData* AssetData : NewCachedAssetData)
+		{
+			delete AssetData;
+		}
+		NewCachedAssetData.Empty();
+		for (TPair<int32, FDiskCachedAssetData*> BlockData : DiskCachedAssetBlocks)
+		{
+			delete[] BlockData.Get<1>();
+		}
+		DiskCachedAssetBlocks.Empty();
+	}
+}
+
+
 void FAssetDataGatherer::ScanPathsSynchronous(const TArray<FString>& InLocalPaths, bool bForceRescan,
 	bool bIgnoreDenyListScanFilters, const FString& SaveCacheFilename, const TArray<FString>& SaveCacheLongPackageNameDirs)
 {
diff --git a/Engine/Source/Runtime/AssetRegistry/Private/AssetDataGatherer.h b/Engine/Source/Runtime/AssetRegistry/Private/AssetDataGatherer.h
--- a/Engine/Source/Runtime/AssetRegistry/Private/AssetDataGatherer.h
+++ b/Engine/Source/Runtime/AssetRegistry/Private/AssetDataGatherer.h
@@ -94,6 +94,7 @@ public:
 		TRingBuffer<FName> VerseFiles;
 
 		SIZE_T GetAllocatedSize() const { return Assets.GetAllocatedSize() + Paths.GetAllocatedSize() + Dependencies.GetAllocatedSize() + CookedPackageNamesWithoutAssetData.GetAllocatedSize() + VerseFiles.GetAllocatedSize(); }
+		void Shrink() { Assets.Trim(); Paths.Trim(); Dependencies.Trim(); CookedPackageNamesWithoutAssetData.Trim(); VerseFiles.Trim(); }
 	};
 	struct FResultContext
 	{
@@ -110,6 +111,12 @@ public:
 	void GetPackageResults(TRingBuffer<FAssetData*>& OutAssetResults, TRingBuffer<FPackageDependencyData>& OutDependencyResults);
 	/** Wait for all monitored assets under the given path to be added to search results. Returns immediately if the given path is not monitored. */
 	void WaitOnPath(FStringView LocalPath);
+	/**
+	 * Empty the cache read from disk and the cache used to write to disk. Disable further caching.
+	 * Used to save memory when cooking after the scan is complete.
+	*/
+	void ClearCache();
+
 	/**
 	 * Add a set of paths to the allow list, optionally force rescanning and ignore deny list on them,
 	 * and wait for all assets in the paths to be added to search results.
@@ -317,8 +324,6 @@ private:
 	bool bGatherAssetPackageData;
 	/** True if dependency data should be gathered. Constant during threading. */
 	bool bGatherDependsData;
-	/** True if the current process allows reading/writing of AssetDataGatherer cache files. Constant during threading. */
-	bool bCacheEnabled;
 
 
 	// Variable section for variables that are atomics read/writable from outside critical sections.
@@ -339,6 +344,8 @@ private:
 	 * Read/writable anywhere.
 	 */
 	std::atomic<bool> bSaveAsyncCacheTriggered;
+	/** True if the current process allows reading/writing of AssetDataGatherer cache files. */
+	std::atomic<bool> bCacheEnabled;
 
 	// Variable section for variables that are read/writable only within ResultsLock.
 
diff --git a/Engine/Source/Runtime/AssetRegistry/Private/AssetRegistry.cpp b/Engine/Source/Runtime/AssetRegistry/Private/AssetRegistry.cpp
--- a/Engine/Source/Runtime/AssetRegistry/Private/AssetRegistry.cpp
+++ b/Engine/Source/Runtime/AssetRegistry/Private/AssetRegistry.cpp
@@ -1668,6 +1668,25 @@ void UAssetRegistryImpl::WaitForCompletion()
 	}
 }
 
+void UAssetRegistryImpl::ClearGathererCache()
+{
+	FWriteScopeLock InterfaceScopeLock(InterfaceLock);
+	GuardedData.ClearGathererCache();
+}
+
+namespace UE::AssetRegistry
+{
+
+void FAssetRegistryImpl::ClearGathererCache()
+{
+	if (GlobalGatherer)
+	{
+		GlobalGatherer->ClearCache();
+	}
+}
+
+}
+
 void UAssetRegistryImpl::WaitForPackage(const FString& PackageName)
 {
 	TRACE_CPUPROFILER_EVENT_SCOPE(UAssetRegistryImpl::WaitForPackage);
@@ -3904,6 +3923,7 @@ Impl::EGatherStatus FAssetRegistryImpl::TickGatherer(Impl::FEventContext& EventC
 	if (OutStatus == EGatherStatus::Complete)
 	{
 		HighestPending = 0;
+		BackgroundResults.Shrink();
 
 		if (!bInitialSearchCompleted && bCanFinishInitialSearch)
 		{
diff --git a/Engine/Source/Runtime/AssetRegistry/Private/AssetRegistry.h b/Engine/Source/Runtime/AssetRegistry/Private/AssetRegistry.h
--- a/Engine/Source/Runtime/AssetRegistry/Private/AssetRegistry.h
+++ b/Engine/Source/Runtime/AssetRegistry/Private/AssetRegistry.h
@@ -107,6 +107,7 @@ public:
 	virtual bool IsSearchAllAssets() const override;
 	virtual bool IsSearchAsync() const override;
 	virtual void WaitForCompletion() override;
+	virtual void ClearGathererCache() override;
 	virtual void WaitForPackage(const FString& PackageName) override;
 	virtual void ScanPathsSynchronous(const TArray<FString>& InPaths, bool bForceRescan = false, bool bIgnoreDenyListScanFilters = false) override;
 	virtual void ScanFilesSynchronous(const TArray<FString>& InFilePaths, bool bForceRescan = false) override;
diff --git a/Engine/Source/Runtime/AssetRegistry/Private/AssetRegistryImpl.h b/Engine/Source/Runtime/AssetRegistry/Private/AssetRegistryImpl.h
--- a/Engine/Source/Runtime/AssetRegistry/Private/AssetRegistryImpl.h
+++ b/Engine/Source/Runtime/AssetRegistry/Private/AssetRegistryImpl.h
@@ -189,6 +189,7 @@ public:
 		TOptional<FAssetsFoundCallback> AssetsFoundCallback = TOptional<FAssetsFoundCallback>());
 	/** Look for and load a single AssetData result from the gatherer. */
 	void TickGatherPackage(Impl::FEventContext& EventContext, const FString& PackageName, const FString& LocalPath);
+	void ClearGathererCache();
 #if WITH_EDITOR
 	void GetProcessLoadedAssetsBatch(TArray<const UObject*>& OutLoadedAssets, uint32 BatchSize);
 	void PushProcessLoadedAssetsBatch(Impl::FEventContext& EventContext,
diff --git a/Engine/Source/Runtime/AssetRegistry/Public/AssetRegistry/IAssetRegistry.h b/Engine/Source/Runtime/AssetRegistry/Public/AssetRegistry/IAssetRegistry.h
--- a/Engine/Source/Runtime/AssetRegistry/Public/AssetRegistry/IAssetRegistry.h
+++ b/Engine/Source/Runtime/AssetRegistry/Public/AssetRegistry/IAssetRegistry.h
@@ -536,6 +536,12 @@ public:
 	UFUNCTION(BlueprintCallable, Category = "AssetRegistry")
 	virtual void WaitForCompletion() = 0;
 
+	/**
+	 * Empty the global gatherer's cache and disable further caching of scans from disk.
+	 * Used to save memory when cooking after the scan is complete.
+	*/
+	virtual void ClearGathererCache() = 0;
+
 	/** Wait for the scan of a specific package to be complete */
 	UFUNCTION(BlueprintCallable, Category = "AssetRegistry")
 	virtual void WaitForPackage(const FString& PackageName) = 0;
