AssetRegistry: AssetDataGatherer: Add synchronization to prevent deleting cache data while it is still in use, and reenable ClearCache.
#rnx

diff --git a/Engine/Source/Editor/UnrealEd/Private/CookOnTheFlyServer.cpp b/Engine/Source/Editor/UnrealEd/Private/CookOnTheFlyServer.cpp
--- a/Engine/Source/Editor/UnrealEd/Private/CookOnTheFlyServer.cpp
+++ b/Engine/Source/Editor/UnrealEd/Private/CookOnTheFlyServer.cpp
@@ -7306,7 +7306,7 @@ void UCookOnTheFlyServer::BlockOnAssetRegistry()
 #endif
 
 	FAssetRegistryGenerator::UpdateAssetManagerDatabase();
-	//AssetRegistry->ClearGathererCache();
+	AssetRegistry->ClearGathererCache();
 }
 
 void UCookOnTheFlyServer::RefreshPlatformAssetRegistries(const TArrayView<const ITargetPlatform* const>& TargetPlatforms)
diff --git a/Engine/Source/Runtime/AssetRegistry/Private/AssetDataGatherer.cpp b/Engine/Source/Runtime/AssetRegistry/Private/AssetDataGatherer.cpp
--- a/Engine/Source/Runtime/AssetRegistry/Private/AssetDataGatherer.cpp
+++ b/Engine/Source/Runtime/AssetRegistry/Private/AssetDataGatherer.cpp
@@ -2893,6 +2893,7 @@ FAssetDataGatherer::FAssetDataGatherer(const TArray<FString>& InLongPackageNames
 	, WaitBatchCount(0)
 	, NumCachedAssetFiles(0)
 	, NumUncachedAssetFiles(0)
+	, CacheInUseCount(0)
 	, bIsSavingAsyncCache(false)
 {
 	using namespace UE::AssetDataGather::Private;
@@ -3016,11 +3017,18 @@ void FAssetDataGatherer::InnerTickLoop(bool bInIsSynchronousTick, bool bContribu
 		if (bContributeToCacheSave)
 		{
 			TryReserveSaveMonolithicCache(bShouldSaveMonolithicCache, AssetsToSave);
+			if (bShouldSaveMonolithicCache)
+			{
+				CacheInUseCount++;
+			}
 		}
 	}
 	if (bShouldSaveMonolithicCache)
 	{
 		SaveCacheFileInternal(GPreloadSettings.GetMonolithicCacheFilename(), AssetsToSave, true /* bIsAsyncCache */);
+		FGathererScopeLock TickScopeLock(&TickLock);
+		check(CacheInUseCount > 0);
+		CacheInUseCount--;
 	}
 }
 
@@ -3539,16 +3547,37 @@ void FAssetDataGatherer::WaitOnPath(FStringView InPath)
 
 void FAssetDataGatherer::ClearCache()
 {
-	FGathererScopeLock TickScopeLock(&TickLock);
-	bool bWasCacheEnabled = bCacheEnabled;
-	bCacheEnabled = false;
+	bool bCacheIsInUseOnOtherThread = false;
+	bool bWasCacheEnabled = false;
 	{
-		FGathererScopeLock ResultsScopeLock(&ResultsLock);
-		bUseMonolithicCache = false;
+		FGathererScopeLock TickScopeLock(&TickLock);
+		bWasCacheEnabled = bCacheEnabled;
+		bCacheEnabled = false;
+		{
+			FGathererScopeLock ResultsScopeLock(&ResultsLock);
+			bUseMonolithicCache = false;
+		}
+		bCacheIsInUseOnOtherThread = CacheInUseCount > 0;
 	}
 
-	if (bWasCacheEnabled)
+	if (!bWasCacheEnabled)
 	{
+		return;
+	}
+
+	// Wait for any cache saves to complete because saves read the cache data we are about to delete.
+	// Saves are executed outside of the lock, but they indicate they are in progress by incrementing CacheInUseCount.
+	// CacheInUseCount is no longer incremented after bCacheEnabled=false which we set above, so starvation should not be possible.
+	while (bCacheIsInUseOnOtherThread)
+	{
+		constexpr float WaitForSaveCompleteTime = .001f;
+		FPlatformProcess::Sleep(WaitForSaveCompleteTime);
+		FGathererScopeLock TickScopeLock(&TickLock);
+		bCacheIsInUseOnOtherThread = CacheInUseCount > 0;
+	}
+
+	{
+		FGathererScopeLock TickScopeLock(&TickLock);
 		NewCachedAssetDataMap.Empty();
 		DiskCachedAssetDataMap.Empty();
 
@@ -3628,14 +3657,20 @@ void FAssetDataGatherer::WaitOnPathsInternal(TArrayView<UE::AssetDataGather::Pri
 		}
 	}
 
-	if (!SaveCacheFilename.IsEmpty())
+	if (!SaveCacheFilename.IsEmpty() && bCacheEnabled)
 	{
 		TArray<TPair<FName, FDiskCachedAssetData*>> AssetsToSave;
 		{
 			FGathererScopeLock TickScopeLock(&TickLock);
 			GetAssetsToSave(SaveCacheLongPackageNameDirs, AssetsToSave);
+			CacheInUseCount++;
 		}
 		SaveCacheFileInternal(SaveCacheFilename, AssetsToSave, false /* bIsAsyncCacheSave */);
+		{
+			FGathererScopeLock TickScopeLock(&TickLock);
+			check(CacheInUseCount > 0);
+			CacheInUseCount--;
+		}
 	}
 }
 
diff --git a/Engine/Source/Runtime/AssetRegistry/Private/AssetDataGatherer.h b/Engine/Source/Runtime/AssetRegistry/Private/AssetDataGatherer.h
--- a/Engine/Source/Runtime/AssetRegistry/Private/AssetDataGatherer.h
+++ b/Engine/Source/Runtime/AssetRegistry/Private/AssetDataGatherer.h
@@ -426,6 +426,12 @@ private:
 	int32 NumCachedAssetFiles;
 	/** How many files in the search results were not in the cache and were read by parsing the file. Read/writable only within TickLock. */
 	int32 NumUncachedAssetFiles;
+	/**
+	 * Incremented when a thread is in the middle of saving any cache and therefore the cache cannot be deleted,
+	 * decremented when the thread is done. Only allowed to be incremented when bCacheEnabled has been recently confirmed to be true.
+	 * Read/writable only within TickLock.
+	 */
+	int32 CacheInUseCount;
 	/**
 	 * True if the current TickInternal is synchronous, which may be because bIsSynchronous or because the game thread has taken over the tick for a synchronous function.
 	 * Read/writable only within TickLock.
