Cooker LLM Tags: Add more tags around allocations made when cooking.
#rnx

diff --git a/Engine/Source/Editor/Kismet/Private/BlueprintCompilationManager.cpp b/Engine/Source/Editor/Kismet/Private/BlueprintCompilationManager.cpp
--- a/Engine/Source/Editor/Kismet/Private/BlueprintCompilationManager.cpp
+++ b/Engine/Source/Editor/Kismet/Private/BlueprintCompilationManager.cpp
@@ -19,6 +19,7 @@
 #include "Engine/TimelineTemplate.h"
 #include "FileHelpers.h"
 #include "FindInBlueprintManager.h"
+#include "HAL/LowLevelMemTracker.h"
 #include "IMessageLogListing.h"
 #include "K2Node_CreateDelegate.h"
 #include "K2Node_CustomEvent.h"
@@ -3327,6 +3328,7 @@ void FBlueprintCompilationManager::FlushCompilationQueue(FUObjectSerializeContex
 {
 	if(BPCMImpl)
 	{
+		LLM_SCOPE_BYNAME(TEXT("Blueprints"));
 		BPCMImpl->FlushCompilationQueueImpl(false, nullptr, nullptr, InLoadContext);
 
 		// We can't support save on compile or keeping old CDOs from GCing when reinstancing is deferred:
diff --git a/Engine/Source/Runtime/AssetRegistry/Private/AssetRegistry.cpp b/Engine/Source/Runtime/AssetRegistry/Private/AssetRegistry.cpp
--- a/Engine/Source/Runtime/AssetRegistry/Private/AssetRegistry.cpp
+++ b/Engine/Source/Runtime/AssetRegistry/Private/AssetRegistry.cpp
@@ -307,6 +307,7 @@ private:
 
 	ELoadResult TryLoad()
 	{
+		LLM_SCOPE(ELLMTag::AssetRegistry);
 		checkf(!ARPath.IsEmpty(), TEXT("TryLoad must not be called until after TrySetPath has succeeded."));
 
 		FAssetRegistryLoadOptions Options;
@@ -689,6 +690,7 @@ UAssetRegistryImpl::UAssetRegistryImpl(const FObjectInitializer& ObjectInitializ
 	UE::AssetRegistry::Impl::FInitializeContext Context{ *this };
 
 	{
+		LLM_SCOPE(ELLMTag::AssetRegistry);
 		FWriteScopeLock InterfaceScopeLock(InterfaceLock);
 		GetInheritanceContextWithRequiredLock(InterfaceScopeLock, Context.InheritanceContext,
 			Context.InheritanceBuffer);
@@ -764,8 +766,6 @@ void FAssetRegistryImpl::LoadPremadeAssetRegistry(Impl::FEventContext& EventCont
 
 void FAssetRegistryImpl::Initialize(Impl::FInitializeContext& Context)
 {
-	LLM_SCOPE(ELLMTag::AssetRegistry);
-
 	const double StartupStartTime = FPlatformTime::Seconds();
 
 	bInitialSearchStarted = false;
@@ -1075,6 +1075,7 @@ void UAssetRegistryImpl::OnPluginLoadingPhaseComplete(ELoadingPhase::Type Loadin
 		return;
 	}
 	{
+		LLM_SCOPE(ELLMTag::AssetRegistry);
 		FWriteScopeLock InterfaceScopeLock(InterfaceLock);
 		GuardedData.OnPostEngineInit(bPhaseSuccessful);
 	}
@@ -1316,6 +1317,7 @@ void InitializeSerializationOptionsFromIni(FAssetRegistrySerializationOptions& O
 
 void FAssetRegistryImpl::CollectCodeGeneratorClasses()
 {
+	LLM_SCOPE(ELLMTag::AssetRegistry); // Tagged here instead of a higher level because it can occur even when reading
 	// Only refresh the list if our registered classes have changed
 	if (ClassGeneratorNamesRegisteredClassesVersionNumber == GetRegisteredClassesVersionNumber())
 	{
@@ -1374,6 +1376,7 @@ void FAssetRegistryImpl::CollectCodeGeneratorClasses()
 
 void UAssetRegistryImpl::OnRefreshNativeClasses()
 {
+	LLM_SCOPE(ELLMTag::AssetRegistry);
 	FWriteScopeLock InterfaceScopeLock(InterfaceLock);
 	GuardedData.RefreshNativeClasses();
 }
@@ -1402,6 +1405,7 @@ void UAssetRegistryImpl::OnFEngineLoopInitCompleteSearchAllAssets()
 
 void UAssetRegistryImpl::OnAssetDependencyGathererRegistered()
 {
+	LLM_SCOPE(ELLMTag::AssetRegistry);
 	FWriteScopeLock InterfaceScopeLock(InterfaceLock);
 	GuardedData.OnAssetDependencyGathererRegistered();
 }
@@ -1409,12 +1413,14 @@ void UAssetRegistryImpl::OnAssetDependencyGathererRegistered()
 
 void UAssetRegistryImpl::OnEnginePreExit()
 {
+	LLM_SCOPE(ELLMTag::AssetRegistry);
 	FWriteScopeLock InterfaceScopeLock(InterfaceLock);
 	GuardedData.OnEnginePreExit();
 }
 
 void UAssetRegistryImpl::FinishDestroy()
 {
+	LLM_SCOPE(ELLMTag::AssetRegistry);
 	{
 		FWriteScopeLock InterfaceScopeLock(InterfaceLock);
 
@@ -1548,6 +1554,7 @@ void UAssetRegistryImpl::SearchAllAssets(bool bSynchronousSearch)
 
 	FEventContext EventContext;
 	{
+		LLM_SCOPE(ELLMTag::AssetRegistry);
 		FWriteScopeLock InterfaceScopeLock(InterfaceLock);
 		FClassInheritanceContext InheritanceContext;
 		FClassInheritanceBuffer InheritanceBuffer;
@@ -1586,8 +1593,6 @@ namespace UE::AssetRegistry
 void FAssetRegistryImpl::SearchAllAssets(Impl::FEventContext& EventContext,
 	Impl::FClassInheritanceContext& InheritanceContext, bool bSynchronousSearch)
 {
-	LLM_SCOPE(ELLMTag::AssetRegistry);
-
 	ConstructGatherer();
 	FAssetDataGatherer& Gatherer = *GlobalGatherer;
 	if (Gatherer.IsSynchronous())
@@ -1637,13 +1642,13 @@ void UAssetRegistryImpl::WaitForCompletion()
 	TRACE_CPUPROFILER_EVENT_SCOPE(UAssetRegistryImpl::WaitForCompletion);
 
 	using namespace UE::AssetRegistry::Impl;
-	LLM_SCOPE(ELLMTag::AssetRegistry);
 
 	for (;;)
 	{
 		FEventContext EventContext;
 		EGatherStatus Status;
 		{
+			LLM_SCOPE(ELLMTag::AssetRegistry);
 			FWriteScopeLock InterfaceScopeLock(InterfaceLock);
 			FClassInheritanceContext InheritanceContext;
 			FClassInheritanceBuffer InheritanceBuffer;
@@ -1670,6 +1675,7 @@ void UAssetRegistryImpl::WaitForCompletion()
 
 void UAssetRegistryImpl::ClearGathererCache()
 {
+	LLM_SCOPE(ELLMTag::AssetRegistry);
 	FWriteScopeLock InterfaceScopeLock(InterfaceLock);
 	GuardedData.ClearGathererCache();
 }
@@ -1693,6 +1699,7 @@ void UAssetRegistryImpl::WaitForPackage(const FString& PackageName)
 
 	UE::AssetRegistry::Impl::FEventContext EventContext;
 	{
+		LLM_SCOPE(ELLMTag::AssetRegistry);
 		FWriteScopeLock InterfaceScopeLock(InterfaceLock);
 		if (GuardedData.IsLoadingAssets())
 		{
@@ -2488,6 +2495,7 @@ bool IAssetRegistry::K2_GetReferencers(FName PackageName, const FAssetRegistryDe
 
 const FAssetPackageData* UAssetRegistryImpl::GetAssetPackageData(FName PackageName) const
 {
+	LLM_SCOPE(ELLMTag::AssetRegistry);
 	FWriteScopeLock InterfaceScopeLock(InterfaceLock);
 	PRAGMA_DISABLE_DEPRECATION_WARNINGS
 	return const_cast<UE::AssetRegistry::FAssetRegistryImpl&>(GuardedData).GetAssetPackageData(PackageName);
@@ -3366,6 +3374,7 @@ bool UAssetRegistryImpl::AddPath(const FString& PathToAdd)
 	UE::AssetRegistry::Impl::FEventContext EventContext;
 	bool bResult;
 	{
+		LLM_SCOPE(ELLMTag::AssetRegistry);
 		FWriteScopeLock InterfaceScopeLock(InterfaceLock);
 		bResult = GuardedData.AddPath(EventContext, PathToAdd);
 	}
@@ -3402,6 +3411,7 @@ bool UAssetRegistryImpl::RemovePath(const FString& PathToRemove)
 	UE::AssetRegistry::Impl::FEventContext EventContext;
 	bool bResult;
 	{
+		LLM_SCOPE(ELLMTag::AssetRegistry);
 		FWriteScopeLock InterfaceScopeLock(InterfaceLock);
 		bResult = GuardedData.RemoveAssetPath(EventContext, FName(*PathToRemove));
 	}
@@ -3448,6 +3458,7 @@ void UAssetRegistryImpl::ScanPathsSynchronousInternal(const TArray<FString>& InD
 		bInForceRescan, bInIgnoreDenyListScanFilters, nullptr /* OutFindAssets */);
 
 	{
+		LLM_SCOPE(ELLMTag::AssetRegistry);
 		FWriteScopeLock InterfaceScopeLock(InterfaceLock);
 		GetInheritanceContextWithRequiredLock(InterfaceScopeLock, InheritanceContext, InheritanceBuffer);
 
@@ -3482,6 +3493,7 @@ void UAssetRegistryImpl::ScanPathsSynchronousInternal(const TArray<FString>& InD
 
 void UAssetRegistryImpl::PrioritizeSearchPath(const FString& PathToPrioritize)
 {
+	LLM_SCOPE(ELLMTag::AssetRegistry);
 	FWriteScopeLock InterfaceScopeLock(InterfaceLock);
 	GuardedData.PrioritizeSearchPath(PathToPrioritize);
 }
@@ -3537,6 +3549,7 @@ void UAssetRegistryImpl::AssetCreated(UObject* NewAsset)
 		bool bShouldSkipAssset;
 		UE::AssetRegistry::Impl::FEventContext EventContext;
 		{
+			LLM_SCOPE(ELLMTag::AssetRegistry);
 			FWriteScopeLock InterfaceScopeLock(InterfaceLock);
 			// If this package was marked as an empty package before, it is no longer empty, so remove it from the list
 			GuardedData.RemoveEmptyPackage(NewPackage->GetFName());
@@ -3570,6 +3583,7 @@ void UAssetRegistryImpl::AssetDeleted(UObject* DeletedAsset)
 
 		bool bShouldSkipAsset;
 		{
+			LLM_SCOPE(ELLMTag::AssetRegistry);
 			FWriteScopeLock InterfaceScopeLock(InterfaceLock);
 
 			// Deleting the last asset in a package causes the package to be garbage collected.
@@ -3637,6 +3651,7 @@ void UAssetRegistryImpl::AssetRenamed(const UObject* RenamedAsset, const FString
 		bool bShouldSkipAsset;
 		UE::AssetRegistry::Impl::FEventContext EventContext;
 		{
+			LLM_SCOPE(ELLMTag::AssetRegistry);
 			FWriteScopeLock InterfaceScopeLock(InterfaceLock);
 			GuardedData.RemoveEmptyPackage(NewPackage->GetFName());
 
@@ -3666,6 +3681,7 @@ void UAssetRegistryImpl::AssetSaved(const UObject& SavedAsset)
 	{
 		return;
 	}
+	LLM_SCOPE(ELLMTag::AssetRegistry);
 
 	FWriteScopeLock InterfaceScopeLock(InterfaceLock);
 	GuardedData.AddLoadedAssetToProcess(SavedAsset);
@@ -3679,6 +3695,7 @@ void UAssetRegistryImpl::AssetTagsFinalized(const UObject& FinalizedAsset)
 	{
 		return;
 	}
+	LLM_SCOPE(ELLMTag::AssetRegistry);
 
 	FWriteScopeLock InterfaceScopeLock(InterfaceLock);
 	GuardedData.AddLoadedAssetToProcess(FinalizedAsset);
@@ -3691,6 +3708,7 @@ void UAssetRegistryImpl::PackageDeleted(UPackage* DeletedPackage)
 	UE::AssetRegistry::Impl::FEventContext EventContext;
 	if (ensure(DeletedPackage))
 	{
+		LLM_SCOPE(ELLMTag::AssetRegistry);
 		FWriteScopeLock InterfaceScopeLock(InterfaceLock);
 		GuardedData.RemovePackageData(EventContext, DeletedPackage->GetFName());
 	}
@@ -3716,7 +3734,6 @@ bool FAssetRegistryImpl::IsLoadingAssets() const
 void UAssetRegistryImpl::Tick(float DeltaTime)
 {
 	checkf(IsInGameThread(), TEXT("The tick function executes deferred loads and events and must be on the game thread to do so."));
-	LLM_SCOPE(ELLMTag::AssetRegistry);
 
 	UE::AssetRegistry::Impl::EGatherStatus Status = UE::AssetRegistry::Impl::EGatherStatus::Active;
 	double TickStartTime = -1; // Force a full flush if DeltaTime < 0
@@ -3732,6 +3749,7 @@ void UAssetRegistryImpl::Tick(float DeltaTime)
 		bInterrupted = false;
 		UE::AssetRegistry::Impl::FEventContext EventContext;
 		{
+			LLM_SCOPE(ELLMTag::AssetRegistry);
 			FWriteScopeLock InterfaceScopeLock(InterfaceLock);
 			UE::AssetRegistry::Impl::FClassInheritanceContext InheritanceContext;
 			UE::AssetRegistry::Impl::FClassInheritanceBuffer InheritanceBuffer;
@@ -4130,6 +4148,7 @@ void UAssetRegistryImpl::Serialize(FArchive& Ar)
 {
 	UE::AssetRegistry::Impl::FEventContext EventContext;
 	{
+		LLM_SCOPE(ELLMTag::AssetRegistry);
 		FWriteScopeLock InterfaceScopeLock(InterfaceLock);
 		GuardedData.Serialize(Ar, EventContext);
 	}
@@ -4164,6 +4183,7 @@ void UAssetRegistryImpl::AppendState(const FAssetRegistryState& InState)
 {
 	UE::AssetRegistry::Impl::FEventContext EventContext;
 	{
+		LLM_SCOPE(ELLMTag::AssetRegistry);
 		FWriteScopeLock InterfaceScopeLock(InterfaceLock);
 		GuardedData.AppendState(EventContext, InState);
 	}
@@ -4191,8 +4211,6 @@ void FAssetRegistryImpl::CachePathsFromState(Impl::FEventContext& EventContext,
 {
 	SCOPED_BOOT_TIMING("FAssetRegistryImpl::CachePathsFromState");
 
-	LLM_SCOPE(ELLMTag::AssetRegistry);
-
 	// Refreshes ClassGeneratorNames if out of date due to module load
 	CollectCodeGeneratorClasses();
 
@@ -5208,6 +5226,7 @@ void UAssetRegistryImpl::OnDirectoryChanged(const TArray<FFileChangeData>& FileC
 
 	UE::AssetRegistry::Impl::FEventContext EventContext;
 	{
+		LLM_SCOPE(ELLMTag::AssetRegistry);
 		FWriteScopeLock InterfaceScopeLock(InterfaceLock);
 		UE::AssetRegistry::Impl::FClassInheritanceContext InheritanceContext;
 		UE::AssetRegistry::Impl::FClassInheritanceBuffer InheritanceBuffer;
@@ -5370,6 +5389,7 @@ void FAssetRegistryImpl::OnDirectoryChanged(Impl::FEventContext& EventContext,
 
 void UAssetRegistryImpl::OnAssetLoaded(UObject *AssetLoaded)
 {
+	LLM_SCOPE(ELLMTag::AssetRegistry);
 	FWriteScopeLock InterfaceScopeLock(InterfaceLock);
 	GuardedData.AddLoadedAssetToProcess(*AssetLoaded);
 }
@@ -5399,6 +5419,7 @@ void UAssetRegistryImpl::ProcessLoadedAssetsToUpdateCache(UE::AssetRegistry::Imp
 	TArray<FAssetData, TInlineAllocator<BatchSize>> BatchAssetDatas;
 
 	{
+		LLM_SCOPE(ELLMTag::AssetRegistry);
 		FWriteScopeLock InterfaceScopeLock(InterfaceLock);
 		GuardedData.GetProcessLoadedAssetsBatch(BatchObjects, BatchSize);
 		if (BatchObjects.Num() == 0)
@@ -5435,6 +5456,7 @@ void UAssetRegistryImpl::ProcessLoadedAssetsToUpdateCache(UE::AssetRegistry::Imp
 			}
 		}
 
+		LLM_SCOPE(ELLMTag::AssetRegistry);
 		FWriteScopeLock InterfaceScopeLock(InterfaceLock);
 		GuardedData.PushProcessLoadedAssetsBatch(EventContext, BatchAssetDatas,
 			TArrayView<const UObject*>(BatchObjects).Slice(Index, CurrentBatchSize-Index));
@@ -5551,6 +5573,7 @@ void UAssetRegistryImpl::ScanModifiedAssetFiles(const TArray<FString>& InFilePat
 {
 	UE::AssetRegistry::Impl::FEventContext EventContext;
 	{
+		LLM_SCOPE(ELLMTag::AssetRegistry);
 		FWriteScopeLock InterfaceScopeLock(InterfaceLock);
 		UE::AssetRegistry::Impl::FClassInheritanceContext InheritanceContext;
 		UE::AssetRegistry::Impl::FClassInheritanceBuffer InheritanceBuffer;
@@ -5652,6 +5675,7 @@ void UAssetRegistryImpl::OnContentPathMounted(const FString& InAssetPath, const
 	IDirectoryWatcher* DirectoryWatcher = nullptr;
 	if (GIsEditor) 	// In-game doesn't listen for directory changes
 	{
+		LLM_SCOPE(ELLMTag::AssetRegistry);
 		DirectoryWatcher = DirectoryWatcherModule.Get();
 		if (DirectoryWatcher)
 		{
@@ -5664,6 +5688,7 @@ void UAssetRegistryImpl::OnContentPathMounted(const FString& InAssetPath, const
 
 	UE::AssetRegistry::Impl::FEventContext EventContext;
 	{
+		LLM_SCOPE(ELLMTag::AssetRegistry);
 		FWriteScopeLock InterfaceScopeLock(InterfaceLock);
 		UE::AssetRegistry::Impl::FClassInheritanceContext InheritanceContext;
 		UE::AssetRegistry::Impl::FClassInheritanceBuffer InheritanceBuffer;
@@ -5742,6 +5767,7 @@ void UAssetRegistryImpl::OnContentPathDismounted(const FString& InAssetPath, con
 
 	UE::AssetRegistry::Impl::FEventContext EventContext;
 	{
+		LLM_SCOPE(ELLMTag::AssetRegistry);
 		FWriteScopeLock InterfaceScopeLock(InterfaceLock);
 		GuardedData.OnContentPathDismounted(EventContext, InAssetPath, AssetPathNoTrailingSlash, FileSystemPath);
 
@@ -5816,6 +5842,7 @@ void FAssetRegistryImpl::OnContentPathDismounted(Impl::FEventContext& EventConte
 void UAssetRegistryImpl::SetTemporaryCachingMode(bool bEnable)
 {
 	checkf(IsInGameThread(), TEXT("Changing Caching mode is only available on the game thread because it affects behavior on all threads"));
+	LLM_SCOPE(ELLMTag::AssetRegistry);
 	FWriteScopeLock InterfaceScopeLock(InterfaceLock);
 	GuardedData.SetTemporaryCachingMode(bEnable);
 }
@@ -5843,6 +5870,7 @@ void FAssetRegistryImpl::SetTemporaryCachingMode(bool bEnable)
 void UAssetRegistryImpl::SetTemporaryCachingModeInvalidated()
 {
 	checkf(IsInGameThread(), TEXT("Invalidating temporary cache is only available on the game thread because it affects behavior on all threads"));
+	LLM_SCOPE(ELLMTag::AssetRegistry);
 	FWriteScopeLock InterfaceScopeLock(InterfaceLock);
 	GuardedData.SetTemporaryCachingModeInvalidated();
 }
@@ -5894,6 +5922,7 @@ void FAssetRegistryImpl::AddCachedBPClassParent(const FTopLevelAssetPath& ClassP
 void FAssetRegistryImpl::UpdateInheritanceBuffer(Impl::FClassInheritanceBuffer& OutBuffer) const
 {
 	TRACE_CPUPROFILER_EVENT_SCOPE(UAssetRegistryImpl::UpdateTemporaryCaches)
+	LLM_SCOPE(ELLMTag::AssetRegistry);
 
 	TMap<UClass*, TSet<UClass*>> NativeSubclasses = GetAllDerivedClasses();
 
@@ -6184,6 +6213,7 @@ void UAssetRegistryImpl::SetManageReferences(const TMultiMap<FAssetIdentifier, F
 {
 	// For performance reasons we call the ShouldSetManager callback when inside the lock. Licensee UAssetManagers
 	// are responsible for not calling AssetRegistry functions from ShouldSetManager as that would create a deadlock
+	LLM_SCOPE(ELLMTag::AssetRegistry);
 	FWriteScopeLock InterfaceScopeLock(InterfaceLock);
 	GuardedData.SetManageReferences(ManagerMap, bClearExisting, RecurseType, ExistingManagedNodes, ShouldSetManager);
 }
@@ -6372,6 +6402,7 @@ bool UAssetRegistryImpl::SetPrimaryAssetIdForObjectPath(const FSoftObjectPath& O
 	UE::AssetRegistry::Impl::FEventContext EventContext;
 	bool bResult;
 	{
+		LLM_SCOPE(ELLMTag::AssetRegistry);
 		FWriteScopeLock InterfaceScopeLock(InterfaceLock);
 		bResult = GuardedData.SetPrimaryAssetIdForObjectPath(EventContext, ObjectPath, PrimaryAssetId);
 	}
diff --git a/Engine/Source/Runtime/AssetRegistry/Private/AssetRegistryState.cpp b/Engine/Source/Runtime/AssetRegistry/Private/AssetRegistryState.cpp
--- a/Engine/Source/Runtime/AssetRegistry/Private/AssetRegistryState.cpp
+++ b/Engine/Source/Runtime/AssetRegistry/Private/AssetRegistryState.cpp
@@ -179,7 +179,6 @@ void FAssetRegistryState::FilterTags(const FAssetDataTagMapSharedView& InTagsAnd
 void FAssetRegistryState::InitializeFromExistingAndPrune(const FAssetRegistryState & ExistingState, const TSet<FName>& RequiredPackages, const TSet<FName>& RemovePackages,
 	const TSet<int32> ChunksToKeep, const FAssetRegistrySerializationOptions& Options)
 {
-	LLM_SCOPE(ELLMTag::AssetRegistry);
 	const bool bIsFilteredByChunkId = ChunksToKeep.Num() != 0;
 	const bool bIsFilteredByRequiredPackages = RequiredPackages.Num() != 0;
 	const bool bIsFilteredByRemovedPackages = RemovePackages.Num() != 0;
@@ -322,7 +321,6 @@ void FAssetRegistryState::InitializeFromExistingAndPrune(const FAssetRegistrySta
 void FAssetRegistryState::InitializeFromExisting(const FAssetDataMap& AssetDataMap, const TMap<FAssetIdentifier, FDependsNode*>& DependsNodeMap, 
 	const TMap<FName, FAssetPackageData*>& AssetPackageDataMap, const FAssetRegistrySerializationOptions& Options, EInitializationMode InInitializationMode)
 {
-	LLM_SCOPE(ELLMTag::AssetRegistry);
 	if (InInitializationMode == EInitializationMode::Rebuild)
 	{
 		Reset();
@@ -1164,8 +1162,6 @@ bool FAssetRegistryState::Save(FArchive& OriginalAr, const FAssetRegistrySeriali
 
 bool FAssetRegistryState::Load(FArchive& OriginalAr, const FAssetRegistryLoadOptions& Options, FAssetRegistryVersion::Type* OutVersion)
 {
-	LLM_SCOPE(ELLMTag::AssetRegistry);
-
 	FAssetRegistryHeader Header;
 	Header.SerializeHeader(OriginalAr);
 	if (OutVersion != nullptr)
diff --git a/Engine/Source/Runtime/CoreUObject/Private/UObject/LinkerLoad.cpp b/Engine/Source/Runtime/CoreUObject/Private/UObject/LinkerLoad.cpp
--- a/Engine/Source/Runtime/CoreUObject/Private/UObject/LinkerLoad.cpp
+++ b/Engine/Source/Runtime/CoreUObject/Private/UObject/LinkerLoad.cpp
@@ -76,6 +76,8 @@ DECLARE_FLOAT_ACCUMULATOR_STAT(TEXT("Fixup editor-only flags time"), STAT_Editor
 
 FName FLinkerLoad::NAME_LoadErrors("LoadErrors");
 
+LLM_DEFINE_TAG(UObject_Linker);
+
 /**
 * Helper function to determine and trace the most important asset class.
 */
@@ -483,6 +485,7 @@ FLinkerLoad* FLinkerLoad::CreateLinker(FUObjectSerializeContext* LoadContext, UP
 FLinkerLoad* FLinkerLoad::CreateLinker(FUObjectSerializeContext* LoadContext, UPackage* Parent, const FPackagePath& PackagePath, uint32 LoadFlags, FArchive* InLoader /*= nullptr*/, const FLinkerInstancingContext* InstancingContext /*= nullptr*/)
 {
 	check(LoadContext);
+	LLM_SCOPE_BYTAG(UObject_Linker);
 
 #if USE_CIRCULAR_DEPENDENCY_LOAD_DEFERRING
 	// we don't want the linker permanently created with the 
@@ -721,6 +724,7 @@ FUObjectSerializeContext* FLinkerLoad::GetSerializeContext()
 FLinkerLoad::ELinkerStatus FLinkerLoad::ProcessPackageSummary(TMap<TPair<FName, FPackageIndex>, FPackageIndex>* ObjectNameWithOuterToExportMap)
 {
 	TRACE_LOADTIME_BEGIN_PROCESS_SUMMARY(this);
+	LLM_SCOPE_BYTAG(UObject_Linker);
 
 	ELinkerStatus Status = LINKER_Loaded;
 	{
@@ -980,6 +984,7 @@ FLinkerLoad::FLinkerLoad(UPackage* InParent, const FPackagePath& InPackagePath,
 #endif // USE_CIRCULAR_DEPENDENCY_LOAD_DEFERRING
 {
 	static_assert((ExportHashCount & (ExportHashCount - 1)) == 0, "ExportHashCount must be power of two");
+	LLM_SCOPE_BYTAG(UObject_Linker);
 
 	if (PackagePath.GetHeaderExtension() == EPackageExtension::Unspecified)
 	{
@@ -4304,6 +4309,7 @@ UObject* FLinkerLoad::Create( UClass* ObjectClass, FName ObjectName, UObject* Ou
 
 void FLinkerLoad::Preload( UObject* Object )
 {
+	LLM_SCOPE_BYTAG(UObject_Linker);
 	//check(IsValidLowLevel());
 	check(Object);
 
diff --git a/Engine/Source/Runtime/CoreUObject/Private/UObject/UObjectBase.cpp b/Engine/Source/Runtime/CoreUObject/Private/UObject/UObjectBase.cpp
--- a/Engine/Source/Runtime/CoreUObject/Private/UObject/UObjectBase.cpp
+++ b/Engine/Source/Runtime/CoreUObject/Private/UObject/UObjectBase.cpp
@@ -469,6 +469,7 @@ static FAutoConsoleCommand DumpPendingUObjectModulesCmd(
 /** Enqueue the registration for this object. */
 void UObjectBase::Register(const TCHAR* PackageName,const TCHAR* InName)
 {
+	LLM_SCOPE(ELLMTag::UObject);
 	TMap<UObjectBase*, FPendingRegistrantInfo>& PendingRegistrants = FPendingRegistrantInfo::GetMap();
 
 	FPendingRegistrant* PendingRegistration = new FPendingRegistrant(this);
@@ -545,6 +546,7 @@ static void UObjectProcessRegistrants()
 
 void UObjectForceRegistration(UObjectBase* Object, bool bCheckForModuleRelease)
 {
+	LLM_SCOPE(ELLMTag::UObject);
 	TMap<UObjectBase*, FPendingRegistrantInfo>& PendingRegistrants = FPendingRegistrantInfo::GetMap();
 
 	FPendingRegistrantInfo* Info = PendingRegistrants.Find(Object);
@@ -687,6 +689,7 @@ void UClassRegisterAllCompiledInClasses()
 	TArray<UClass*> AddedClasses;
 #endif
 	SCOPED_BOOT_TIMING("UClassRegisterAllCompiledInClasses");
+	LLM_SCOPE(ELLMTag::UObject);
 
 	FClassDeferredRegistry& Registry = FClassDeferredRegistry::Get();
 
diff --git a/Engine/Source/Runtime/CoreUObject/Private/UObject/UObjectLinker.cpp b/Engine/Source/Runtime/CoreUObject/Private/UObject/UObjectLinker.cpp
--- a/Engine/Source/Runtime/CoreUObject/Private/UObject/UObjectLinker.cpp
+++ b/Engine/Source/Runtime/CoreUObject/Private/UObject/UObjectLinker.cpp
@@ -136,6 +136,7 @@ void UObject::SetLinker( FLinkerLoad* LinkerLoad, int32 LinkerIndex, bool bShoul
 	}
 	if (Existing.Linker != LinkerLoad || Existing.LinkerIndex != LinkerIndex)
 	{
+		LLM_SCOPE_BYTAG(UObject_Linker);
 		LinkerAnnotation.AddAnnotation(this, FLinkerIndexPair(LinkerLoad, LinkerIndex));
 	}
 	if (bShouldDetachExisting)
diff --git a/Engine/Source/Runtime/CoreUObject/Public/UObject/LinkerLoad.h b/Engine/Source/Runtime/CoreUObject/Public/UObject/LinkerLoad.h
--- a/Engine/Source/Runtime/CoreUObject/Public/UObject/LinkerLoad.h
+++ b/Engine/Source/Runtime/CoreUObject/Public/UObject/LinkerLoad.h
@@ -8,6 +8,7 @@
 #include "Containers/UnrealString.h"
 #include "CoreGlobals.h"
 #include "CoreMinimal.h"
+#include "HAL/LowLevelMemTracker.h"
 #include "HAL/PlatformMath.h"
 #include "HAL/ThreadSafeCounter.h"
 #include "Misc/AssertionMacros.h"
@@ -52,6 +53,8 @@ struct FScopedSlowTask;
 struct FUObjectSerializeContext;
 template <typename FuncType> class TFunction;
 
+LLM_DECLARE_TAG(UObject_Linker);
+
 /*----------------------------------------------------------------------------
 	FLinkerLoad.
 ----------------------------------------------------------------------------*/
diff --git a/Engine/Source/Runtime/Engine/Private/DataTable.cpp b/Engine/Source/Runtime/Engine/Private/DataTable.cpp
--- a/Engine/Source/Runtime/Engine/Private/DataTable.cpp
+++ b/Engine/Source/Runtime/Engine/Private/DataTable.cpp
@@ -2,6 +2,7 @@
 
 #include "Engine/DataTable.h"
 #include "AssetRegistry/AssetData.h"
+#include "HAL/LowLevelMemTracker.h"
 #include "Internationalization/TextPackageNamespaceUtil.h"
 #include "Internationalization/StabilizeLocalizationKeys.h"
 #include "Serialization/PropertyLocalizationDataGathering.h"
@@ -17,6 +18,8 @@
 
 #include UE_INLINE_GENERATED_CPP_BY_NAME(DataTable)
 
+LLM_DEFINE_TAG(DataTable);
+
 namespace
 {
 #if WITH_EDITORONLY_DATA
@@ -229,6 +232,7 @@ void UDataTable::HandleDataTableChanged(FName ChangedRowName)
 void UDataTable::Serialize(FStructuredArchiveRecord Record)
 {
 	FArchive& BaseArchive = Record.GetUnderlyingArchive();
+	LLM_SCOPE_BYTAG(DataTable);
 
 #if WITH_EDITORONLY_DATA
 	// Make sure and update RowStructName before calling the parent Serialize (which will save the properties)
