AssetRegistry: Update function comment for GetAssetsByClass to indicate it requires a TopLevelAssetPath.
#rnx

diff --git a/Engine/Source/Runtime/AssetRegistry/Public/AssetRegistry/IAssetRegistry.h b/Engine/Source/Runtime/AssetRegistry/Public/AssetRegistry/IAssetRegistry.h
--- a/Engine/Source/Runtime/AssetRegistry/Public/AssetRegistry/IAssetRegistry.h
+++ b/Engine/Source/Runtime/AssetRegistry/Public/AssetRegistry/IAssetRegistry.h
@@ -161,7 +161,7 @@ public:
 	/**
 	 * Gets asset data for all assets with the supplied class
 	 *
-	 * @param ClassName the class name of the assets requested
+	 * @param ClassPathName the full path of the class name of the assets requested, in a TopLevelAssetPath structure.
 	 * @param OutAssetData the list of assets in this path
 	 * @param bSearchSubClasses if true, all subclasses of the passed in class will be searched as well
 	 */
