AssetRegistry: Fix PrioritizeSearchPath, which is used by the ContentBrowser to avoid waiting for the entire AR search. AssetDataDiscovery will return false if it has already finished the directory scan; even if it returns false we still may need to prioritize the file scans in AssetDataGatherer.
#rnx

diff --git a/Engine/Source/Runtime/AssetRegistry/Private/AssetDataGatherer.cpp b/Engine/Source/Runtime/AssetRegistry/Private/AssetDataGatherer.cpp
--- a/Engine/Source/Runtime/AssetRegistry/Private/AssetDataGatherer.cpp
+++ b/Engine/Source/Runtime/AssetRegistry/Private/AssetDataGatherer.cpp
@@ -4186,10 +4186,7 @@ void FAssetDataGatherer::SetDirectoryProperties(FStringView LocalPath, const UE:
 		return;
 	}
 
-	if (!Discovery->TrySetDirectoryProperties(LocalAbsPath, InProperties, false))
-	{
-		return;
-	}
+	Discovery->TrySetDirectoryProperties(LocalAbsPath, InProperties, false);
 
 	{
 		FGathererScopeLock ResultsScopeLock(&ResultsLock);
