AssetRegistry: Optimize a slow linear search in WaitForPackage by changing BackgroundResults.Assets from a TArray to a TMultiMap.
#rnx

diff --git a/Engine/Source/Runtime/AssetRegistry/Private/AssetDataGatherer.cpp b/Engine/Source/Runtime/AssetRegistry/Private/AssetDataGatherer.cpp
--- a/Engine/Source/Runtime/AssetRegistry/Private/AssetDataGatherer.cpp
+++ b/Engine/Source/Runtime/AssetRegistry/Private/AssetDataGatherer.cpp
@@ -3493,9 +3493,18 @@ void FAssetDataGatherer::GetAndTrimSearchResults(FResults& InOutResults, FResult
 		InArray.Reset();
 	};
 
-	MoveAppendRangeToRingBuffer(InOutResults.Assets, AssetResults);
+	for (FAssetData* AssetData : AssetResults)
+	{
+		InOutResults.Assets.Add(AssetData->PackageName, AssetData);
+	}
+	AssetResults.Reset();
 	MoveAppendRangeToRingBuffer(InOutResults.Paths, DiscoveredPaths);
-	MoveAppendRangeToRingBuffer(InOutResults.Dependencies, DependencyResults);
+	for (FPackageDependencyData& DependencyData : DependencyResults)
+	{
+		FName PackageName = DependencyData.PackageName;
+		InOutResults.Dependencies.Add(PackageName, MoveTemp(DependencyData));
+	}
+	DependencyResults.Reset();
 	MoveAppendRangeToRingBuffer(InOutResults.CookedPackageNamesWithoutAssetData, CookedPackageNamesWithoutAssetDataResults);
 	MoveAppendRangeToRingBuffer(InOutResults.VerseFiles, VerseResults);
 
@@ -3517,22 +3526,21 @@ void FAssetDataGatherer::GetAndTrimSearchResults(FResults& InOutResults, FResult
 	OutContext.bAbleToProgress = !bIsIdle;
 }
 
-void FAssetDataGatherer::GetPackageResults(TRingBuffer<FAssetData*>& OutAssetResults, TRingBuffer<FPackageDependencyData>& OutDependencyResults)
+void FAssetDataGatherer::GetPackageResults(TMultiMap<FName, FAssetData*>& OutAssetResults, TMultiMap<FName, FPackageDependencyData>& OutDependencyResults)
 {
 	FGathererScopeLock ResultsScopeLock(&ResultsLock);
 
-	auto MoveAppendRangeToRingBuffer = [](auto& OutRingBuffer, auto& InArray)
+	for (FAssetData* AssetData : AssetResults)
 	{
-		OutRingBuffer.Reserve(OutRingBuffer.Num() + InArray.Num());
-		for (auto& Element : InArray)
-		{
-			OutRingBuffer.Add(MoveTemp(Element));
-		}
-		InArray.Reset();
-	};
-
-	MoveAppendRangeToRingBuffer(OutAssetResults, AssetResults);
-	MoveAppendRangeToRingBuffer(OutDependencyResults, DependencyResults);
+		OutAssetResults.Add(AssetData->PackageName, AssetData);
+	}
+	AssetResults.Reset();
+	for (FPackageDependencyData& DependencyData : DependencyResults)
+	{
+		FName PackageName = DependencyData.PackageName;
+		OutDependencyResults.Add(PackageName, MoveTemp(DependencyData));
+	}
+	DependencyResults.Reset();
 }
 
 void FAssetDataGatherer::WaitOnPath(FStringView InPath)
diff --git a/Engine/Source/Runtime/AssetRegistry/Private/AssetDataGatherer.h b/Engine/Source/Runtime/AssetRegistry/Private/AssetDataGatherer.h
--- a/Engine/Source/Runtime/AssetRegistry/Private/AssetDataGatherer.h
+++ b/Engine/Source/Runtime/AssetRegistry/Private/AssetDataGatherer.h
@@ -87,14 +87,21 @@ public:
 	// Receiving Results (possibly while tick is running)
 	struct FResults
 	{
-		TRingBuffer<FAssetData*> Assets;
+		TMultiMap<FName, FAssetData*> Assets;
 		TRingBuffer<FString> Paths;
-		TRingBuffer<FPackageDependencyData> Dependencies;
+		TMultiMap<FName, FPackageDependencyData> Dependencies;
 		TRingBuffer<FString> CookedPackageNamesWithoutAssetData;
 		TRingBuffer<FName> VerseFiles;
 
 		SIZE_T GetAllocatedSize() const { return Assets.GetAllocatedSize() + Paths.GetAllocatedSize() + Dependencies.GetAllocatedSize() + CookedPackageNamesWithoutAssetData.GetAllocatedSize() + VerseFiles.GetAllocatedSize(); }
-		void Shrink() { Assets.Trim(); Paths.Trim(); Dependencies.Trim(); CookedPackageNamesWithoutAssetData.Trim(); VerseFiles.Trim(); }
+		void Shrink()
+		{
+			Assets.Shrink();
+			Paths.Trim();
+			Dependencies.Shrink();
+			CookedPackageNamesWithoutAssetData.Trim();
+			VerseFiles.Trim();
+		}
 	};
 	struct FResultContext
 	{
@@ -108,7 +115,7 @@ public:
 	/** Gets search results from the data gatherer. */
 	void GetAndTrimSearchResults(FResults& InOutResults, FResultContext& OutContext);
 	/** Gets just the AssetResults and DependencyResults from the data gatherer. */
-	void GetPackageResults(TRingBuffer<FAssetData*>& OutAssetResults, TRingBuffer<FPackageDependencyData>& OutDependencyResults);
+	void GetPackageResults(TMultiMap<FName, FAssetData*>& OutAssetResults, TMultiMap<FName, FPackageDependencyData>& OutDependencyResults);
 	/** Wait for all monitored assets under the given path to be added to search results. Returns immediately if the given path is not monitored. */
 	void WaitOnPath(FStringView LocalPath);
 	/**
diff --git a/Engine/Source/Runtime/AssetRegistry/Private/AssetRegistry.cpp b/Engine/Source/Runtime/AssetRegistry/Private/AssetRegistry.cpp
--- a/Engine/Source/Runtime/AssetRegistry/Private/AssetRegistry.cpp
+++ b/Engine/Source/Runtime/AssetRegistry/Private/AssetRegistry.cpp
@@ -3508,26 +3508,6 @@ void FAssetRegistryImpl::PrioritizeSearchPath(const FString& PathToPrioritize)
 		return;
 	}
 	GlobalGatherer->PrioritizeSearchPath(PathToPrioritize);
-
-	// Also prioritize the queue of background search results
-	int32 FirstNonPriorityIndex = 0;
-	for (int Index = 0; Index < BackgroundResults.Assets.Num(); ++Index)
-	{
-		FAssetData* PriorityElement = BackgroundResults.Assets[Index];
-		if (PriorityElement && PriorityElement->PackagePath.ToString().StartsWith(PathToPrioritize))
-		{
-			Swap(BackgroundResults.Assets[FirstNonPriorityIndex++], BackgroundResults.Assets[Index]);
-		}
-	}
-	FirstNonPriorityIndex = 0;
-	for (int Index = 0; Index < BackgroundResults.Paths.Num(); ++Index)
-	{
-		FString& PriorityElement = BackgroundResults.Paths[Index];
-		if (PriorityElement.StartsWith(PathToPrioritize))
-		{
-			Swap(BackgroundResults.Paths[FirstNonPriorityIndex++], BackgroundResults.Paths[Index]);
-		}
-	}
 }
 
 }
@@ -3974,35 +3954,31 @@ void FAssetRegistryImpl::TickGatherPackage(Impl::FEventContext& EventContext, co
 	// Gather results from the background search
 	GlobalGatherer->GetPackageResults(BackgroundResults.Assets, BackgroundResults.Dependencies);
 
-	TRingBuffer<FAssetData*> PackageAssets;
-	TRingBuffer<FPackageDependencyData> PackageDependencyDatas;
-	for (int n = 0; n < BackgroundResults.Assets.Num(); ++n)
+	TArray<FAssetData*> PackageAssets;
+	TArray<FPackageDependencyData> PackageDependencyDatas;
+	BackgroundResults.Assets.MultiFind(PackageFName, PackageAssets);
+	BackgroundResults.Assets.Remove(PackageFName);
+	BackgroundResults.Dependencies.MultiFind(PackageFName, PackageDependencyDatas);
+	BackgroundResults.Dependencies.Remove(PackageFName);
+	if (PackageAssets.Num() > 0)
 	{
-		FAssetData* Asset = BackgroundResults.Assets[n];
-		if (Asset->PackageName == PackageFName)
+		TMultiMap<FName, FAssetData*> PackageAssetsMap;
+		PackageAssetsMap.Reserve(PackageAssets.Num());
+		for (FAssetData* PackageAsset : PackageAssets)
 		{
-			Swap(BackgroundResults.Assets[n], BackgroundResults.Assets.Last());
-			PackageAssets.Add(BackgroundResults.Assets.PopValue());
-			--n;
+			PackageAssetsMap.Add(PackageFName, PackageAsset);
 		}
+		AssetSearchDataGathered(EventContext, -1., PackageAssetsMap);
 	}
-	for (int n = 0; n < BackgroundResults.Dependencies.Num(); ++n)
+	if (PackageDependencyDatas.Num() > 0)
 	{
-		FPackageDependencyData& DependencyData = BackgroundResults.Dependencies[n];
-		if (DependencyData.PackageName == PackageFName)
+		TMultiMap<FName, FPackageDependencyData> PackageDependencyDatasMap;
+		PackageDependencyDatasMap.Reserve(PackageDependencyDatas.Num());
+		for (FPackageDependencyData& DependencyData : PackageDependencyDatas)
 		{
-			Swap(BackgroundResults.Dependencies[n], BackgroundResults.Dependencies.Last());
-			PackageDependencyDatas.Add(BackgroundResults.Dependencies.PopValue());
-			--n;
+			PackageDependencyDatasMap.Add(PackageFName, MoveTemp(DependencyData));
 		}
-	}
-	if (PackageAssets.Num() > 0)
-	{
-		AssetSearchDataGathered(EventContext, -1., PackageAssets);
-	}
-	if (PackageDependencyDatas.Num() > 0)
-	{
-		DependencyDataGathered(-1., PackageDependencyDatas);
+		DependencyDataGathered(-1., PackageDependencyDatasMap);
 	}
 }
 
@@ -4520,7 +4496,7 @@ void FAssetRegistryImpl::ScanPathsSynchronous(Impl::FScanPathContext& Context)
 	Gatherer.ScanPathsSynchronous(Context.LocalPaths, Context.bForceRescan, Context.bIgnoreDenyListScanFilters, CacheFilename, Context.PackageDirs);
 	TArray<FName> FoundAssetPackageNames;
 
-	auto AssetsFoundCallback = [&Context, &FoundAssetPackageNames, this](const TRingBuffer<FAssetData*>& InFoundAssets)
+	auto AssetsFoundCallback = [&Context, &FoundAssetPackageNames, this](const TMultiMap<FName, FAssetData*>& InFoundAssets)
 	{
 		Context.NumFoundAssets = InFoundAssets.Num();
 
@@ -4528,8 +4504,9 @@ void FAssetRegistryImpl::ScanPathsSynchronous(Impl::FScanPathContext& Context)
 		FoundAssetPackageNames.Reserve(Context.NumFoundAssets);
 
 		// The gatherer may have added other assets that were scanned as part of the ongoing background scan; remove any assets that were not in the requested paths
-		for (FAssetData* AssetData : InFoundAssets)
+		for (const TPair<FName,FAssetData*>& Pair : InFoundAssets)
 		{
+			FAssetData* AssetData = Pair.Value;
 			bool bIsInRequestedPaths = false;
 
 			TStringBuilder<128> PackageNameStr;
@@ -4651,7 +4628,7 @@ void FAssetRegistryImpl::PostLoadAssetRegistryTags(FAssetData* AssetData)
 }
 #endif
 
-void FAssetRegistryImpl::AssetSearchDataGathered(Impl::FEventContext& EventContext, const double TickStartTime, TRingBuffer<FAssetData*>& AssetResults)
+void FAssetRegistryImpl::AssetSearchDataGathered(Impl::FEventContext& EventContext, const double TickStartTime, TMultiMap<FName, FAssetData*>& AssetResults)
 {
 	const bool bFlushFullBuffer = TickStartTime < 0;
 
@@ -4669,11 +4646,12 @@ void FAssetRegistryImpl::AssetSearchDataGathered(Impl::FEventContext& EventConte
 	}
 
 	// Add the found assets
-	while (AssetResults.Num() > 0)
+	for (TMultiMap<FName, FAssetData*>::TIterator Iter(AssetResults); Iter; ++Iter)
 	{
 		// Delete or take ownership of the BackgroundResult; it was originally new'd by an FPackageReader
-		TUniquePtr<FAssetData> BackgroundResult(AssetResults.PopFrontValue());
+		TUniquePtr<FAssetData> BackgroundResult(Iter.Value());
 		CA_ASSUME(BackgroundResult.Get() != nullptr);
+		Iter.RemoveCurrent();
 
 		// Try to update any asset data that may already exist
 		FCachedAssetKey Key(*BackgroundResult);
@@ -4776,15 +4754,16 @@ void FAssetRegistryImpl::PathDataGathered(Impl::FEventContext& EventContext, con
 	}
 }
 
-void FAssetRegistryImpl::DependencyDataGathered(const double TickStartTime, TRingBuffer<FPackageDependencyData>& DependsResults)
+void FAssetRegistryImpl::DependencyDataGathered(const double TickStartTime, TMultiMap<FName, FPackageDependencyData>& DependsResults)
 {
 	using namespace UE::AssetRegistry;
 	const bool bFlushFullBuffer = TickStartTime < 0;
 
 	TMap<FName, FName> CachedDepToRedirect;
-	while (DependsResults.Num() > 0)
+	for (TMultiMap<FName, FPackageDependencyData>::TIterator Iter(DependsResults); Iter; ++Iter)
 	{
-		FPackageDependencyData Result = DependsResults.PopFrontValue();
+		FPackageDependencyData Result = MoveTemp(Iter.Value());
+		Iter.RemoveCurrent();
 
 		checkf(!GIsEditor || Result.bHasPackageData, TEXT("We rely on PackageData being read for every gathered Asset in the editor."));
 		if (Result.bHasPackageData)
diff --git a/Engine/Source/Runtime/AssetRegistry/Private/AssetRegistryImpl.h b/Engine/Source/Runtime/AssetRegistry/Private/AssetRegistryImpl.h
--- a/Engine/Source/Runtime/AssetRegistry/Private/AssetRegistryImpl.h
+++ b/Engine/Source/Runtime/AssetRegistry/Private/AssetRegistryImpl.h
@@ -182,7 +182,7 @@ public:
 	/** Waits for the gatherer to be idle if it is operating synchronously. */
 	void WaitForGathererIdleIfSynchronous();
 	/** Callback type for TickGatherer */
-	typedef TFunctionRef<void(const TRingBuffer<FAssetData*>&)> FAssetsFoundCallback;
+	typedef TFunctionRef<void(const TMultiMap<FName, FAssetData*>&)> FAssetsFoundCallback;
 	/** Consume any results from the gatherer and return its status */
 	Impl::EGatherStatus TickGatherer(Impl::FEventContext& EventContext,
 		Impl::FClassInheritanceContext& InheritanceContext, const double TickStartTime, bool& bOutInterrupted,
@@ -281,7 +281,7 @@ private:
 	 * Called every tick to when data is retrieved by the background asset search.
 	 * If TickStartTime is < 0, the entire list of gathered assets will be cached. Also used in sychronous searches
 	 */
-	void AssetSearchDataGathered(Impl::FEventContext& EventContext, const double TickStartTime, TRingBuffer<FAssetData*>& AssetResults);
+	void AssetSearchDataGathered(Impl::FEventContext& EventContext, const double TickStartTime, TMultiMap<FName, FAssetData*>& AssetResults);
 
 	/**
 	 * Called every tick when data is retrieved by the background path search.
@@ -290,7 +290,7 @@ private:
 	void PathDataGathered(Impl::FEventContext& EventContext, const double TickStartTime, TRingBuffer<FString>& PathResults);
 
 	/** Called every tick when data is retrieved by the background dependency search */
-	void DependencyDataGathered(const double TickStartTime, TRingBuffer<FPackageDependencyData>& DependsResults);
+	void DependencyDataGathered(const double TickStartTime, TMultiMap<FName, FPackageDependencyData>& DependsResults);
 
 	/** Called every tick when data is retrieved by the background search for cooked packages that do not contain asset data */
 	void CookedPackageNamesWithoutAssetDataGathered(Impl::FEventContext& EventContext, const double TickStartTime,
