AssetDataGatherer: When blocking on a file result due to WaitOnPath, scan only the bottom-most directory instead of additionally scanning the parent directories. This reduces blocking on the gamethread during editor and cooker startup.
#rnx

diff --git a/Engine/Source/Runtime/AssetRegistry/Private/AssetDataGatherer.cpp b/Engine/Source/Runtime/AssetRegistry/Private/AssetDataGatherer.cpp
--- a/Engine/Source/Runtime/AssetRegistry/Private/AssetDataGatherer.cpp
+++ b/Engine/Source/Runtime/AssetRegistry/Private/AssetDataGatherer.cpp
@@ -732,24 +732,51 @@ void FScanDir::Update(FScanDir*& OutCursor, FInherited& InOutParentData)
 	{
 		return;
 	}
+
+	FScanDir* ThisCursor = nullptr;
 	if (ShouldScan(InOutParentData))
 	{
-		OutCursor = this;
-		return;
+		ThisCursor = this;
 	}
 
 	if (SubDirs.Num())
 	{
 		FScanDir* SubDirToScan = FindHighestPrioritySubDir();
-		if (SubDirToScan)
+		// If AccumulatedPriority of this and the subdir are the same, the tie goes to the subdir, since 
+		// the accumulatedpriority of a parentdir is clamped to be at least as strong (at least as small) as the subdir's accumulatedpriority
+		if (SubDirToScan && (!ThisCursor || SubDirToScan->AccumulatedPriority <= ThisCursor->AccumulatedPriority))
 		{
-			OutCursor = SubDirToScan;
-			FInherited Accumulated(InOutParentData, DirectData);
-			InOutParentData = Accumulated;
-			return;
+			FInherited ThisCursorParentData(InOutParentData, DirectData);
+			for (;;) // Continue Updating SubDirs until we find one that needs scanning or all are complete or lower in priority than this
+			{
+				FInherited SubDirsParentData(ThisCursorParentData);
+				SubDirToScan->Update(OutCursor, SubDirsParentData);
+				if (OutCursor != this)
+				{
+					InOutParentData = SubDirsParentData;
+					return;
+				}
+				// Otherwise the SubDir finished updating and we should move to the next one
+				check(SubDirToScan->IsComplete());
+				SubDirToScan = FindHighestPrioritySubDir();
+				if (SubDirToScan && (!ThisCursor || SubDirToScan->AccumulatedPriority <= ThisCursor->AccumulatedPriority))
+				{
+					continue;
+				}
+				else
+				{
+					break;
+				}
+			}
 		}
 	}
 
+	if (ThisCursor)
+	{
+		OutCursor = ThisCursor;
+		return;
+	}
+
 	OutCursor = Parent; // Note this will be null for the root ScanDir
 	if (!Parent)
 	{
@@ -880,8 +907,9 @@ void FScanDir::SetComplete(bool bInIsComplete)
 			}
 		}
 		// Upon completion, subdirs that do not need to be maintained are deleted, which is done by removing them from the parent
-		// ScanDirs need to be maintained if they are the root, or have persistent settings, or have child ScanDirs that need to be maintained.
-		if (Parent != nullptr && !HasPersistentSettings() && SubDirs.IsEmpty())
+		// ScanDirs need to be maintained if they are the root, or have persistent settings, or have child ScanDirs that need to be maintained,
+		// or the parent scan has not been done yet.
+		if ((Parent != nullptr && Parent->bHasScanned) && !HasPersistentSettings() && SubDirs.IsEmpty())
 		{
 			Parent->RemoveSubDir(GetRelPath());
 			// *this is Shutdown (e.g. Parent is now null) and it may also have been deallocated
diff --git a/Engine/Source/Runtime/AssetRegistry/Private/AssetDataGathererPrivate.h b/Engine/Source/Runtime/AssetRegistry/Private/AssetDataGathererPrivate.h
--- a/Engine/Source/Runtime/AssetRegistry/Private/AssetDataGathererPrivate.h
+++ b/Engine/Source/Runtime/AssetRegistry/Private/AssetDataGathererPrivate.h
@@ -29,7 +29,10 @@ enum class EPriority : uint8
 	/** Optional information (e.g. use of the ReferenceViewer) is unavailable until the file/directory is completed */
 	High,
 	/** Nothing has requested the file/directory yet */
-	Normal
+	Normal,
+
+	Highest = Blocking,
+	Lowest = Normal,
 };
 constexpr uint32 CountEPriority = static_cast<uint32>(EPriority::Normal) + 1;
 
