AssetRegistry: Fix DumpAssetRegistry to sort the ObjectNames; sorting was broken when they changed from FName to FCachedAssetKey. Add pagebreaks at each section.
#rnx

diff --git a/Engine/Source/Runtime/AssetRegistry/Private/AssetRegistryState.cpp b/Engine/Source/Runtime/AssetRegistry/Private/AssetRegistryState.cpp
--- a/Engine/Source/Runtime/AssetRegistry/Private/AssetRegistryState.cpp
+++ b/Engine/Source/Runtime/AssetRegistry/Private/AssetRegistryState.cpp
@@ -2295,7 +2295,7 @@ void FAssetRegistryState::Dump(const TArray<FString>& Arguments, TArray<FString>
 	int32 NumLinesInPage = 0;
 	const int32 LineTerminatorLen = TCString<TCHAR>::Strlen(LINE_TERMINATOR);
 
-	auto FinishPage = [&PageBuffer, &NumLinesInPage, HashStartValue, HashMultiplier, PageEndSearchLength, &OutPages, &OverflowText, LineTerminatorLen](bool bLastPage)
+	auto FinishPage = [&PageBuffer, &NumLinesInPage, HashStartValue, HashMultiplier, PageEndSearchLength, &OutPages, &OverflowText, LineTerminatorLen](bool bManualPageBreak)
 	{
 		int32 PageEndIndex = PageBuffer.Len();
 		const TCHAR* BufferEnd = PageBuffer.GetData() + PageEndIndex;
@@ -2304,7 +2304,7 @@ void FAssetRegistryState::Dump(const TArray<FString>& Arguments, TArray<FString>
 		// because after one missing or added line, every page from that point on will be offset and therefore different, making false positive differences
 		// To make pages after one missing or added line the same, we look for a good page ending based on the text of all the lines near the end of the current page
 		// By choosing specific-valued texts as page breaks, we will usually randomly get lucky and have the two diffs pick the same line for the end of the page
-		if (!bLastPage && NumLinesInPage > PageEndSearchLength)
+		if (!bManualPageBreak && NumLinesInPage > PageEndSearchLength)
 		{
 			const TCHAR* WinningLineEnd = BufferEnd;
 			uint32 WinningLineValue = 0;
@@ -2371,9 +2371,17 @@ void FAssetRegistryState::Dump(const TArray<FString>& Arguments, TArray<FString>
 			}
 		}
 	};
+	auto AddPageBreak = [LinesPerPage, &NumLinesInPage, &FinishPage]()
+	{
+		if (LinesPerPage > 1 && NumLinesInPage != 0)
+		{
+			FinishPage(true);
+		}
+	};
 
 	if (bAllFields || Arguments.Contains(TEXT("ObjectPath")))
 	{
+		AddPageBreak();
 		PageBuffer.Append(TEXT("--- Begin CachedAssetsByObjectPath ---"));
 		AddLine();
 
@@ -2384,7 +2392,7 @@ void FAssetRegistryState::Dump(const TArray<FString>& Arguments, TArray<FString>
 			Keys.Emplace(FCachedAssetKey(*AssetData));
 		}
 		Keys.Sort([](const FCachedAssetKey& A, const FCachedAssetKey& B) {
-			return A.Compare(B) < 0;
+			return WriteToString<1024>(A).ToView().Compare(WriteToString<1024>(B).ToView()) < 0;
 		});
 
 		for (const FCachedAssetKey& Key : Keys)
@@ -2400,22 +2408,26 @@ void FAssetRegistryState::Dump(const TArray<FString>& Arguments, TArray<FString>
 
 	if (bAllFields || Arguments.Contains(TEXT("PackageName")))
 	{
+		AddPageBreak();
 		PrintAssetDataMap(TEXT("CachedAssetsByPackageName"), CachedAssetsByPackageName, PageBuffer, AddLine);
 	}
 
 	if (bAllFields || Arguments.Contains(TEXT("Path")))
 	{
+		AddPageBreak();
 		PrintAssetDataMap(TEXT("CachedAssetsByPath"), CachedAssetsByPath, PageBuffer, AddLine);
 	}
 
 	if (bAllFields || Arguments.Contains(TEXT("Class")))
 	{
+		AddPageBreak();
 		PrintAssetDataMap(TEXT("CachedAssetsByClass"), CachedAssetsByClass, PageBuffer, AddLine);
 	}
 
 	// Only print this if it's requested specifically - '-all' will print tags-per-asset rather than assets-per-tag 
 	if (Arguments.Contains(TEXT("Tag")))
 	{
+		AddPageBreak();
 		PrintAssetDataMap(TEXT("CachedAssetsByTag"), CachedAssetsByTag, PageBuffer, AddLine,
 			[&PageBuffer, &AddLine](const FName& TagName, const FAssetData& Data)
 			{
@@ -2426,6 +2438,7 @@ void FAssetRegistryState::Dump(const TArray<FString>& Arguments, TArray<FString>
 	if (bAllFields || Arguments.Contains(TEXT("AssetTags")))
 	{
 		int32 Counter = 0;
+		AddPageBreak();
 		PageBuffer.Append(TEXT("--- Begin AssetTags ---"));
 		AddLine();
 
@@ -2454,6 +2467,7 @@ void FAssetRegistryState::Dump(const TArray<FString>& Arguments, TArray<FString>
 
 	if ((bAllFields || Arguments.Contains(TEXT("Dependencies"))) && !bDumpDependencyDetails)
 	{
+		AddPageBreak();
 		PageBuffer.Appendf(TEXT("--- Begin CachedDependsNodes ---"));
 		AddLine();
 
@@ -2478,6 +2492,7 @@ void FAssetRegistryState::Dump(const TArray<FString>& Arguments, TArray<FString>
 	if (bDumpDependencyDetails)
 	{
 		using namespace UE::AssetRegistry;
+		AddPageBreak();
 		PageBuffer.Append(TEXT("--- Begin CachedDependsNodes ---"));
 		AddLine();
 
@@ -2602,6 +2617,7 @@ void FAssetRegistryState::Dump(const TArray<FString>& Arguments, TArray<FString>
 	}
 	if (bAllFields || Arguments.Contains(TEXT("PackageData")))
 	{
+		AddPageBreak();
 		PageBuffer.Append(TEXT("--- Begin CachedPackageData ---"));
 		AddLine();
 
@@ -2629,6 +2645,7 @@ void FAssetRegistryState::Dump(const TArray<FString>& Arguments, TArray<FString>
 	if (bAllFields || Arguments.Contains(TEXT("AssetBundles")))
 	{
 		int32 Counter = 0;
+		AddPageBreak();
 		PageBuffer.Append(TEXT("--- Begin AssetBundles ---"));
 		AddLine();
 
