AssetDataGatherer:WaitOnPath fix: reset priority to normal when we are done with the blocking scan of a directory.
#rnx

diff --git a/Engine/Source/Runtime/AssetRegistry/Private/AssetDataGatherer.cpp b/Engine/Source/Runtime/AssetRegistry/Private/AssetDataGatherer.cpp
--- a/Engine/Source/Runtime/AssetRegistry/Private/AssetDataGatherer.cpp
+++ b/Engine/Source/Runtime/AssetRegistry/Private/AssetDataGatherer.cpp
@@ -610,6 +610,11 @@ void FScanDir::MarkFileAlreadyScanned(FStringView BaseName)
 	AlreadyScannedFiles.Emplace(BaseName);
 }
 
+EPriority FScanDir::GetDirectPriority() const
+{
+	return DirectPriority;
+}
+
 void FScanDir::SetDirectPriority(EPriority InPriority)
 {
 	DirectPriority = InPriority;
@@ -1898,6 +1903,7 @@ void FAssetDataDiscovery::SetPropertiesAndWait(FPathExistence& QueryPath, bool b
 			{
 				// We are going to wait on the path, so set its priority to blocking
 				SetIsIdle(false);
+				EPriority OriginalPriority = ScanDir->GetDirectPriority();
 				ScanDir->SetDirectPriority(EPriority::Blocking);
 				InvalidateCursor();
 
@@ -1919,10 +1925,16 @@ void FAssetDataDiscovery::SetPropertiesAndWait(FPathExistence& QueryPath, bool b
 					}
 					if (ScanDir->IsComplete() || (!bScanEntireTree && ScanDir->HasScanned()))
 					{
+						// Restore the Priority.
+						// TODO: We dropped the lock so something else could have changed the value of ScanDir's priority and we are now clobbering it.
+						// This can be fixed by adding a counter for each level of the priority. That's complex and cpu expensive, so not doing it until we need it.
+						ScanDir->SetDirectPriority(OriginalPriority);
 						break;
 					}
 					else if (!ensureMsgf(!bIsIdle, TEXT("It should not be possible for the Discovery to go idle while there is an incomplete ScanDir.")))
 					{
+						// Restore the Priority
+						ScanDir->SetDirectPriority(OriginalPriority);
 						break;
 					}
 				}
diff --git a/Engine/Source/Runtime/AssetRegistry/Private/AssetDataGathererPrivate.h b/Engine/Source/Runtime/AssetRegistry/Private/AssetDataGathererPrivate.h
--- a/Engine/Source/Runtime/AssetRegistry/Private/AssetDataGathererPrivate.h
+++ b/Engine/Source/Runtime/AssetRegistry/Private/AssetDataGathererPrivate.h
@@ -361,6 +361,7 @@ public:
 	void MarkFileAlreadyScanned(FStringView BaseName);
 	/** Set the direct priority of this ScanDir and update the accumulated priority to match it. */
 	void SetDirectPriority(EPriority InPriority);
+	EPriority GetDirectPriority() const;
 
 	/**
 	 * Called from the Tick; handle the list of subdirs and files that were found from IterateDirectoryStat called
