AssetRegistry - Separate permissions for reading and writing from Gatherer cache so we can read but not write when in multiprocess.
#rnx

diff --git a/Engine/Source/Runtime/AssetRegistry/Private/AssetDataGatherer.cpp b/Engine/Source/Runtime/AssetRegistry/Private/AssetDataGatherer.cpp
--- a/Engine/Source/Runtime/AssetRegistry/Private/AssetDataGatherer.cpp
+++ b/Engine/Source/Runtime/AssetRegistry/Private/AssetDataGatherer.cpp
@@ -2804,28 +2804,38 @@ struct FPreloadSettings
 		FString ProjectIntermediateDir = FPaths::ProjectIntermediateDir();
 		bForceDependsGathering = FParse::Param(FCommandLine::Get(), TEXT("ForceDependsGathering"));
 		bGatherDependsData = (GIsEditor && !FParse::Param(FCommandLine::Get(), TEXT("NoDependsGathering"))) || bForceDependsGathering;
-		bCacheEnabled = !FParse::Param(FCommandLine::Get(), TEXT("NoAssetRegistryCache")) && !FParse::Param(FCommandLine::Get(), TEXT("multiprocess"));
+		bool bNoAssetRegistryCache = FParse::Param(FCommandLine::Get(), TEXT("NoAssetRegistryCache"));
+		bool bNoAssetRegistryCacheRead = FParse::Param(FCommandLine::Get(), TEXT("NoAssetRegistryCacheRead"));
+		bool bNoAssetRegistryCacheWrite = FParse::Param(FCommandLine::Get(), TEXT("NoAssetRegistryCacheWrite"));
+		uint32 MultiprocessId = 0;
+		FParse::Value(FCommandLine::Get(), TEXT("multiprocessid="), MultiprocessId);
+		bool bMultiprocess = MultiprocessId > 0 || FParse::Param(FCommandLine::Get(), TEXT("multiprocess"));
+		bCacheReadEnabled = !bNoAssetRegistryCache && !bNoAssetRegistryCacheRead;
+		bCacheWriteEnabled = !bNoAssetRegistryCache && !bNoAssetRegistryCacheWrite && !bMultiprocess;
 		bool bAsyncEnabled = FPlatformProcess::SupportsMultithreading() && FTaskGraphInterface::IsRunning();
 
 		MonolithicCacheFilename = FPaths::ProjectIntermediateDir() / (bGatherDependsData ? TEXT("CachedAssetRegistry.bin") : TEXT("CachedAssetRegistryNoDeps.bin"));
 #if UE_EDITOR // See note on FPreloader for why we only allow preloading if UE_EDITOR
-		if (bCacheEnabled && bAsyncEnabled)
-		{
-			bPreloadMonolithicCache = GIsEditor && (!IsRunningCommandlet() || IsRunningCookCommandlet());
-		}
-		else
+		bMonolithicCacheActivatedDuringPreload = bAsyncEnabled && GIsEditor && (!IsRunningCommandlet() || IsRunningCookCommandlet());
+#else
+		bMonolithicCacheActivatedDuringPreload = false;
 #endif
-		{
-			bPreloadMonolithicCache = false;
-		}
 	}
-	bool IsCacheEnabled() const
+	bool IsCacheReadEnabled() const
+	{
+		return bCacheReadEnabled;
+	}
+	bool IsCacheWriteEnabled() const
+	{
+		return bCacheWriteEnabled;
+	}
+	bool IsMonolithicCacheActivatedDuringPreload() const
 	{
-		return bCacheEnabled;
+		return bMonolithicCacheActivatedDuringPreload;
 	}
 	bool IsPreloadMonolithicCache() const
 	{
-		return bPreloadMonolithicCache;
+		return bCacheReadEnabled && bMonolithicCacheActivatedDuringPreload;
 	}
 	bool IsGatherDependsData() const
 	{
@@ -2843,8 +2853,9 @@ private:
 	FString MonolithicCacheFilename;
 	bool bForceDependsGathering = false;
 	bool bGatherDependsData = false;
-	bool bCacheEnabled = false;
-	bool bPreloadMonolithicCache = false;
+	bool bCacheReadEnabled = false;
+	bool bCacheWriteEnabled = false;
+	bool bMonolithicCacheActivatedDuringPreload = false;
 	bool bInitialized = false;
 };
 FPreloadSettings GPreloadSettings;
@@ -2944,9 +2955,12 @@ FAssetDataGatherer::FAssetDataGatherer(const TArray<FString>& InLongPackageNames
 	GPreloadSettings.Initialize();
 	bGatherAssetPackageData = GIsEditor || GPreloadSettings.IsForceDependsGathering();
 	bGatherDependsData = GPreloadSettings.IsGatherDependsData();
-	bCacheEnabled = GPreloadSettings.IsCacheEnabled();
-	// If PreloadMonolithicCache is true, we are using MonolithicCache. Otherwise it may be set to true later if game/commandlet calls SearchAllAssets.
-	bUseMonolithicCache = GPreloadSettings.IsPreloadMonolithicCache();
+	bCacheReadEnabled = GPreloadSettings.IsCacheReadEnabled();
+	bCacheWriteEnabled = GPreloadSettings.IsCacheWriteEnabled();
+	// If IsMonolithicCacheActivatedDuringPreload is true, we are already instructed to use the MonolithicCache.
+	// Otherwise it may be set to true later if game/commandlet calls SearchAllAssets.
+	bReadMonolithicCache = bCacheReadEnabled && GPreloadSettings.IsMonolithicCacheActivatedDuringPreload();
+	bWriteMonolithicCache = bCacheWriteEnabled && GPreloadSettings.IsMonolithicCacheActivatedDuringPreload();
 	LastCacheWriteTime = FPlatformTime::Seconds();
 
 #if !UE_BUILD_SHIPPING
@@ -2988,12 +3002,16 @@ FAssetDataGatherer::~FAssetDataGatherer()
 void FAssetDataGatherer::ActivateMonolithicCache()
 {
 	FGathererScopeLock ResultsScopeLock(&ResultsLock);
-	if (!bCacheEnabled || bUseMonolithicCache )
+	bool bNewWriteMonolithicCache = bCacheWriteEnabled;
+	bool bNewReadMonolithicCache = bCacheReadEnabled;
+	if ((!bNewWriteMonolithicCache || bWriteMonolithicCache) &&
+		(!bNewReadMonolithicCache || bReadMonolithicCache))
 	{
 		return;
 	}
 
-	bUseMonolithicCache = true;
+	bWriteMonolithicCache = bNewWriteMonolithicCache;
+	bReadMonolithicCache = bNewReadMonolithicCache;
 	LastCacheWriteTime = FPlatformTime::Seconds();
 }
 
@@ -3139,7 +3157,7 @@ void FAssetDataGatherer::TickInternal(bool& bOutIsTickInterrupt)
 	TArray<FString, FBatchInlineAllocator> LocalCookedPackageNamesWithoutAssetDataResults;
 	TArray<FName, FBatchInlineAllocator> LocalVerseResults;
 	bool bLoadMonolithicCache = false;
-	bool bLocalIsCacheEnabled = false;
+	bool bLocalIsCacheWriteEnabled = false;
 	double LocalLastCacheWriteTime = 0.0;
 	bool bWaitBatchCountDecremented = false;
 	bOutIsTickInterrupt = false;
@@ -3184,13 +3202,13 @@ void FAssetDataGatherer::TickInternal(bool& bOutIsTickInterrupt)
 			SetIsIdle(true);
 			return;
 		}
-		if (bUseMonolithicCache && !bHasLoadedMonolithicCache)
+		if (bReadMonolithicCache && !bHasLoadedMonolithicCache)
 		{
 			bLoadMonolithicCache = true;
 		}
 		LocalLastCacheWriteTime = LastCacheWriteTime;
 	}
-	bLocalIsCacheEnabled = bCacheEnabled;
+	bLocalIsCacheWriteEnabled = bCacheWriteEnabled;
 
 	// Load the async cache if not yet loaded
 	if (bLoadMonolithicCache)
@@ -3331,7 +3349,7 @@ void FAssetDataGatherer::TickInternal(bool& bOutIsTickInterrupt)
 			// Add the results from a cooked package into our results on cooked package
 			LocalCookedPackageNamesWithoutAssetDataResults.Append(MoveTemp(ReadContext.CookedPackageNamesWithoutAssetData));
 			// Do not add the results from a cooked package into the map of data we keep to write out the new version of the cache file
-			bool bCachePackage = bLocalIsCacheEnabled
+			bool bCachePackage = bLocalIsCacheWriteEnabled
 				&& LocalCookedPackageNamesWithoutAssetDataResults.Num() == 0
 				&& ensure(ReadContext.AssetFileData.Type == EGatherableFileType::PackageFile);
 			if (bCachePackage)
@@ -3402,7 +3420,7 @@ void FAssetDataGatherer::TickInternal(bool& bOutIsTickInterrupt)
 			}
 		}
 
-		if (bUseMonolithicCache && !bIsSavingAsyncCache 
+		if (bWriteMonolithicCache && !bIsSavingAsyncCache 
 			&& FPlatformTime::Seconds() - LocalLastCacheWriteTime >= AssetDataGathererConstants::MinSecondsToElapseBeforeCacheWrite)
 		{
 			bSaveAsyncCacheTriggered = true;
@@ -3604,11 +3622,13 @@ void FAssetDataGatherer::ClearCache()
 	bool bWasCacheEnabled = false;
 	{
 		FGathererScopeLock TickScopeLock(&TickLock);
-		bWasCacheEnabled = bCacheEnabled;
-		bCacheEnabled = false;
+		bWasCacheEnabled = bCacheWriteEnabled || bCacheReadEnabled;
+		bCacheWriteEnabled = false;
+		bCacheReadEnabled = false;
 		{
 			FGathererScopeLock ResultsScopeLock(&ResultsLock);
-			bUseMonolithicCache = false;
+			bReadMonolithicCache = false;
+			bWriteMonolithicCache = false;
 		}
 		bCacheIsInUseOnOtherThread = CacheInUseCount > 0;
 	}
@@ -3710,7 +3730,7 @@ void FAssetDataGatherer::WaitOnPathsInternal(TArrayView<UE::AssetDataGather::Pri
 		}
 	}
 
-	if (!SaveCacheFilename.IsEmpty() && bCacheEnabled)
+	if (!SaveCacheFilename.IsEmpty() && bCacheWriteEnabled)
 	{
 		TArray<TPair<FName, FDiskCachedAssetData*>> AssetsToSave;
 		{
@@ -3774,9 +3794,14 @@ bool FAssetDataGatherer::IsGatheringDependencies() const
 	return bGatherDependsData;
 }
 
-bool FAssetDataGatherer::IsCacheEnabled() const
+bool FAssetDataGatherer::IsCacheReadEnabled() const
+{
+	return bCacheReadEnabled;
+}
+
+bool FAssetDataGatherer::IsCacheWriteEnabled() const
 {
-	return bCacheEnabled;
+	return bCacheWriteEnabled;
 }
 
 FString FAssetDataGatherer::GetCacheFilename(TConstArrayView<FString> CacheFilePackagePaths)
@@ -3808,7 +3833,7 @@ FString FAssetDataGatherer::GetCacheFilename(TConstArrayView<FString> CacheFileP
 void FAssetDataGatherer::LoadCacheFile(FStringView CacheFilename)
 {
 	using namespace UE::AssetDataGather::Private;
-	if (!bCacheEnabled)
+	if (!bCacheReadEnabled)
 	{
 		return;
 	}
@@ -3857,7 +3882,7 @@ void FAssetDataGatherer::TryReserveSaveMonolithicCache(bool& bOutShouldSave, TAr
 	CHECK_IS_LOCKED_CURRENT_THREAD(TickLock);
 	{
 		FGathererScopeLock ResultsScopeLock(&ResultsLock);
-		bOutShouldSave = bUseMonolithicCache;
+		bOutShouldSave = bWriteMonolithicCache;
 	}
 	if (bOutShouldSave)
 	{
@@ -3899,7 +3924,7 @@ void FAssetDataGatherer::GetAssetsToSave(TArrayView<const FString> SaveCacheLong
 
 void FAssetDataGatherer::SaveCacheFileInternal(const FString& CacheFilename, const TArray<TPair<FName,FDiskCachedAssetData*>>& AssetsToSave, bool bIsAsyncCacheSave)
 {
-	if (CacheFilename.IsEmpty() || !bCacheEnabled)
+	if (CacheFilename.IsEmpty() || !bCacheWriteEnabled)
 	{
 		return;
 	}
diff --git a/Engine/Source/Runtime/AssetRegistry/Private/AssetDataGatherer.h b/Engine/Source/Runtime/AssetRegistry/Private/AssetDataGatherer.h
--- a/Engine/Source/Runtime/AssetRegistry/Private/AssetDataGatherer.h
+++ b/Engine/Source/Runtime/AssetRegistry/Private/AssetDataGatherer.h
@@ -148,8 +148,10 @@ public:
 	void SetInitialPluginsLoaded();
 	/** Report whether the gatherer is configured to load depends data in addition to asset data. */
 	bool IsGatheringDependencies() const;
-	/** Return whether the current process enables reading and writing AssetDataGatherer cache files. */
-	bool IsCacheEnabled() const;
+	/** Return whether the current process enables reading AssetDataGatherer cache files. */
+	bool IsCacheReadEnabled() const;
+	/** Return whether the current process enables writing AssetDataGatherer cache files. */
+	bool IsCacheWriteEnabled() const;
 	/** Calculate the cache filename that should be used for the given list of package paths. */
 	FString GetCacheFilename(TConstArrayView<FString> CacheFilePackagePaths);
 	/** Attempt to read the cache file at the given LocalPath, and store all of its results in the in-memory cache. */
@@ -351,8 +353,10 @@ private:
 	 * Read/writable anywhere.
 	 */
 	std::atomic<bool> bSaveAsyncCacheTriggered;
-	/** True if the current process allows reading/writing of AssetDataGatherer cache files. */
-	std::atomic<bool> bCacheEnabled;
+	/** True if the current process allows reading of AssetDataGatherer cache files. */
+	std::atomic<bool> bCacheReadEnabled;
+	/** True if the current process allows writing of AssetDataGatherer cache files. */
+	std::atomic<bool> bCacheWriteEnabled;
 
 	// Variable section for variables that are read/writable only within ResultsLock.
 
@@ -389,11 +393,16 @@ private:
 	 */
 	int32 NumPathsToSearchAtLastSyncPoint;
 	/**
-	 * Track whether we are using a monolithic cache that should be loaded/saved during tick.
-	 * Wether we are or not, the AssetRegistry can also call LoadCacheFile/ScanPathsSynchronous to load/save smaller files.
+	 * Track whether we are allowed to read from a monolithic cache that should be loaded during tick.
+	 * Even if we are or not, if bCacheReadEnabled the AssetRegistry can also call LoadCacheFile/ScanPathsSynchronous to load/save smaller files.
 	 * Read/writable only within ResultsLock.
 	 */
-	bool bUseMonolithicCache;
+	bool bReadMonolithicCache;
+	/**
+	 * Track whether we are allowed to write to the monolithic cache.
+	 * Read/writable only within ResultsLock.
+	 */
+	bool bWriteMonolithicCache;
 	/** If bHasLoadedMonolithicCache is true, track whether the cache has been loaded. Read/writable only within ResultsLock. */
 	bool bHasLoadedMonolithicCache;
 	/** Track whether the Discovery subsystem has gone idle and we have read all filenames from it. Read/writable only within ResultsLock. */
diff --git a/Engine/Source/Runtime/AssetRegistry/Private/AssetRegistry.cpp b/Engine/Source/Runtime/AssetRegistry/Private/AssetRegistry.cpp
--- a/Engine/Source/Runtime/AssetRegistry/Private/AssetRegistry.cpp
+++ b/Engine/Source/Runtime/AssetRegistry/Private/AssetRegistry.cpp
@@ -4467,7 +4467,7 @@ void FAssetRegistryImpl::ScanPathsSynchronous(Impl::FScanPathContext& Context)
 
 	// Add a cache file for any not-yet-scanned dirs
 	TArray<FString> CacheFilePackagePaths;
-	if (!Context.bForceRescan && Gatherer.IsCacheEnabled())
+	if (!Context.bForceRescan && (Gatherer.IsCacheReadEnabled() || Gatherer.IsCacheWriteEnabled()))
 	{
 		for (int n = 0; n < Context.LocalDirs.Num(); ++n)
 		{
