AssetRegistry: Add Extension to FAssetPackageData. This makes the AssetRegistry's calculation of package extension fully robust and supports future changes in the cooker and linker that will remove searching for the package filename on disk in favor of using the AssetRegistry.
#rnx

diff --git a/Engine/Source/Runtime/AssetRegistry/Private/AssetRegistry.cpp b/Engine/Source/Runtime/AssetRegistry/Private/AssetRegistry.cpp
--- a/Engine/Source/Runtime/AssetRegistry/Private/AssetRegistry.cpp
+++ b/Engine/Source/Runtime/AssetRegistry/Private/AssetRegistry.cpp
@@ -2613,11 +2613,28 @@ bool UAssetRegistryImpl::DoesPackageExistOnDisk(FName PackageName, FString* OutC
 		}
 		if (OutExtension)
 		{
-			// TODO: Save the extension on the AssetPackageData rather than deriving it here
-			// Guessing the extension based on map vs non-map also does not support text assets and maps which have a different extension
-			TArray<FAssetData> Assets;
-			GetAssetsByPackageName(PackageName, Assets, /*bIncludeOnlyDiskAssets*/ true);
-			*OutExtension = CalculateExtension(PackageNameStr, Assets);
+			if (AssetPackageData->Extension == EPackageExtension::Unspecified || AssetPackageData->Extension == EPackageExtension::Custom)
+			{
+				FString FileName;
+				if (FPackageName::DoesPackageExist(PackageNameStr, &FileName, false /* InAllowTextFormats */))
+				{
+					check(!FileName.IsEmpty());
+					*OutExtension = FPaths::GetExtension(FileName, true /* bIncludeDot */);
+				}
+				else
+				{
+					UE_LOG(LogAssetRegistry, Error,
+						TEXT("UAssetRegistryImpl::DoesPackageExistOnDisk failed to find the extension for %s. The package exists in the AssetRegistry but does not exist on disk."),
+						*PackageNameStr);
+					TArray<FAssetData> Assets;
+					GetAssetsByPackageName(PackageName, Assets, /*bIncludeOnlyDiskAssets*/ true);
+					*OutExtension = CalculateExtension(PackageNameStr, Assets);
+				}
+			}
+			else
+			{
+				*OutExtension = LexToString(AssetPackageData->Extension);
+			}
 		}
 		return true;
 	}
diff --git a/Engine/Source/Runtime/AssetRegistry/Private/PackageReader.cpp b/Engine/Source/Runtime/AssetRegistry/Private/PackageReader.cpp
--- a/Engine/Source/Runtime/AssetRegistry/Private/PackageReader.cpp
+++ b/Engine/Source/Runtime/AssetRegistry/Private/PackageReader.cpp
@@ -10,6 +10,7 @@
 #include "Logging/MessageLog.h"
 #include "Misc/Guid.h"
 #include "Misc/PackageName.h"
+#include "Misc/PackagePath.h"
 #include "UObject/Class.h"
 #include "UObject/Linker.h"
 #include "UObject/PackageTrailer.h"
@@ -485,6 +486,7 @@ bool FPackageReader::ReadDependencyData(FPackageDependencyData& OutDependencyDat
 		PackageData.FileVersionUE = PackageFileSummary.GetFileVersionUE();
 		PackageData.FileVersionLicenseeUE = PackageFileSummary.GetFileVersionLicenseeUE();
 		PackageData.SetIsLicenseeVersion(PackageFileSummary.SavedByEngineVersion.IsLicenseeVersion());
+		PackageData.Extension = FPackagePath::ParseExtension(PackageFilename);
 
 		if (!SerializeImportedClasses(ImportMap, PackageData.ImportedClasses))
 		{
diff --git a/Engine/Source/Runtime/CoreUObject/Private/AssetRegistry/AssetData.cpp b/Engine/Source/Runtime/CoreUObject/Private/AssetRegistry/AssetData.cpp
--- a/Engine/Source/Runtime/CoreUObject/Private/AssetRegistry/AssetData.cpp
+++ b/Engine/Source/Runtime/CoreUObject/Private/AssetRegistry/AssetData.cpp
@@ -851,6 +851,24 @@ void FAssetPackageData::SerializeForCacheInternal(FArchive& Ar, FAssetPackageDat
 		}
 		Ar << PackageData.ImportedClasses;
 	}
+	if (Version >= FAssetRegistryVersion::AssetPackageDataHasExtension)
+	{
+		FString ExtensionText;
+		if (Ar.IsLoading())
+		{
+			Ar << ExtensionText;
+			PackageData.Extension = FPackagePath::ParseExtension(ExtensionText);
+		}
+		else
+		{
+			ExtensionText = LexToString(PackageData.Extension);
+			Ar << ExtensionText;
+		}
+	}
+	else if (Ar.IsLoading())
+	{
+		Extension = EPackageExtension::Unspecified;
+	}
 }
 
 COREUOBJECT_API void FAssetPackageData::SerializeForCache(FArchive& Ar)
diff --git a/Engine/Source/Runtime/CoreUObject/Public/AssetRegistry/AssetData.h b/Engine/Source/Runtime/CoreUObject/Public/AssetRegistry/AssetData.h
--- a/Engine/Source/Runtime/CoreUObject/Public/AssetRegistry/AssetData.h
+++ b/Engine/Source/Runtime/CoreUObject/Public/AssetRegistry/AssetData.h
@@ -30,6 +30,7 @@
 #include "Misc/Guid.h"
 #include "Misc/Optional.h"
 #include "Misc/PackageName.h"
+#include "Misc/PackagePath.h"
 #include "Misc/SecureHash.h"
 #include "Misc/StringBuilder.h"
 #include "Serialization/Archive.h"
@@ -94,6 +95,7 @@ struct COREUOBJECT_API FAssetRegistryVersion
 		ClassPaths,							// Classes are serialized as path names rather than short object names, e.g. /Script/Engine.StaticMesh
 		RemoveAssetPathFNames,				// Asset bundles are serialized as FTopLevelAssetPath instead of FSoftObjectPath, deprecated FAssetData::ObjectPath	
 		AddedHeader,						// Added header with bFilterEditorOnlyData flag 
+		AssetPackageDataHasExtension,		// Added Extension to AssetPackageData.
 
 		// -----<new versions can be added above this line>-------------------------------------------------
 		VersionPlusOne,
@@ -930,12 +932,16 @@ private:
 	/** Bit storage for flags */
 	uint32 Flags;
 
+public:
+	EPackageExtension Extension;
+
 public:
 
 	FAssetPackageData()
 		: DiskSize(0)
 		, FileVersionLicenseeUE(-1)
 		, Flags(0)
+		, Extension(EPackageExtension::Unspecified)
 	{
 	}
 
