Generated class packages that have no digests will now also be marked as external when building a project to prevent a scenario where the compiler obliterates their types yet the packages don't get recompiled

+ Adding AR API to check if AR should scan for assets on startup

diff --git a/Engine/Source/Runtime/AssetRegistry/Private/AssetRegistry.cpp b/Engine/Source/Runtime/AssetRegistry/Private/AssetRegistry.cpp
--- a/Engine/Source/Runtime/AssetRegistry/Private/AssetRegistry.cpp
+++ b/Engine/Source/Runtime/AssetRegistry/Private/AssetRegistry.cpp
@@ -805,17 +805,7 @@ void FAssetRegistryImpl::Initialize(Impl::FInitializeContext& Context)
 	Utils::InitializeSerializationOptionsFromIni(SerializationOptions, FString());
 	Utils::InitializeSerializationOptionsFromIni(DevelopmentSerializationOptions, FString(), UE::AssetRegistry::ESerializationTarget::ForDevelopment);
 
-	// If in the editor or cookcommandlet, we start the GlobalGatherer now
-	// In the game or other commandlets, we do not construct it until project or commandlet code calls SearchAllAssets or ScanPathsSynchronous
-	bool bSearchAllAssetsAtStart = GIsEditor && (!IsRunningCommandlet() || IsRunningCookCommandlet());
-#if !UE_BUILD_SHIPPING
-	bool bCommandlineAllAssetsAtStart;
-	if (FParse::Bool(FCommandLine::Get(), TEXT("AssetGatherAll="), bCommandlineAllAssetsAtStart))
-	{
-		bSearchAllAssetsAtStart = bCommandlineAllAssetsAtStart;
-	}
-#endif
-	if (bSearchAllAssetsAtStart)
+	if (ShouldSearchAllAssetsAtStart())
 	{
 		ConstructGatherer();
 		if (!GlobalGatherer->IsSynchronous())
@@ -6825,4 +6815,19 @@ void GetAssetForPackages(TConstArrayView<FName> PackageNames, TMap<FName, FAsset
 
 }
 
+bool ShouldSearchAllAssetsAtStart()
+{
+	// If in the editor or cookcommandlet, we start the GlobalGatherer now
+	// In the game or other commandlets, we do not construct it until project or commandlet code calls SearchAllAssets or ScanPathsSynchronous
+	bool bSearchAllAssetsAtStart = GIsEditor && (!IsRunningCommandlet() || IsRunningCookCommandlet());
+#if !UE_BUILD_SHIPPING
+	bool bCommandlineAllAssetsAtStart;
+	if (FParse::Bool(FCommandLine::Get(), TEXT("AssetGatherAll="), bCommandlineAllAssetsAtStart))
+	{
+		bSearchAllAssetsAtStart = bCommandlineAllAssetsAtStart;
+	}
+#endif // !UE_BUILD_SHIPPING
+	return bSearchAllAssetsAtStart;
+}
+
 } // namespace AssetRegistry
diff --git a/Engine/Source/Runtime/AssetRegistry/Public/AssetRegistry/IAssetRegistry.h b/Engine/Source/Runtime/AssetRegistry/Public/AssetRegistry/IAssetRegistry.h
--- a/Engine/Source/Runtime/AssetRegistry/Public/AssetRegistry/IAssetRegistry.h
+++ b/Engine/Source/Runtime/AssetRegistry/Public/AssetRegistry/IAssetRegistry.h
@@ -788,6 +788,11 @@ namespace AssetRegistry
 	 */
 	ASSETREGISTRY_API const FAssetData* GetMostImportantAsset(TConstArrayView<const FAssetData*> PackageAssetDatas, bool bRequireOneTopLevelAsset);
 
+	/*
+	* Returns true if the asset registry should start searching all assets on startup
+	*/
+	ASSETREGISTRY_API bool ShouldSearchAllAssetsAtStart();
+
 	// Wildcards (*) used when looking up assets in the asset registry
 	extern ASSETREGISTRY_API const FName WildcardFName;
 	extern ASSETREGISTRY_API const FTopLevelAssetPath WildcardPathName;
