Use dev serialization options to read dev AR

diff --git a/Engine/Source/Runtime/AssetRegistry/Private/AssetRegistry.cpp b/Engine/Source/Runtime/AssetRegistry/Private/AssetRegistry.cpp
--- a/Engine/Source/Runtime/AssetRegistry/Private/AssetRegistry.cpp
+++ b/Engine/Source/Runtime/AssetRegistry/Private/AssetRegistry.cpp
@@ -730,8 +730,7 @@ void FAssetRegistryImpl::LoadPremadeAssetRegistry(Impl::FEventContext& EventCont
 			}
 			else
 			{
-				State.InitializeFromExisting(ARState, SerializationOptions, Mode);
-				CachePathsFromState(EventContext, ARState);
+				InitializeState(EventContext, ARState, Mode);
 			}
 		}
 		else
@@ -754,8 +753,7 @@ void FAssetRegistryImpl::LoadPremadeAssetRegistry(Impl::FEventContext& EventCont
 				FAssetRegistryState PluginState;
 				PluginState.Load(SerializedAssetData);
 
-				State.InitializeFromExisting(PluginState, SerializationOptions, FAssetRegistryState::EInitializationMode::Append);
-				CachePathsFromState(EventContext, PluginState);
+				AppendState(EventContext, PluginState);
 			}
 		}
 	}
@@ -1383,6 +1381,7 @@ void FAssetRegistryImpl::RefreshNativeClasses()
 
 	// Read default serialization options
 	Utils::InitializeSerializationOptionsFromIni(SerializationOptions, FString());
+	Utils::InitializeSerializationOptionsFromIni(DevelopmentSerializationOptions, FString(), UE::AssetRegistry::ESerializationTarget::ForDevelopment);
 }
 
 }
@@ -4184,12 +4183,26 @@ void UAssetRegistryImpl::AppendState(const FAssetRegistryState& InState)
 namespace UE::AssetRegistry
 {
 
-void FAssetRegistryImpl::AppendState(Impl::FEventContext& EventContext, const FAssetRegistryState& InState)
+void FAssetRegistryImpl::InitializeState(Impl::FEventContext& EventContext, const FAssetRegistryState& InState, FAssetRegistryState::EInitializationMode Mode)
 {
-	State.InitializeFromExisting(InState, SerializationOptions, FAssetRegistryState::EInitializationMode::Append);
+	// @note Using this define to identify cooked editor for now. We might want to change the name later if we find more differences.
+	State.InitializeFromExisting(
+		InState,
+#if ASSETREGISTRY_ENABLE_PREMADE_REGISTRY_IN_EDITOR
+		DevelopmentSerializationOptions,
+#else
+		SerializationOptions,
+#endif
+		Mode);
+
 	CachePathsFromState(EventContext, InState);
 }
 
+void FAssetRegistryImpl::AppendState(Impl::FEventContext& EventContext, const FAssetRegistryState& InState)
+{
+	InitializeState(EventContext, InState, FAssetRegistryState::EInitializationMode::Append);
+}
+
 void FAssetRegistryImpl::CachePathsFromState(Impl::FEventContext& EventContext, const FAssetRegistryState& InState)
 {
 	SCOPED_BOOT_TIMING("FAssetRegistryImpl::CachePathsFromState");
diff --git a/Engine/Source/Runtime/AssetRegistry/Private/AssetRegistryImpl.h b/Engine/Source/Runtime/AssetRegistry/Private/AssetRegistryImpl.h
--- a/Engine/Source/Runtime/AssetRegistry/Private/AssetRegistryImpl.h
+++ b/Engine/Source/Runtime/AssetRegistry/Private/AssetRegistryImpl.h
@@ -143,6 +143,7 @@ public:
 	void ScanModifiedAssetFiles(Impl::FEventContext& EventContext, Impl::FClassInheritanceContext& InheritanceContext,
 		const TArray<FString>& InFilePaths);
 	void Serialize(FArchive& Ar, Impl::FEventContext& EventContext);
+	void InitializeState(Impl::FEventContext& EventContext, const FAssetRegistryState& InState, FAssetRegistryState::EInitializationMode Mode);
 	void AppendState(Impl::FEventContext& EventContext, const FAssetRegistryState& InState);
 	void GetAllocatedSize(bool bLogDetailed, SIZE_T& StateSize, SIZE_T& StaticSize, SIZE_T& SearchSize) const;
 	bool IsLoadingAssets() const;
