fixed invalid memory order identifiers

diff --git a/Engine/Source/Runtime/AssetRegistry/Private/CriticalSectionQueryable.h b/Engine/Source/Runtime/AssetRegistry/Private/CriticalSectionQueryable.h
--- a/Engine/Source/Runtime/AssetRegistry/Private/CriticalSectionQueryable.h
+++ b/Engine/Source/Runtime/AssetRegistry/Private/CriticalSectionQueryable.h
@@ -46,18 +46,18 @@ public:
 
 	bool IsLockedOnCurrentThread() const
 	{
-		return OwnerThreadId.load(std::memory_order::memory_order_relaxed) == FPlatformTLS::GetCurrentThreadId();
+		return OwnerThreadId.load(std::memory_order_relaxed) == FPlatformTLS::GetCurrentThreadId();
 	}
 
 private:
 	void SetOwner()
 	{
-		OwnerThreadId.store(FPlatformTLS::GetCurrentThreadId(), std::memory_order::memory_order_relaxed);
+		OwnerThreadId.store(FPlatformTLS::GetCurrentThreadId(), std::memory_order_relaxed);
 	}
 
 	void ClearOwner()
 	{
-		OwnerThreadId.store(0, std::memory_order::memory_order_relaxed);
+		OwnerThreadId.store(0, std::memory_order_relaxed);
 	}
 
 	FCriticalSection Inner;
diff --git a/Engine/Source/Runtime/Core/Public/Internationalization/StringTableCore.h b/Engine/Source/Runtime/Core/Public/Internationalization/StringTableCore.h
--- a/Engine/Source/Runtime/Core/Public/Internationalization/StringTableCore.h
+++ b/Engine/Source/Runtime/Core/Public/Internationalization/StringTableCore.h
@@ -220,7 +220,7 @@ public:
 	static bool CanFindOrLoadStringTableAsset()
 	{
 		return FInternationalization::IsAvailable()
-			&& DeferFindOrLoad.load(std::memory_order::memory_order_relaxed) <= 0
+			&& DeferFindOrLoad.load(std::memory_order_relaxed) <= 0
 			&& (!InstancePtr || InstancePtr->CanFindOrLoadStringTableAssetImpl());
 	}
 
