Header removal - Remove CoreMinimal from AssetRegistry headers.
#rnx

diff --git a/Engine/Source/Runtime/AssetRegistry/Public/AssetRegistry/AssetRegistryHelpers.h b/Engine/Source/Runtime/AssetRegistry/Public/AssetRegistry/AssetRegistryHelpers.h
--- a/Engine/Source/Runtime/AssetRegistry/Public/AssetRegistry/AssetRegistryHelpers.h
+++ b/Engine/Source/Runtime/AssetRegistry/Public/AssetRegistry/AssetRegistryHelpers.h
@@ -8,7 +8,6 @@
 #include "Containers/Array.h"
 #include "Containers/Set.h"
 #include "Containers/UnrealString.h"
-#include "CoreMinimal.h"
 #include "UObject/NameTypes.h"
 #include "UObject/Object.h"
 #include "UObject/ObjectMacros.h"
@@ -17,6 +16,10 @@
 #include "UObject/UObjectGlobals.h"
 #include "Containers/ArrayView.h"
 
+#if UE_ENABLE_INCLUDE_ORDER_DEPRECATED_IN_5_2
+#include "CoreMinimal.h"
+#endif
+
 #include "AssetRegistryHelpers.generated.h"
 
 class IAssetRegistry;
diff --git a/Engine/Source/Runtime/AssetRegistry/Public/AssetRegistry/AssetRegistryModule.h b/Engine/Source/Runtime/AssetRegistry/Public/AssetRegistry/AssetRegistryModule.h
--- a/Engine/Source/Runtime/AssetRegistry/Public/AssetRegistry/AssetRegistryModule.h
+++ b/Engine/Source/Runtime/AssetRegistry/Public/AssetRegistry/AssetRegistryModule.h
@@ -6,7 +6,6 @@
 #include "AssetRegistry/IAssetRegistry.h"
 #include "Containers/Array.h"
 #include "Containers/UnrealString.h"
-#include "CoreMinimal.h"
 #include "CoreTypes.h"
 #include "Misc/AssetRegistryInterface.h"
 #include "Modules/ModuleInterface.h"
@@ -14,6 +13,10 @@
 #include "UObject/NameTypes.h"
 #include "UObject/Object.h"
 
+#if UE_ENABLE_INCLUDE_ORDER_DEPRECATED_IN_5_2
+#include "CoreMinimal.h"
+#endif
+
 class UPackage;
 
 namespace AssetRegistryConstants
diff --git a/Engine/Source/Runtime/AssetRegistry/Public/AssetRegistry/AssetRegistryState.h b/Engine/Source/Runtime/AssetRegistry/Public/AssetRegistry/AssetRegistryState.h
--- a/Engine/Source/Runtime/AssetRegistry/Public/AssetRegistry/AssetRegistryState.h
+++ b/Engine/Source/Runtime/AssetRegistry/Public/AssetRegistry/AssetRegistryState.h
@@ -10,7 +10,6 @@
 #include "Containers/Map.h"
 #include "Containers/Set.h"
 #include "Containers/StringFwd.h"
-#include "CoreMinimal.h"
 #include "CoreTypes.h"
 #include "HAL/PlatformCrt.h"
 #include "Misc/AssetRegistryInterface.h"
@@ -19,6 +18,10 @@
 #include "UObject/SoftObjectPath.h"
 #include "UObject/TopLevelAssetPath.h"
 
+#if UE_ENABLE_INCLUDE_ORDER_DEPRECATED_IN_5_2
+#include "CoreMinimal.h"
+#endif
+
 class FArchive;
 class FAssetDataTagMap;
 class FAssetDataTagMapSharedView;
diff --git a/Engine/Source/Runtime/AssetRegistry/Public/AssetRegistry/IAssetRegistry.h b/Engine/Source/Runtime/AssetRegistry/Public/AssetRegistry/IAssetRegistry.h
--- a/Engine/Source/Runtime/AssetRegistry/Public/AssetRegistry/IAssetRegistry.h
+++ b/Engine/Source/Runtime/AssetRegistry/Public/AssetRegistry/IAssetRegistry.h
@@ -2,10 +2,9 @@
 
 #pragma once
 
-#include "CoreMinimal.h"
-
 #include "AssetRegistry/ARFilter.h"
 #include "AssetRegistry/AssetIdentifier.h"
+#include "Containers/Array.h"
 #include "Containers/BitArray.h"
 #include "Containers/StringFwd.h"
 #include "Misc/AssetRegistryInterface.h"
@@ -14,6 +13,7 @@
 
 #if UE_ENABLE_INCLUDE_ORDER_DEPRECATED_IN_5_2
 #include "AssetRegistry/AssetData.h"
+#include "CoreMinimal.h"
 #endif
 
 #include "IAssetRegistry.generated.h"
@@ -23,11 +23,12 @@
 #endif
 
 class FArchive;
+class FAssetRegistryState;
+class FDependsNode;
 struct FARFilter;
 struct FARCompiledFilter;
+struct FAssetData;
 struct FAssetRegistrySerializationOptions;
-class FAssetRegistryState;
-class FDependsNode;
 struct FPackageFileSummary;
 struct FObjectExport;
 struct FObjectImport;
diff --git a/Engine/Source/Runtime/AssetRegistry/Public/AssetRegistry/PathTree.h b/Engine/Source/Runtime/AssetRegistry/Public/AssetRegistry/PathTree.h
--- a/Engine/Source/Runtime/AssetRegistry/Public/AssetRegistry/PathTree.h
+++ b/Engine/Source/Runtime/AssetRegistry/Public/AssetRegistry/PathTree.h
@@ -4,10 +4,13 @@
 
 #include "Containers/Map.h"
 #include "Containers/Set.h"
-#include "CoreMinimal.h"
 #include "HAL/Platform.h"
 #include "UObject/NameTypes.h"
 
+#if UE_ENABLE_INCLUDE_ORDER_DEPRECATED_IN_5_2
+#include "CoreMinimal.h"
+#endif
+
 template <typename FuncType> class TFunctionRef;
 
 class ASSETREGISTRY_API FPathTree
diff --git a/Engine/Source/Runtime/CoreUObject/Public/AssetRegistry/ARFilter.h b/Engine/Source/Runtime/CoreUObject/Public/AssetRegistry/ARFilter.h
--- a/Engine/Source/Runtime/CoreUObject/Public/AssetRegistry/ARFilter.h
+++ b/Engine/Source/Runtime/CoreUObject/Public/AssetRegistry/ARFilter.h
@@ -2,12 +2,15 @@
 
 #pragma once
 
-#include "CoreMinimal.h"
 #include "UObject/Class.h"
 #include "UObject/ObjectMacros.h"
 #include "UObject/SoftObjectPath.h"
 #include "UObject/TopLevelAssetPath.h"
 
+#if UE_ENABLE_INCLUDE_ORDER_DEPRECATED_IN_5_2
+#include "CoreMinimal.h"
+#endif
+
 struct FAssetData;
 
 /**
diff --git a/Engine/Source/Runtime/CoreUObject/Public/AssetRegistry/AssetBundleData.h b/Engine/Source/Runtime/CoreUObject/Public/AssetRegistry/AssetBundleData.h
--- a/Engine/Source/Runtime/CoreUObject/Public/AssetRegistry/AssetBundleData.h
+++ b/Engine/Source/Runtime/CoreUObject/Public/AssetRegistry/AssetBundleData.h
@@ -4,7 +4,6 @@
 
 #include "Containers/Array.h"
 #include "Containers/UnrealString.h"
-#include "CoreMinimal.h"
 #include "CoreTypes.h"
 #include "Misc/AssertionMacros.h"
 #include "UObject/Class.h"
@@ -14,6 +13,10 @@
 #include "UObject/SoftObjectPtr.h"
 #include "UObject/TopLevelAssetPath.h"
 
+#if UE_ENABLE_INCLUDE_ORDER_DEPRECATED_IN_5_2
+#include "CoreMinimal.h"
+#endif
+
 class FOutputDevice;
 class UObject;
 
diff --git a/Engine/Source/Runtime/CoreUObject/Public/AssetRegistry/AssetData.h b/Engine/Source/Runtime/CoreUObject/Public/AssetRegistry/AssetData.h
--- a/Engine/Source/Runtime/CoreUObject/Public/AssetRegistry/AssetData.h
+++ b/Engine/Source/Runtime/CoreUObject/Public/AssetRegistry/AssetData.h
@@ -15,7 +15,6 @@
 #include "Containers/StringFwd.h"
 #include "Containers/StringView.h"
 #include "Containers/UnrealString.h"
-#include "CoreMinimal.h"
 #include "HAL/PlatformMath.h"
 #include "HAL/UnrealMemory.h"
 #include "IO/IoChunkId.h"
@@ -54,6 +53,10 @@
 #include "UObject/UObjectGlobals.h"
 #include "UObject/UnrealNames.h"
 
+#if UE_ENABLE_INCLUDE_ORDER_DEPRECATED_IN_5_2
+#include "CoreMinimal.h"
+#endif
+
 struct FAssetBundleData;
 struct FCustomVersion;
 class FCbWriter;
diff --git a/Engine/Source/Runtime/CoreUObject/Public/AssetRegistry/AssetDataTagMap.h b/Engine/Source/Runtime/CoreUObject/Public/AssetRegistry/AssetDataTagMap.h
--- a/Engine/Source/Runtime/CoreUObject/Public/AssetRegistry/AssetDataTagMap.h
+++ b/Engine/Source/Runtime/CoreUObject/Public/AssetRegistry/AssetDataTagMap.h
@@ -2,7 +2,6 @@
 
 #pragma once
 
-#include "CoreMinimal.h"
 #include "Containers/SortedMap.h"
 #include "HAL/CriticalSection.h"
 #include "Misc/StringBuilder.h"
@@ -10,6 +9,10 @@
 #include "Templates/TypeCompatibleBytes.h"
 #include "UObject/TopLevelAssetPath.h"
 
+#if UE_ENABLE_INCLUDE_ORDER_DEPRECATED_IN_5_2
+#include "CoreMinimal.h"
+#endif
+
 class FAssetRegistryState;
 class FAssetTagValueRef;
 class FAssetDataTagMapSharedView;
diff --git a/Engine/Source/Runtime/CoreUObject/Public/Misc/AssetRegistryInterface.h b/Engine/Source/Runtime/CoreUObject/Public/Misc/AssetRegistryInterface.h
--- a/Engine/Source/Runtime/CoreUObject/Public/Misc/AssetRegistryInterface.h
+++ b/Engine/Source/Runtime/CoreUObject/Public/Misc/AssetRegistryInterface.h
@@ -5,7 +5,6 @@
 #include "Containers/Array.h"
 #include "Containers/ContainersFwd.h"
 #include "Containers/Set.h"
-#include "CoreMinimal.h"
 #include "CoreTypes.h"
 #include "Misc/EnumClassFlags.h"
 #include "Modules/ModuleInterface.h"
@@ -13,6 +12,10 @@
 #include "UObject/NameTypes.h"
 #include "UObject/TopLevelAssetPath.h"
 
+#if UE_ENABLE_INCLUDE_ORDER_DEPRECATED_IN_5_2
+#include "CoreMinimal.h"
+#endif
+
 class FLinkerSave;
 class FStructuredArchiveRecord;
 class ITargetPlatform;
