AssetDataGatherer: Optimize Added FFilesToSearch.NumAvailable() by adding a AvailableFilesNum field instead of recalculating it every call.
PR #9939: (Contributed by [REDACTED_PERSON])
#rnx

diff --git a/Engine/Source/Runtime/AssetRegistry/Private/AssetDataGatherer.cpp b/Engine/Source/Runtime/AssetRegistry/Private/AssetDataGatherer.cpp
--- a/Engine/Source/Runtime/AssetRegistry/Private/AssetDataGatherer.cpp
+++ b/Engine/Source/Runtime/AssetRegistry/Private/AssetDataGatherer.cpp
@@ -4382,6 +4382,7 @@ namespace UE::AssetDataGather::Private
 
 void FFilesToSearch::AddPriorityFile(FGatheredPathData&& FilePath)
 {
+	++AvailableFilesNum;
 	BlockingFiles.Add(MoveTemp(FilePath));
 }
 
@@ -4394,11 +4395,13 @@ void FFilesToSearch::AddDirectory(FString&& DirAbsPath, TArray<FGatheredPathData
 	check(!DirAbsPath.IsEmpty());
 
 	FTreeNode& Node = Root.FindOrAddNode(DirAbsPath);
+	AvailableFilesNum += FilePaths.Num();
 	Node.AddFiles(MoveTemp(FilePaths));
 }
 
 void FFilesToSearch::AddFileAgainAfterTimeout(FGatheredPathData&& FilePath)
 {
+	++AvailableFilesNum;
 	BlockingFiles.AddFront(MoveTemp(FilePath));
 }
 
@@ -4413,6 +4416,7 @@ void FFilesToSearch::RetryLaterRetryFiles()
 	{
 		FGatheredPathData FilePath = LaterRetryFiles.PopFrontValue();
 		FTreeNode& Node = Root.FindOrAddNode(FPathViews::GetPath(FilePath.LocalAbsPath));
+		++AvailableFilesNum;
 		Node.AddFile(MoveTemp(FilePath));
 	}
 }
@@ -4420,12 +4424,15 @@ void FFilesToSearch::RetryLaterRetryFiles()
 template <typename AllocatorType>
 void FFilesToSearch::PopFront(TArray<FGatheredPathData, AllocatorType>& Out, int32 NumToPop)
 {
+	int32 InitialNumToPop = NumToPop;
 	while (NumToPop > 0 && !BlockingFiles.IsEmpty())
 	{
 		Out.Add(BlockingFiles.PopFrontValue());
 		--NumToPop;
 	}
 	Root.PopFiles(Out, NumToPop);
+	AvailableFilesNum += NumToPop - InitialNumToPop;
+	check(AvailableFilesNum >= 0);
 }
 
 void FFilesToSearch::PrioritizeDirectory(FStringView DirAbsPath, EPriority Priority)
@@ -4500,7 +4507,7 @@ int32 FFilesToSearch::Num() const
 
 int32 FFilesToSearch::GetNumAvailable() const
 {
-	return BlockingFiles.Num() + Root.NumFiles();
+	return AvailableFilesNum;
 }
 
 SIZE_T FFilesToSearch::GetAllocatedSize() const
diff --git a/Engine/Source/Runtime/AssetRegistry/Private/AssetDataGathererPrivate.h b/Engine/Source/Runtime/AssetRegistry/Private/AssetDataGathererPrivate.h
--- a/Engine/Source/Runtime/AssetRegistry/Private/AssetDataGathererPrivate.h
+++ b/Engine/Source/Runtime/AssetRegistry/Private/AssetDataGathererPrivate.h
@@ -232,6 +232,7 @@ private:
 	FTreeNode Root;
 	TRingBuffer<UE::AssetDataGather::Private::FGatheredPathData> BlockingFiles;
 	TRingBuffer<UE::AssetDataGather::Private::FGatheredPathData> LaterRetryFiles;
+	int32 AvailableFilesNum = 0;
 };
 
 
