AssetRegistry - Change SavePackage to store the extended AssetRegistryTags for the saved AssetDatas into the global AssetRegistry, instead of just the non-extended. Change updating AssetDatas when packages are loaded to NOT remove no-longer-existing tags, since these tags might be present in the extended tags.
#rn Minor, AssetRegistry

diff --git a/Engine/Source/Editor/UnrealEd/Private/EditorEngine.cpp b/Engine/Source/Editor/UnrealEd/Private/EditorEngine.cpp
--- a/Engine/Source/Editor/UnrealEd/Private/EditorEngine.cpp
+++ b/Engine/Source/Editor/UnrealEd/Private/EditorEngine.cpp
@@ -4690,14 +4690,7 @@ FSavePackageResultStruct UEditorEngine::Save(UPackage* InOuter, UObject* InAsset
 	{
 		// Notify the asset registry
 		IAssetRegistry& AssetRegistry = IAssetRegistry::GetChecked();
-		ForEachObjectWithPackage(InOuter, [&AssetRegistry](UObject* Object)
-			{
-				if (Object->IsAsset() && !UE::AssetRegistry::FFiltering::ShouldSkipAsset(Object))
-				{
-					AssetRegistry.AssetSaved(*Object);
-				}
-				return true;
-			}, false /*bIncludeNestedObjects*/);
+		AssetRegistry.AssetsSaved(MoveTemp(Result.SavedAssets));
 	}
 
 	return Result;
diff --git a/Engine/Source/Runtime/AssetRegistry/Private/AssetRegistry.cpp b/Engine/Source/Runtime/AssetRegistry/Private/AssetRegistry.cpp
--- a/Engine/Source/Runtime/AssetRegistry/Private/AssetRegistry.cpp
+++ b/Engine/Source/Runtime/AssetRegistry/Private/AssetRegistry.cpp
@@ -3666,17 +3666,65 @@ void UAssetRegistryImpl::AssetRenamed(const UObject* RenamedAsset, const FString
 }
 
 void UAssetRegistryImpl::AssetSaved(const UObject& SavedAsset)
+{
+	AssetFullyUpdateTags(const_cast<UObject*>(&SavedAsset));
+}
+
+void UAssetRegistryImpl::AssetsSaved(TArray<FAssetData>&& Assets)
 {
 #if WITH_EDITOR
-	if (!ensure(SavedAsset.IsAsset()))
+	UE::AssetRegistry::Impl::FEventContext EventContext;
 	{
-		return;
+		FWriteScopeLock InterfaceScopeLock(InterfaceLock);
+		GuardedData.AssetsSaved(EventContext, MoveTemp(Assets));
 	}
+	Broadcast(EventContext);
+#endif
+}
+
+void UAssetRegistryImpl::AssetFullyUpdateTags(UObject* Object)
+{
+#if WITH_EDITOR
+	FAssetData AssetData(Object);
+	TArray<FAssetData> Assets;
+	Assets.Add(MoveTemp(AssetData));
+
+	UE::AssetRegistry::Impl::FEventContext EventContext;
+	{
+		FWriteScopeLock InterfaceScopeLock(InterfaceLock);
+		GuardedData.AssetsSaved(EventContext, MoveTemp(Assets));
+	}
+	Broadcast(EventContext);
+#endif
+}
+
+namespace UE::AssetRegistry
+{
+
+#if WITH_EDITOR
+void FAssetRegistryImpl::AssetsSaved(UE::AssetRegistry::Impl::FEventContext& EventContext, TArray<FAssetData>&& Assets)
+{
 	LLM_SCOPE(ELLMTag::AssetRegistry);
+	for (FAssetData& NewAssetData : Assets)
+	{
+		FCachedAssetKey Key(NewAssetData);
+		FAssetData** DataFromGather = State.CachedAssets.Find(Key);
 
-	FWriteScopeLock InterfaceScopeLock(InterfaceLock);
-	GuardedData.AddLoadedAssetToProcess(SavedAsset);
+		AssetDataObjectPathsUpdatedOnLoad.Add(NewAssetData.GetSoftObjectPath());
+
+		if (!DataFromGather)
+		{
+			FAssetData* ClonedAssetData = new FAssetData(MoveTemp(NewAssetData));
+			AddAssetData(EventContext, ClonedAssetData);
+		}
+		else
+		{
+			UpdateAssetData(EventContext, *DataFromGather, MoveTemp(NewAssetData), false /* bKeepDeletedTags */);
+		}
+	}
+}
 #endif
+
 }
 
 void UAssetRegistryImpl::AssetTagsFinalized(const UObject& FinalizedAsset)
@@ -4710,7 +4758,7 @@ void FAssetRegistryImpl::AssetSearchDataGathered(Impl::FEventContext& EventConte
 					PostLoadAssetRegistryTags(BackgroundResult.Get());
 #endif
 					// The asset exists in the cache from disk and has not yet been loaded into memory, update it with the new background data
-					UpdateAssetData(EventContext, ExistingAssetData, MoveTemp(*BackgroundResult));
+					UpdateAssetData(EventContext, ExistingAssetData, MoveTemp(*BackgroundResult), false /* bKeepDeletedTags */);
 				}
 			}
 		}
@@ -5046,7 +5094,8 @@ void FAssetRegistryImpl::AddAssetData(Impl::FEventContext& EventContext, FAssetD
 	}
 }
 
-void FAssetRegistryImpl::UpdateAssetData(Impl::FEventContext& EventContext, FAssetData* AssetData, FAssetData&& NewAssetData)
+void FAssetRegistryImpl::UpdateAssetData(Impl::FEventContext& EventContext, FAssetData* AssetData,
+	FAssetData&& NewAssetData, bool bKeepDeletedTags)
 {
 	// Update the class map if updating a blueprint
 	if (ClassGeneratorNames.Contains(AssetData->AssetClassPath))
@@ -5088,6 +5137,33 @@ void FAssetRegistryImpl::UpdateAssetData(Impl::FEventContext& EventContext, FAss
 		}
 	}
 
+	if (bKeepDeletedTags)
+	{
+		TOptional<FAssetDataTagMap> UpdatedTags;
+		AssetData->TagsAndValues.ForEach([&NewAssetData, &UpdatedTags](const TPair<FName, FAssetTagValueRef>& TagPair)
+			{
+				if (UpdatedTags)
+				{
+					if (!UpdatedTags->Contains(TagPair.Key))
+					{
+						UpdatedTags->Add(TagPair.Key, TagPair.Value.GetStorageString());
+					}
+				}
+				else
+				{
+					if (!NewAssetData.TagsAndValues.Contains(TagPair.Key))
+					{
+						UpdatedTags.Emplace(NewAssetData.TagsAndValues.CopyMap());
+						UpdatedTags->Add(TagPair.Key, TagPair.Value.GetStorageString());
+					}
+				}
+			});
+		if (UpdatedTags)
+		{
+			NewAssetData.TagsAndValues = FAssetDataTagMapSharedView(MoveTemp(*UpdatedTags));
+		}
+	}
+
 	bool bModified;
 	State.UpdateAssetData(AssetData, MoveTemp(NewAssetData), &bModified);
 	
@@ -5515,9 +5591,6 @@ void UAssetRegistryImpl::ProcessLoadedAssetsToUpdateCache(UE::AssetRegistry::Imp
 		return;
 	}
 
-	LLM_SCOPE(ELLMTag::AssetRegistry);
-
-
 	constexpr int32 BatchSize = 16;
 	TArray<const UObject*> BatchObjects;
 	TArray<FAssetData, TInlineAllocator<BatchSize>> BatchAssetDatas;
@@ -5641,7 +5714,14 @@ void FAssetRegistryImpl::PushProcessLoadedAssetsBatch(Impl::FEventContext& Event
 		}
 		else
 		{
-			UpdateAssetData(EventContext, *DataFromGather, MoveTemp(NewAssetData));
+			// When updating disk-based AssetData with the AssetData from a loaded UObject, we keep
+			// existing tags from disk even if they no are no longer returned from the GetAssetRegistryTags
+			// function on the loaded UObject. We do this because they might come from GetAssetRegistryTagsExtended,
+			// which is only called during Save.
+			// Modified tag values on the other hand do overwrite the old values from disk.
+			// This means that the only way to delete no-longer present tags from an AssetData
+			// is to resave the package, or to manually call AssetFullyUpdateTags from c++.
+			UpdateAssetData(EventContext, *DataFromGather, MoveTemp(NewAssetData), true /* bKeepDeletedTags */);
 		}
 	}
 
@@ -6534,7 +6614,7 @@ bool FAssetRegistryImpl::SetPrimaryAssetIdForObjectPath(Impl::FEventContext& Eve
 
 	FAssetData NewAssetData(*AssetData);
 	NewAssetData.TagsAndValues = FAssetDataTagMapSharedView(MoveTemp(TagsAndValues));
-	UpdateAssetData(EventContext, AssetData, MoveTemp(NewAssetData));
+	UpdateAssetData(EventContext, AssetData, MoveTemp(NewAssetData), false /* bKeepDeletedTags */);
 
 	return true;
 }
diff --git a/Engine/Source/Runtime/AssetRegistry/Private/AssetRegistry.h b/Engine/Source/Runtime/AssetRegistry/Private/AssetRegistry.h
--- a/Engine/Source/Runtime/AssetRegistry/Private/AssetRegistry.h
+++ b/Engine/Source/Runtime/AssetRegistry/Private/AssetRegistry.h
@@ -140,6 +140,8 @@ public:
 	virtual void AssetDeleted(UObject* DeletedAsset) override;
 	virtual void AssetRenamed(const UObject* RenamedAsset, const FString& OldObjectPath) override;
 	virtual void AssetSaved(const UObject& SavedAsset) override;
+	virtual void AssetsSaved(TArray<FAssetData>&& SavedAssets) override;
+	virtual void AssetFullyUpdateTags(UObject* Object) override;
 	virtual void AssetTagsFinalized(const UObject& FinalizedAsset) override;
 
 	virtual void PackageDeleted(UPackage* DeletedPackage) override;
diff --git a/Engine/Source/Runtime/AssetRegistry/Private/AssetRegistryImpl.h b/Engine/Source/Runtime/AssetRegistry/Private/AssetRegistryImpl.h
--- a/Engine/Source/Runtime/AssetRegistry/Private/AssetRegistryImpl.h
+++ b/Engine/Source/Runtime/AssetRegistry/Private/AssetRegistryImpl.h
@@ -194,6 +194,7 @@ public:
 	void TickGatherPackage(Impl::FEventContext& EventContext, const FString& PackageName, const FString& LocalPath);
 	void ClearGathererCache();
 #if WITH_EDITOR
+	void AssetsSaved(UE::AssetRegistry::Impl::FEventContext& EventContext, TArray<FAssetData>&& Assets);
 	void GetProcessLoadedAssetsBatch(TArray<const UObject*>& OutLoadedAssets, uint32 BatchSize);
 	void PushProcessLoadedAssetsBatch(Impl::FEventContext& EventContext,
 		TArrayView<FAssetData> LoadedAssetDatas, TArrayView<const UObject*> UnprocessedFromBatch);
@@ -306,7 +307,7 @@ private:
 	void AddAssetData(Impl::FEventContext& EventContext, FAssetData* AssetData);
 
 	/** Updates an existing asset data with the new value and updates lookup maps */
-	void UpdateAssetData(Impl::FEventContext& EventContext, FAssetData* AssetData, FAssetData&& NewAssetData);
+	void UpdateAssetData(Impl::FEventContext& EventContext, FAssetData* AssetData, FAssetData&& NewAssetData, bool bKeepDeletedTags);
 
 	/** Removes the asset data from the lookup maps */
 	bool RemoveAssetData(Impl::FEventContext& EventContext, FAssetData* AssetData);
diff --git a/Engine/Source/Runtime/AssetRegistry/Public/AssetRegistry/AssetRegistryModule.h b/Engine/Source/Runtime/AssetRegistry/Public/AssetRegistry/AssetRegistryModule.h
--- a/Engine/Source/Runtime/AssetRegistry/Public/AssetRegistry/AssetRegistryModule.h
+++ b/Engine/Source/Runtime/AssetRegistry/Public/AssetRegistry/AssetRegistryModule.h
@@ -77,9 +77,17 @@ public:
 		IAssetRegistry::GetChecked().AssetRenamed(RenamedAsset, OldObjectPath);
 	}
 
+	UE_DEPRECATED(5.2, "Use the new AssetsSaved function that takes FAssetData.")
 	static void AssetSaved(const UObject& SavedAsset)
 	{
+		PRAGMA_DISABLE_DEPRECATION_WARNINGS;
 		IAssetRegistry::GetChecked().AssetSaved(SavedAsset);
+		PRAGMA_ENABLE_DEPRECATION_WARNINGS;
+	}
+
+	static void AssetsSaved(TArray<FAssetData>&& SavedAssets)
+	{
+		IAssetRegistry::GetChecked().AssetsSaved(MoveTemp(SavedAssets));
 	}
 
 	static void PackageDeleted(UPackage* DeletedPackage)
diff --git a/Engine/Source/Runtime/AssetRegistry/Public/AssetRegistry/IAssetRegistry.h b/Engine/Source/Runtime/AssetRegistry/Public/AssetRegistry/IAssetRegistry.h
--- a/Engine/Source/Runtime/AssetRegistry/Public/AssetRegistry/IAssetRegistry.h
+++ b/Engine/Source/Runtime/AssetRegistry/Public/AssetRegistry/IAssetRegistry.h
@@ -576,8 +576,22 @@ public:
 	virtual void AssetRenamed(const UObject* RenamedAsset, const FString& OldObjectPath) = 0;
 
 	/** Informs the asset registry that an in-memory asset has been saved */
+	UE_DEPRECATED(5.2, "Use the new AssetsSaved function that takes FAssetData.")
 	virtual void AssetSaved(const UObject& SavedAsset) = 0;
 
+	/** Called during SavePackage to update the AssetRegistry's copy of the AssetDatas in the package to match the newly saved values. */
+	virtual void AssetsSaved(TArray<FAssetData>&& Assets) = 0;
+
+	/**
+	 * Called on demand from systems that need to fully update an AssetData's tags. When an Asset is loaded its tags
+	 * are updated by calling UObject::GetAssetRegistryTags, but GetAssetRegistryTagsExtended is not called, and tags
+	 * that exist in the Old AssetData but not in the results from GetAssetRegistryTags are kept because they might be
+	 * extended tags. When an asset is saved, GetAssetRegistryTagsExtended is called and all old tags are deleted in
+	 * favor of the new list. AssetFullyUpdated allows a manual trigger of the on-SavePackage behavior:
+	 * GetAssetRegistryTagsExtended is called and all old tags are deleted in favor of the new list.
+	 */
+	virtual void AssetFullyUpdateTags(UObject* Object) = 0;
+
 	/** Informs the asset registry that an in-memory package has been deleted, and all associated assets should be removed */
 	virtual void PackageDeleted (UPackage* DeletedPackage) = 0;
 
diff --git a/Engine/Source/Runtime/CoreUObject/Private/UObject/Package.cpp b/Engine/Source/Runtime/CoreUObject/Private/UObject/Package.cpp
--- a/Engine/Source/Runtime/CoreUObject/Private/UObject/Package.cpp
+++ b/Engine/Source/Runtime/CoreUObject/Private/UObject/Package.cpp
@@ -37,6 +37,30 @@ UPackage::FOnPackageDirtyStateChanged UPackage::PackageDirtyStateChangedEvent;
  */
 UPackage::FOnPackageMarkedDirty UPackage::PackageMarkedDirtyEvent;
 
+PRAGMA_DISABLE_DEPRECATION_WARNINGS // Silence deprecation warnings for deprecated CookedHash member
+FSavePackageResultStruct::FSavePackageResultStruct()
+	: Result(ESavePackageResult::Error), TotalFileSize(0), SerializedPackageFlags(0)
+{
+}
+FSavePackageResultStruct::FSavePackageResultStruct(ESavePackageResult InResult)
+	: Result(InResult), TotalFileSize(0), SerializedPackageFlags(0)
+{
+}
+FSavePackageResultStruct::FSavePackageResultStruct(ESavePackageResult InResult, int64 InTotalFileSize)
+	: Result(InResult), TotalFileSize(InTotalFileSize), SerializedPackageFlags(0)
+{
+}
+FSavePackageResultStruct::FSavePackageResultStruct(ESavePackageResult InResult, int64 InTotalFileSize,
+	uint32 InSerializedPackageFlags, TPimplPtr<FLinkerSave> Linker)
+	: Result(InResult), TotalFileSize(InTotalFileSize), SerializedPackageFlags(InSerializedPackageFlags),
+	LinkerSave(MoveTemp(Linker))
+{
+}
+FSavePackageResultStruct::FSavePackageResultStruct(FSavePackageResultStruct&& Other) = default;
+FSavePackageResultStruct& FSavePackageResultStruct::operator=(FSavePackageResultStruct&& Other) = default;
+FSavePackageResultStruct::~FSavePackageResultStruct() = default;
+PRAGMA_DISABLE_DEPRECATION_WARNINGS;
+
 void UPackage::PostInitProperties()
 {
 	Super::PostInitProperties();
diff --git a/Engine/Source/Runtime/CoreUObject/Private/UObject/SavePackage/SaveContext.h b/Engine/Source/Runtime/CoreUObject/Private/UObject/SavePackage/SaveContext.h
--- a/Engine/Source/Runtime/CoreUObject/Private/UObject/SavePackage/SaveContext.h
+++ b/Engine/Source/Runtime/CoreUObject/Private/UObject/SavePackage/SaveContext.h
@@ -2,6 +2,7 @@
 
 #pragma once
 
+#include "AssetRegistry/AssetData.h"
 #include "HAL/FileManager.h"
 #include "HAL/IConsoleManager.h"
 #include "Misc/AssetRegistryInterface.h"
@@ -966,8 +967,10 @@ public:
 		}
 
 		ESavePackageResult FinalResult = IsStubRequested() ? ESavePackageResult::GenerateStub : ESavePackageResult::Success;
-		return FSavePackageResultStruct(FinalResult, TotalPackageSizeUncompressed,
+		FSavePackageResultStruct FinalResultStruct(FinalResult, TotalPackageSizeUncompressed,
 			SerializedPackageFlags, IsCompareLinker() ? MoveTemp(GetHarvestedRealm().Linker) : nullptr);
+		FinalResultStruct.SavedAssets = MoveTemp(SavedAssets);
+		return FinalResultStruct;
 	}
 
 	FObjectSaveContextData& GetObjectSaveContext()
@@ -994,6 +997,11 @@ public:
 		return HarvestedRealms[(uint32)(Realm == ESaveRealm::None ? CurrentHarvestingRealm : Realm)];
 	}
 
+	TArray<FAssetData>& GetSavedAssets()
+	{
+		return SavedAssets;
+	}
+
 public:
 	ESavePackageResult Result;
 
@@ -1059,6 +1067,10 @@ private:
 
 	// Set of harvested prestream packages, should be deprecated
 	TSet<UPackage*> PrestreamPackages;
+
+	// Set of AssetDatas created for the Assets saved into the package
+	TArray<FAssetData> SavedAssets;
+
 };
 
 const TCHAR* LexToString(FSaveContext::ESaveableStatus Status);
diff --git a/Engine/Source/Runtime/CoreUObject/Private/UObject/SavePackage/SavePackageUtilities.cpp b/Engine/Source/Runtime/CoreUObject/Private/UObject/SavePackage/SavePackageUtilities.cpp
--- a/Engine/Source/Runtime/CoreUObject/Private/UObject/SavePackage/SavePackageUtilities.cpp
+++ b/Engine/Source/Runtime/CoreUObject/Private/UObject/SavePackage/SavePackageUtilities.cpp
@@ -4,6 +4,7 @@
 
 #include "Algo/Sort.h"
 #include "Algo/Unique.h"
+#include "AssetRegistry/AssetData.h"
 #include "Blueprint/BlueprintSupport.h"
 #include "CoreMinimal.h"
 #include "HAL/FileManager.h"
@@ -2021,7 +2022,9 @@ namespace UE
 namespace AssetRegistry
 {
 	// See the corresponding ReadPackageDataMain and ReadPackageDataDependencies defined in PackageReader.cpp in AssetRegistry module
-	void WritePackageData(FStructuredArchiveRecord& ParentRecord, bool bIsCooking, const UPackage* Package, FLinkerSave* Linker, const TSet<UObject*>& ImportsUsedInGame, const TSet<FName>& SoftPackagesUsedInGame, const ITargetPlatform* TargetPlatform)
+	void WritePackageData(FStructuredArchiveRecord& ParentRecord, bool bIsCooking, const UPackage* Package,
+		FLinkerSave* Linker, const TSet<UObject*>& ImportsUsedInGame, const TSet<FName>& SoftPackagesUsedInGame,
+		const ITargetPlatform* TargetPlatform, TArray<FAssetData>* OutAssetDatas)
 	{
 		// To avoid large patch sizes, we have frozen cooked package format at the format before VER_UE4_ASSETREGISTRY_DEPENDENCYFLAGS
 		bool bPreDependencyFormat = bIsCooking;
@@ -2061,6 +2064,12 @@ namespace AssetRegistry
 		}
 		int32 ObjectCount = AssetObjects.Num();
 		FStructuredArchive::FArray AssetArray = AssetRegistryRecord.EnterArray(TEXT("TagMap"), ObjectCount);
+		if (OutAssetDatas)
+		{
+			OutAssetDatas->Reset();
+		}
+		FString PackageName = Package->GetName();
+
 		for (int32 ObjectIdx = 0; ObjectIdx < AssetObjects.Num(); ++ObjectIdx)
 		{
 			const UObject* Object = AssetObjects[ObjectIdx];
@@ -2103,6 +2112,20 @@ namespace AssetRegistry
 
 				TagMap.EnterElement(Key) << Value;
 			}
+
+			if (OutAssetDatas)
+			{
+				FAssetDataTagMap TagsAndValues;
+				for (UObject::FAssetRegistryTag& Tag : Tags)
+				{
+					if (!Tag.Name.IsNone() && !Tag.Value.IsEmpty())
+					{
+						TagsAndValues.Add(Tag.Name, MoveTemp(Tag.Value));
+					}
+				}
+				OutAssetDatas->Emplace(PackageName, ObjectPath, FTopLevelAssetPath(ObjectClassName),
+					MoveTemp(TagsAndValues), Package->GetChunkIDs(), Package->GetPackageFlags());
+			}
 		}
 		if (bPreDependencyFormat)
 		{
diff --git a/Engine/Source/Runtime/CoreUObject/Private/UObject/SavePackage2.cpp b/Engine/Source/Runtime/CoreUObject/Private/UObject/SavePackage2.cpp
--- a/Engine/Source/Runtime/CoreUObject/Private/UObject/SavePackage2.cpp
+++ b/Engine/Source/Runtime/CoreUObject/Private/UObject/SavePackage2.cpp
@@ -1791,7 +1791,9 @@ ESavePackageResult WritePackageHeader(FStructuredArchive::FRecord& StructuredArc
 	{
 		// Save asset registry data so the editor can search for information about assets in this package
 		SCOPED_SAVETIMER(UPackage_Save_SaveAssetRegistryData);
-		UE::AssetRegistry::WritePackageData(StructuredArchiveRoot, SaveContext.IsCooking(), SaveContext.GetPackage(), Linker, SaveContext.GetImportsUsedInGame(), SaveContext.GetSoftPackagesUsedInGame(), SaveContext.GetTargetPlatform());
+		UE::AssetRegistry::WritePackageData(StructuredArchiveRoot, SaveContext.IsCooking(), SaveContext.GetPackage(),
+			Linker, SaveContext.GetImportsUsedInGame(), SaveContext.GetSoftPackagesUsedInGame(),
+			SaveContext.GetTargetPlatform(), &SaveContext.GetSavedAssets());
 	}
 	// Save level information used by World browser
 	{
diff --git a/Engine/Source/Runtime/CoreUObject/Public/Misc/AssetRegistryInterface.h b/Engine/Source/Runtime/CoreUObject/Public/Misc/AssetRegistryInterface.h
--- a/Engine/Source/Runtime/CoreUObject/Public/Misc/AssetRegistryInterface.h
+++ b/Engine/Source/Runtime/CoreUObject/Public/Misc/AssetRegistryInterface.h
@@ -21,6 +21,7 @@ class FStructuredArchiveRecord;
 class ITargetPlatform;
 class UObject;
 class UPackage;
+struct FAssetData;
 
 namespace EAssetRegistryDependencyType
 {
@@ -193,7 +194,10 @@ namespace UE::AssetRegistry
 	};
 
 	// Functions to read and write the data used by the AssetRegistry in each package; the format of this data is separate from the format of the data in the asset registry
-	COREUOBJECT_API void WritePackageData(FStructuredArchiveRecord& ParentRecord, bool bIsCooking, const UPackage* Package, FLinkerSave* Linker, const TSet<UObject*>& ImportsUsedInGame, const TSet<FName>& SoftPackagesUsedInGame, const ITargetPlatform* TargetPlatform);
+	COREUOBJECT_API void WritePackageData(FStructuredArchiveRecord& ParentRecord, bool bIsCooking,
+		const UPackage* Package, FLinkerSave* Linker, const TSet<UObject*>& ImportsUsedInGame,
+		const TSet<FName>& SoftPackagesUsedInGame, const ITargetPlatform* TargetPlatform,
+		TArray<FAssetData>* OutAssetDatas);
 	// ReadPackageDataMain and ReadPackageDataDependencies are declared in IAssetRegistry.h, in the AssetRegistry module, because they depend upon some structures defined in the AssetRegistry module
 
 	namespace Private
diff --git a/Engine/Source/Runtime/CoreUObject/Public/UObject/Package.h b/Engine/Source/Runtime/CoreUObject/Public/UObject/Package.h
--- a/Engine/Source/Runtime/CoreUObject/Public/UObject/Package.h
+++ b/Engine/Source/Runtime/CoreUObject/Public/UObject/Package.h
@@ -43,6 +43,7 @@ class FString;
 class ITargetPlatform;
 class UFunction;
 class UMetaData;
+struct FAssetData;
 struct FMD5Hash;
 struct FPackageSaveInfo;
 struct FSavePackageArgs;
@@ -82,7 +83,6 @@ enum class ESavePackageResult
 /**
 * Struct returned from save package, contains the enum as well as extra data about what was written
 */
-PRAGMA_DISABLE_DEPRECATION_WARNINGS // Silence deprecation warnings for deprecated CookedHash member in implicit constructors
 struct FSavePackageResultStruct
 {
 	/** Success/failure of the save operation */
@@ -100,11 +100,17 @@ struct FSavePackageResultStruct
 	/** Linker for linker comparison after save. */
 	TPimplPtr<FLinkerSave> LinkerSave;
 
+	TArray<FAssetData> SavedAssets;
+
 	/** Constructors, it will implicitly construct from the result enum */
-	FSavePackageResultStruct() : Result(ESavePackageResult::Error), TotalFileSize(0), SerializedPackageFlags(0) {}
-	FSavePackageResultStruct(ESavePackageResult InResult) : Result(InResult), TotalFileSize(0), SerializedPackageFlags(0) {}
-	FSavePackageResultStruct(ESavePackageResult InResult, int64 InTotalFileSize) : Result(InResult), TotalFileSize(InTotalFileSize), SerializedPackageFlags(0) {}
-	FSavePackageResultStruct(ESavePackageResult InResult, int64 InTotalFileSize, uint32 InSerializedPackageFlags, TPimplPtr<FLinkerSave> Linker = nullptr) : Result(InResult), TotalFileSize(InTotalFileSize), SerializedPackageFlags(InSerializedPackageFlags), LinkerSave(MoveTemp(Linker)) {}
+	COREUOBJECT_API FSavePackageResultStruct();
+	COREUOBJECT_API FSavePackageResultStruct(ESavePackageResult InResult);
+	COREUOBJECT_API FSavePackageResultStruct(ESavePackageResult InResult, int64 InTotalFileSize);
+	COREUOBJECT_API FSavePackageResultStruct(ESavePackageResult InResult, int64 InTotalFileSize,
+		uint32 InSerializedPackageFlags, TPimplPtr<FLinkerSave> Linker = nullptr);
+	COREUOBJECT_API FSavePackageResultStruct(FSavePackageResultStruct&& Other);
+	COREUOBJECT_API FSavePackageResultStruct& operator=(FSavePackageResultStruct&& Other);
+	COREUOBJECT_API ~FSavePackageResultStruct();
 
 	bool operator==(const FSavePackageResultStruct& Other) const
 	{
@@ -125,7 +131,6 @@ struct FSavePackageResultStruct
 			Result == ESavePackageResult::ReplaceCompletely;
 	}
 };
-PRAGMA_ENABLE_DEPRECATION_WARNINGS
 
 /**
 * A package.
