AssetDataGatherer optimization: Reserve DependencyResults to reduce the number of reallocations.
PR #9938: (Contributed by [REDACTED_PERSON])
#rnx

diff --git a/Engine/Source/Runtime/AssetRegistry/Private/AssetDataGatherer.cpp b/Engine/Source/Runtime/AssetRegistry/Private/AssetDataGatherer.cpp
--- a/Engine/Source/Runtime/AssetRegistry/Private/AssetDataGatherer.cpp
+++ b/Engine/Source/Runtime/AssetRegistry/Private/AssetDataGatherer.cpp
@@ -3187,6 +3187,7 @@ void FAssetDataGatherer::TickInternal(bool& bOutIsTickInterrupt)
 			}
 		}
 
+		DependencyResults.Reserve(FilesToSearch->GetNumAvailable() + DependencyResults.Num());
 		FilesToSearch->PopFront(LocalFilesToSearch, NumToProcess);
 
 		// If no work is available mark idle and exit
