World Partition: remove skipping of external actors when updating asset-registry tags from in-memory objects following 23729645.

diff --git a/Engine/Source/Runtime/AssetRegistry/Private/AssetRegistry.cpp b/Engine/Source/Runtime/AssetRegistry/Private/AssetRegistry.cpp
--- a/Engine/Source/Runtime/AssetRegistry/Private/AssetRegistry.cpp
+++ b/Engine/Source/Runtime/AssetRegistry/Private/AssetRegistry.cpp
@@ -5671,12 +5671,6 @@ void FAssetRegistryImpl::GetProcessLoadedAssetsBatch(TArray<const UObject*>& Out
 			continue;
 		}
 
-		//@todo_ow: this will skip actors because after postload some actors might not have proper transform
-		if (LoadedAsset->HasAnyFlags(RF_HasExternalPackage))
-		{
-			continue;
-		}
-
 		// Take a new snapshot of the asset's data every time it loads or saves
 
 		UPackage* InMemoryPackage = LoadedAsset->GetOutermost();
