Engine changes required to support [redacted].
- Package Reader - add flags to control what sections are loaded and add stats on section sizes.
- Asset registry package reader - break AR data serialization into small structs that can be re-used.
- Soft object paths - Add a serialize function that doesn't do redirector fixup
- World partition - Add accessors for AR data tag names, add function to get encoded actor descriptor AR data.
- FNames - add missing CORE_API tags, remove unused FNameEntryId member

diff --git a/Engine/Source/Runtime/AssetRegistry/Private/PackageReader.cpp b/Engine/Source/Runtime/AssetRegistry/Private/PackageReader.cpp
--- a/Engine/Source/Runtime/AssetRegistry/Private/PackageReader.cpp
+++ b/Engine/Source/Runtime/AssetRegistry/Private/PackageReader.cpp
@@ -1031,32 +1031,25 @@ namespace UE::AssetRegistry
 		return ClassPathName;
 	}
 
-	// See the corresponding WritePackageData defined in SavePackageUtilities.cpp in CoreUObject module
-	bool ReadPackageDataMain(FArchive& BinaryArchive, const FString& PackageName, const FPackageFileSummary& PackageFileSummary, int64& OutDependencyDataOffset,
-		TArray<FAssetData*>& OutAssetDataList, EReadPackageDataMainErrorCode& OutError, const TArray<FObjectImport>* InImports, const TArray<FObjectExport>* InExports)
+	bool FDeserializePackageData::DoSerialize(FArchive& BinaryArchive, const FPackageFileSummary& PackageFileSummary, EReadPackageDataMainErrorCode& OutError)
 	{
-		OutError = EReadPackageDataMainErrorCode::Unknown;
-
-		const FString PackagePath = FPackageName::GetLongPackagePath(PackageName);
-		const int64 PackageFileSize = BinaryArchive.TotalSize();
-		const bool bIsMapPackage = (PackageFileSummary.GetPackageFlags() & PKG_ContainsMap) != 0;
-
 		// To avoid large patch sizes, we have frozen cooked package format at the format before VER_UE4_ASSETREGISTRY_DEPENDENCYFLAGS
-		bool bPreDependencyFormat = PackageFileSummary.GetFileVersionUE() < VER_UE4_ASSETREGISTRY_DEPENDENCYFLAGS || !!(PackageFileSummary.GetPackageFlags() & PKG_FilterEditorOnly);
+		const bool bPreDependencyFormat = PackageFileSummary.GetFileVersionUE() < VER_UE4_ASSETREGISTRY_DEPENDENCYFLAGS || !!(PackageFileSummary.GetPackageFlags() & PKG_FilterEditorOnly);
 
 		// Load offsets to optionally-read data
 		if (bPreDependencyFormat)
 		{
-			OutDependencyDataOffset = INDEX_NONE;
+			DependencyDataOffset = INDEX_NONE;
 		}
 		else
 		{
-			BinaryArchive << OutDependencyDataOffset;
+			BinaryArchive << DependencyDataOffset;
 		}
 
 		// Load the object count
-		int32 ObjectCount = 0;
+		ObjectCount = 0;
 		BinaryArchive << ObjectCount;
+		const int64 PackageFileSize = BinaryArchive.TotalSize();
 		const int32 MinBytesPerObject = 1;
 		if (BinaryArchive.IsError() || ObjectCount < 0 || PackageFileSize < BinaryArchive.Tell() + ObjectCount * MinBytesPerObject)
 		{
@@ -1064,12 +1057,64 @@ namespace UE::AssetRegistry
 			return false;
 		}
 
+		return true;
+	}
+
+	bool FDeserializeObjectPackageData::DoSerialize(FArchive& BinaryArchive, EReadPackageDataMainErrorCode& OutError)
+	{
+		const int32 MinBytesPerTag = 1;
+		const int64 PackageFileSize = BinaryArchive.TotalSize();
+
+		BinaryArchive << ObjectPath;
+		BinaryArchive << ObjectClassName;
+		// @todo make sure this is a full path name
+		BinaryArchive << TagCount;
+		if (BinaryArchive.IsError() || TagCount < 0 || PackageFileSize < BinaryArchive.Tell() + TagCount * MinBytesPerTag)
+		{
+			OutError = EReadPackageDataMainErrorCode::InvalidTagCount;
+			return false;
+		}		
+
+		return true;
+	}
+
+	bool FDeserializeTagData::DoSerialize(FArchive& BinaryArchive, EReadPackageDataMainErrorCode& OutError)
+	{
+		BinaryArchive << Key;
+		BinaryArchive << Value;
+		if (BinaryArchive.IsError())
+		{
+			OutError = EReadPackageDataMainErrorCode::InvalidTag;
+			return false;
+		}
+
+		return true;
+	}
+
+	// See the corresponding WritePackageData defined in SavePackageUtilities.cpp in CoreUObject module
+	bool ReadPackageDataMain(FArchive& BinaryArchive, const FString& PackageName, const FPackageFileSummary& PackageFileSummary, int64& OutDependencyDataOffset,
+		TArray<FAssetData*>& OutAssetDataList, EReadPackageDataMainErrorCode& OutError, const TArray<FObjectImport>* InImports, const TArray<FObjectExport>* InExports)
+	{
+		OutError = EReadPackageDataMainErrorCode::Unknown;
+
+		const FString PackagePath = FPackageName::GetLongPackagePath(PackageName);
+		const int64 PackageFileSize = BinaryArchive.TotalSize();
+		const bool bIsMapPackage = (PackageFileSummary.GetPackageFlags() & PKG_ContainsMap) != 0;
+
+		FDeserializePackageData DeserializePackageData;
+		if (!DeserializePackageData.DoSerialize(BinaryArchive, PackageFileSummary, OutError))
+		{
+			return false;
+		}
+
+		OutDependencyDataOffset = DeserializePackageData.DependencyDataOffset;
+
 		// Worlds that were saved before they were marked public do not have asset data so we will synthesize it here to make sure we see all legacy umaps
 		// We will also do this for maps saved after they were marked public but no asset data was saved for some reason. A bug caused this to happen for some maps.
 		if (bIsMapPackage)
 		{
 			const bool bLegacyPackage = PackageFileSummary.GetFileVersionUE() < VER_UE4_PUBLIC_WORLDS;
-			const bool bNoMapAsset = (ObjectCount == 0);
+			const bool bNoMapAsset = (DeserializePackageData.ObjectCount == 0);
 			if (bLegacyPackage || bNoMapAsset)
 			{
 				FString AssetName = FPackageName::GetLongPackageAssetName(PackageName);
@@ -1077,41 +1122,29 @@ namespace UE::AssetRegistry
 			}
 		}
 
-		const int32 MinBytesPerTag = 1;
 		// UAsset files usually only have one asset, maps and redirectors have multiple
-		for (int32 ObjectIdx = 0; ObjectIdx < ObjectCount; ++ObjectIdx)
+		for (int32 ObjectIdx = 0; ObjectIdx < DeserializePackageData.ObjectCount; ++ObjectIdx)
 		{
-			FString ObjectPath;
-			FString ObjectClassName;
-			int32 TagCount = 0;
-			BinaryArchive << ObjectPath;
-			BinaryArchive << ObjectClassName;
-			// @todo make sure this is a full path name
-			BinaryArchive << TagCount;
-			if (BinaryArchive.IsError() || TagCount < 0 || PackageFileSize < BinaryArchive.Tell() + TagCount * MinBytesPerTag)
+			FDeserializeObjectPackageData ObjectPackageData;
+			if (!ObjectPackageData.DoSerialize(BinaryArchive, OutError))
 			{
-				OutError = EReadPackageDataMainErrorCode::InvalidTagCount;
 				return false;
 			}
 
 			FAssetDataTagMap TagsAndValues;
-			TagsAndValues.Reserve(TagCount);
+			TagsAndValues.Reserve(ObjectPackageData.TagCount);
 
-			for (int32 TagIdx = 0; TagIdx < TagCount; ++TagIdx)
+			for (int32 TagIdx = 0; TagIdx < ObjectPackageData.TagCount; ++TagIdx)
 			{
-				FString Key;
-				FString Value;
-				BinaryArchive << Key;
-				BinaryArchive << Value;
-				if (BinaryArchive.IsError())
+				FDeserializeTagData TagData;
+				if (!TagData.DoSerialize(BinaryArchive, OutError))
 				{
-					OutError = EReadPackageDataMainErrorCode::InvalidTag;
 					return false;
 				}
 
-				if (!Key.IsEmpty() && !Value.IsEmpty())
+				if (!TagData.Key.IsEmpty() && !TagData.Value.IsEmpty())
 				{
-					TagsAndValues.Add(FName(*Key), Value);
+					TagsAndValues.Add(FName(*TagData.Key), TagData.Value);
 				}
 			}
 
@@ -1119,42 +1152,42 @@ namespace UE::AssetRegistry
 			// Here we simply skip over them
 			if (bIsMapPackage && PackageFileSummary.GetFileVersionUE() < VER_UE4_PUBLIC_WORLDS)
 			{
-				if (ObjectPath != FPackageName::GetLongPackageAssetName(PackageName))
+				if (ObjectPackageData.ObjectPath != FPackageName::GetLongPackageAssetName(PackageName))
 				{
 					continue;
 				}
 			}
 
 			// if we have an object path that starts with the package then this asset is outer-ed to another package
-			const bool bFullObjectPath = ObjectPath.StartsWith(TEXT("/"), ESearchCase::CaseSensitive);
+			const bool bFullObjectPath = ObjectPackageData.ObjectPath.StartsWith(TEXT("/"), ESearchCase::CaseSensitive);
 
 			// if we do not have a full object path already, build it
 			if (!bFullObjectPath)
 			{
 				// if we do not have a full object path, ensure that we have a top level object for the package and not a sub object
-				if (!ensureMsgf(!ObjectPath.Contains(TEXT("."), ESearchCase::CaseSensitive), TEXT("Cannot make FAssetData for sub object %s in package %s!"), *ObjectPath, *PackageName))
+				if (!ensureMsgf(!ObjectPackageData.ObjectPath.Contains(TEXT("."), ESearchCase::CaseSensitive), TEXT("Cannot make FAssetData for sub object %s in package %s!"), *ObjectPackageData.ObjectPath, *PackageName))
 				{
-					UE_ASSET_LOG(LogAssetRegistry, Warning, *PackageName, TEXT("Cannot make FAssetData for sub object %s!"), *ObjectPath);
+					UE_ASSET_LOG(LogAssetRegistry, Warning, *PackageName, TEXT("Cannot make FAssetData for sub object %s!"), *ObjectPackageData.ObjectPath);
 					continue;
 				}
-				ObjectPath = PackageName + TEXT(".") + ObjectPath;
+				ObjectPackageData.ObjectPath = PackageName + TEXT(".") + ObjectPackageData.ObjectPath;
 			}
 			// Previously export couldn't have its outer as an import
 			else if (PackageFileSummary.GetFileVersionUE() < VER_UE4_NON_OUTER_PACKAGE_IMPORT)
 			{
-				UE_ASSET_LOG(LogAssetRegistry, Warning, *PackageName, TEXT("Package has invalid export %s, resave source package!"), *ObjectPath);
+				UE_ASSET_LOG(LogAssetRegistry, Warning, *PackageName, TEXT("Package has invalid export %s, resave source package!"), *ObjectPackageData.ObjectPath);
 				continue;
 			}
 
 			// Create a new FAssetData for this asset and update it with the gathered data
-			if (!ObjectClassName.IsEmpty() && FPackageName::IsShortPackageName(ObjectClassName))
+			if (!ObjectPackageData.ObjectClassName.IsEmpty() && FPackageName::IsShortPackageName(ObjectPackageData.ObjectClassName))
 			{
 				int64 CurrentPos = BinaryArchive.Tell();
-				ObjectClassName = ReconstructFullClassPath(BinaryArchive, PackageName, PackageFileSummary,
-					ObjectClassName, InImports, InExports);
+				ObjectPackageData.ObjectClassName = ReconstructFullClassPath(BinaryArchive, PackageName, PackageFileSummary,
+					ObjectPackageData.ObjectClassName, InImports, InExports);
 				BinaryArchive.Seek(CurrentPos);
 			}
-			OutAssetDataList.Add(new FAssetData(PackageName, ObjectPath, FTopLevelAssetPath(ObjectClassName), MoveTemp(TagsAndValues), PackageFileSummary.ChunkIDs, PackageFileSummary.GetPackageFlags()));
+			OutAssetDataList.Add(new FAssetData(PackageName, ObjectPackageData.ObjectPath, FTopLevelAssetPath(ObjectPackageData.ObjectClassName), MoveTemp(TagsAndValues), PackageFileSummary.ChunkIDs, PackageFileSummary.GetPackageFlags()));
 		}
 
 		return true;
diff --git a/Engine/Source/Runtime/AssetRegistry/Public/AssetRegistry/IAssetRegistry.h b/Engine/Source/Runtime/AssetRegistry/Public/AssetRegistry/IAssetRegistry.h
--- a/Engine/Source/Runtime/AssetRegistry/Public/AssetRegistry/IAssetRegistry.h
+++ b/Engine/Source/Runtime/AssetRegistry/Public/AssetRegistry/IAssetRegistry.h
@@ -4,6 +4,7 @@
 
 #include "AssetRegistry/ARFilter.h"
 #include "AssetRegistry/AssetIdentifier.h"
+#include "AssetRegistry/AssetDataTagMap.h"
 #include "Containers/Array.h"
 #include "Containers/BitArray.h"
 #include "Containers/StringFwd.h"
@@ -775,6 +776,32 @@ namespace AssetRegistry
 		InvalidTagCount = 2,
 		InvalidTag = 3,
 	};
+
+	struct ASSETREGISTRY_API FDeserializePackageData
+	{
+		int64 DependencyDataOffset = INDEX_NONE;
+		int32 ObjectCount = 0;
+
+		bool DoSerialize(FArchive& BinaryArchive, const FPackageFileSummary& PackageFileSummary, EReadPackageDataMainErrorCode& OutError);
+	};
+
+	struct ASSETREGISTRY_API FDeserializeObjectPackageData
+	{
+		FString ObjectPath;
+		FString ObjectClassName;
+		int32 TagCount = 0;
+
+		bool DoSerialize(FArchive& BinaryArchive, EReadPackageDataMainErrorCode& OutError);
+	};
+
+	struct ASSETREGISTRY_API FDeserializeTagData
+	{
+		FString Key;
+		FString Value;
+
+		bool DoSerialize(FArchive& BinaryArchive, EReadPackageDataMainErrorCode& OutError);
+	};
+
 	// Functions to read and write the data used by the AssetRegistry in each package; the format of this data is separate from the format of the data in the asset registry
 	// WritePackageData is declared in AssetRegistryInterface.h, in the CoreUObject module, because it is needed by SavePackage in CoreUObject
 	ASSETREGISTRY_API bool ReadPackageDataMain(FArchive& BinaryArchive, const FString& PackageName, const FPackageFileSummary& PackageFileSummary,
diff --git a/Engine/Source/Runtime/Core/Public/UObject/NameTypes.h b/Engine/Source/Runtime/Core/Public/UObject/NameTypes.h
--- a/Engine/Source/Runtime/Core/Public/UObject/NameTypes.h
+++ b/Engine/Source/Runtime/Core/Public/UObject/NameTypes.h
@@ -266,10 +266,10 @@ public:
 	 *
 	 * @param OutSize must be at least GetNameLength()
 	 */
-	void GetUnterminatedName(TCHAR* OutName, uint32 OutSize) const;
+	CORE_API void GetUnterminatedName(TCHAR* OutName, uint32 OutSize) const;
 
 	/** Copy null-terminated name to TCHAR buffer without allocating. */
-	void GetName(TCHAR(&OutName)[NAME_SIZE]) const;
+	CORE_API void GetName(TCHAR(&OutName)[NAME_SIZE]) const;
 
 	/** Copy null-terminated name to ANSICHAR buffer without allocating. Entry must not be wide. */
 	CORE_API void GetAnsiName(ANSICHAR(&OutName)[NAME_SIZE]) const;
@@ -341,7 +341,6 @@ private:
  */
 struct FNameEntrySerialized
 {
-	FNameEntryId Index;
 	bool bIsWide = false;
 
 	union
diff --git a/Engine/Source/Runtime/CoreUObject/Private/UObject/SoftObjectPath.cpp b/Engine/Source/Runtime/CoreUObject/Private/UObject/SoftObjectPath.cpp
--- a/Engine/Source/Runtime/CoreUObject/Private/UObject/SoftObjectPath.cpp
+++ b/Engine/Source/Runtime/CoreUObject/Private/UObject/SoftObjectPath.cpp
@@ -267,31 +267,7 @@ void FSoftObjectPath::SerializePath(FArchive& Ar)
 
 	if (bSerializeInternals)
 	{
-		if (Ar.IsLoading() && Ar.UEVer() < VER_UE4_ADDED_SOFT_OBJECT_PATH)
-		{
-			FString Path;
-			Ar << Path;
-
-			if (Ar.UEVer() < VER_UE4_KEEP_ONLY_PACKAGE_NAMES_IN_STRING_ASSET_REFERENCES_MAP)
-			{
-				Path = FPackageName::GetNormalizedObjectPath(Path);
-			}
-
-			SetPath(MoveTemp(Path));
-		}
-		else if( Ar.IsLoading() && Ar.UEVer() < EUnrealEngineObjectUE5Version::FSOFTOBJECTPATH_REMOVE_ASSET_PATH_FNAMES)
-		{
-			FName AssetPathName;
-			Ar << AssetPathName;
-			AssetPath = WriteToString<FName::StringBufferSize>(AssetPathName).ToView();
-
-			Ar << SubPathString;
-		}
-		else
-		{
-			Ar << AssetPath;
-			Ar << SubPathString;
-		}
+		SerializePathWithoutFixup(Ar);
 	}
 
 #if WITH_EDITOR
@@ -310,7 +286,7 @@ void FSoftObjectPath::SerializePath(FArchive& Ar)
 				PreSavePath(nullptr);
 			}
 		}
-		if (Ar.GetPortFlags()&PPF_DuplicateForPIE)
+		if (Ar.GetPortFlags() & PPF_DuplicateForPIE)
 		{
 			// Remap unique ID if necessary
 			// only for fixing up cross-level references, inter-level references handled in FDuplicateDataReader
@@ -320,6 +296,35 @@ void FSoftObjectPath::SerializePath(FArchive& Ar)
 #endif // WITH_EDITOR
 }
 
+void FSoftObjectPath::SerializePathWithoutFixup(FArchive& Ar)
+{
+	if (Ar.IsLoading() && Ar.UEVer() < VER_UE4_ADDED_SOFT_OBJECT_PATH)
+	{
+		FString Path;
+		Ar << Path;
+
+		if (Ar.UEVer() < VER_UE4_KEEP_ONLY_PACKAGE_NAMES_IN_STRING_ASSET_REFERENCES_MAP)
+		{
+			Path = FPackageName::GetNormalizedObjectPath(Path);
+		}
+
+		SetPath(MoveTemp(Path));
+	}
+	else if (Ar.IsLoading() && Ar.UEVer() < EUnrealEngineObjectUE5Version::FSOFTOBJECTPATH_REMOVE_ASSET_PATH_FNAMES)
+	{
+		FName AssetPathName;
+		Ar << AssetPathName;
+		AssetPath = WriteToString<FName::StringBufferSize>(AssetPathName).ToView();
+
+		Ar << SubPathString;
+	}
+	else
+	{
+		Ar << AssetPath;
+		Ar << SubPathString;
+	}
+}
+
 bool FSoftObjectPath::operator==(FSoftObjectPath const& Other) const
 {
 	return AssetPath == Other.AssetPath && SubPathString == Other.SubPathString;
diff --git a/Engine/Source/Runtime/CoreUObject/Public/UObject/SoftObjectPath.h b/Engine/Source/Runtime/CoreUObject/Public/UObject/SoftObjectPath.h
--- a/Engine/Source/Runtime/CoreUObject/Public/UObject/SoftObjectPath.h
+++ b/Engine/Source/Runtime/CoreUObject/Public/UObject/SoftObjectPath.h
@@ -302,6 +302,9 @@ PRAGMA_ENABLE_DEPRECATION_WARNINGS;
 	/** Serializes the internal path and also handles save/PIE fixups. Call this from the archiver overrides */
 	void SerializePath(FArchive& Ar);
 
+	/** Serializes the internal path without any save/PIE fixups. Only call this directly if you know what you are doing */
+	void SerializePathWithoutFixup(FArchive& Ar);
+
 	/** Fixes up path for saving, call if saving with a method that skips SerializePath. This can modify the path, it will return true if it was modified */
 	bool PreSavePath(bool* bReportSoftObjectPathRedirects = nullptr);
 
diff --git a/Engine/Source/Runtime/Engine/Private/WorldPartition/WorldPartitionActorDescUtils.cpp b/Engine/Source/Runtime/Engine/Private/WorldPartition/WorldPartitionActorDescUtils.cpp
--- a/Engine/Source/Runtime/Engine/Private/WorldPartition/WorldPartitionActorDescUtils.cpp
+++ b/Engine/Source/Runtime/Engine/Private/WorldPartition/WorldPartitionActorDescUtils.cpp
@@ -16,6 +16,16 @@
 static FName NAME_ActorMetaDataClass(TEXT("ActorMetaDataClass"));
 static FName NAME_ActorMetaData(TEXT("ActorMetaData"));
 
+FName FWorldPartitionActorDescUtils::ActorMetaDataClassTagName()
+{
+	return NAME_ActorMetaDataClass;
+}
+
+FName FWorldPartitionActorDescUtils::ActorMetaDataTagName()
+{
+	return NAME_ActorMetaData;
+}
+
 bool FWorldPartitionActorDescUtils::IsValidActorDescriptorFromAssetData(const FAssetData& InAssetData)
 {
 	return InAssetData.FindTag(NAME_ActorMetaDataClass) && InAssetData.FindTag(NAME_ActorMetaData);
@@ -126,12 +136,17 @@ void FWorldPartitionActorDescUtils::AppendAssetDataTagsFromActor(const AActor* I
 	const FString ActorMetaDataClass = GetParentNativeClass(InActor->GetClass())->GetPathName();
 	OutTags.Add(UObject::FAssetRegistryTag(NAME_ActorMetaDataClass, ActorMetaDataClass, UObject::FAssetRegistryTag::TT_Hidden));
 
-	TArray<uint8> SerializedData;
-	ActorDesc->SerializeTo(SerializedData);
-	const FString ActorMetaData = FBase64::Encode(SerializedData);
+	const FString ActorMetaData = GetAssetDataFromActorDescriptor(ActorDesc);
 	OutTags.Add(UObject::FAssetRegistryTag(NAME_ActorMetaData, ActorMetaData, UObject::FAssetRegistryTag::TT_Hidden));
 }
 
+FString FWorldPartitionActorDescUtils::GetAssetDataFromActorDescriptor(TUniquePtr<FWorldPartitionActorDesc>& InActorDesc)
+{
+	TArray<uint8> SerializedData;
+	InActorDesc->SerializeTo(SerializedData);
+	return FBase64::Encode(SerializedData);
+}
+
 void FWorldPartitionActorDescUtils::UpdateActorDescriptorFromActor(const AActor* InActor, TUniquePtr<FWorldPartitionActorDesc>& OutActorDesc)
 {
 	TUniquePtr<FWorldPartitionActorDesc> NewActorDesc(InActor->CreateActorDesc());
diff --git a/Engine/Source/Runtime/Engine/Public/WorldPartition/WorldPartitionActorDesc.h b/Engine/Source/Runtime/Engine/Public/WorldPartition/WorldPartitionActorDesc.h
--- a/Engine/Source/Runtime/Engine/Public/WorldPartition/WorldPartitionActorDesc.h
+++ b/Engine/Source/Runtime/Engine/Public/WorldPartition/WorldPartitionActorDesc.h
@@ -70,6 +70,7 @@ class ENGINE_API FWorldPartitionActorDesc
 	friend struct FWorldPartitionReferenceImpl;
 	friend struct FWorldPartitionActorDescUtils;
 	friend struct FWorldPartitionActorDescUnitTestAcccessor;
+	friend class FAssetRootPackagePatcher;
 
 public:
 	virtual ~FWorldPartitionActorDesc() {}
diff --git a/Engine/Source/Runtime/Engine/Public/WorldPartition/WorldPartitionActorDescUtils.h b/Engine/Source/Runtime/Engine/Public/WorldPartition/WorldPartitionActorDescUtils.h
--- a/Engine/Source/Runtime/Engine/Public/WorldPartition/WorldPartitionActorDescUtils.h
+++ b/Engine/Source/Runtime/Engine/Public/WorldPartition/WorldPartitionActorDescUtils.h
@@ -11,6 +11,11 @@ struct FAssetData;
 
 struct ENGINE_API FWorldPartitionActorDescUtils
 {
+	/** @return The asset registry tag name for ActorMetaDataClass */
+	static FName ActorMetaDataClassTagName();
+	/** @return The asset registry tag name for ActorMetaData */
+	static FName ActorMetaDataTagName();
+
 	/** 
 	 * Checks if the provided asset data contains a valid actor descriptor.
 	 * @param InAssetData	The asset data to look into.
@@ -39,6 +44,13 @@ struct ENGINE_API FWorldPartitionActorDescUtils
 	 */
 	static void AppendAssetDataTagsFromActor(const AActor* InActor, TArray<UObject::FAssetRegistryTag>& OutTags);
 
+	/** 
+	 * Return the actor descriptor data.
+	 * @param InActorDesc	The actor descriptor that will generate the data.
+	 * @return				The actor descriptor data.
+	 */
+	static FString GetAssetDataFromActorDescriptor(TUniquePtr<FWorldPartitionActorDesc>& InActorDesc);
+
 	/** 
 	 * Update an actor descriptor with new values coming from the provided actor.
 	 * @param InActor		The actor to update from.
