AssetRegistry: Optimize boot times when loading and consuming premade asset registries.
In LoadPremadeAssetRegistry always choose the rebuild path if current state is empty, else choose the fastest path to merge states.
Optimize cooked client (and multithreaded cooked server) loading by utilizing the cooked editor async consume logic (while reducing permutations and code complexity).

Win64 Development Results:
CookedEditor, first FlushAsyncLoading down from 3.7s to 0.5s, by reducing the async LoadPremadeAssetRegistry consume from 4.6s to ~450 ms.
Client, LoadModule_AssetRegistry down from 200-400ms to ~20ms, by utilizing the async LoadPremadeAssetRegistry consume path.

#rnx

diff --git a/Engine/Source/Runtime/AssetRegistry/Private/AssetRegistry.cpp b/Engine/Source/Runtime/AssetRegistry/Private/AssetRegistry.cpp
--- a/Engine/Source/Runtime/AssetRegistry/Private/AssetRegistry.cpp
+++ b/Engine/Source/Runtime/AssetRegistry/Private/AssetRegistry.cpp
@@ -177,14 +177,6 @@ static bool CanLoadAsync()
 	return FPlatformProcess::SupportsMultithreading() && FTaskGraphInterface::IsRunning();
 }
 
-static bool CanConsumeAsync()
-{
-	// Async Consumption (copying the preloaded Premade FAssetRegistryState into the global AR) is only
-	// available when async is available, and is not currently used in the cooked game because we have
-	// to block on it anyway.
-	return CanLoadAsync() && !FPlatformProperties::RequiresCookedData();
-}
-
 /** Returns the paths to possible Premade AssetRegistry files, ordered from highest priority to lowest. */
 TArray<FString, TInlineAllocator<2>> GetPriorityPaths()
 {
@@ -307,6 +299,7 @@ private:
 
 	ELoadResult TryLoad()
 	{
+		TRACE_CPUPROFILER_EVENT_SCOPE(FCookedAssetRegistryPreloader::TryLoad);
 		LLM_SCOPE(ELLMTag::AssetRegistry);
 		checkf(!ARPath.IsEmpty(), TEXT("TryLoad must not be called until after TrySetPath has succeeded."));
 
@@ -323,6 +316,7 @@ private:
 
 	void DelayedInitialize()
 	{
+		TRACE_CPUPROFILER_EVENT_SCOPE(FCookedAssetRegistryPreloader::DelayedInitialize);
 		// This function will run before any UObject (ie UAssetRegistryImpl) code can run, so we don't need to do any thread safety
 		// CanLoadAsync - we have to check this after the task graph is ready
 		if (!CanLoadAsync())
@@ -364,6 +358,7 @@ private:
 
 	void KickPreload()
 	{
+		TRACE_CPUPROFILER_EVENT_SCOPE(FCookedAssetRegistryPreloader::KickPreload);
 		// Called from Within the Lock
 		check(LoadState == EState::NotFound && !ARPath.IsEmpty());
 		LoadState = EState::Loading;
@@ -504,7 +499,6 @@ private:
 }
 GPreloader;
 
-#if ASSETREGISTRY_ENABLE_PREMADE_REGISTRY_IN_EDITOR
 FAsyncConsumer::~FAsyncConsumer()
 {
 	if (Consumed)
@@ -552,7 +546,7 @@ void FAsyncConsumer::Wait(UAssetRegistryImpl& UARI, FWriteScopeLock& ScopeLock)
 void FAsyncConsumer::Consume(UAssetRegistryImpl& UARI, UE::AssetRegistry::Impl::FEventContext& EventContext, ELoadResult LoadResult, FAssetRegistryState&& ARState)
 {
 	// Called within the lock
-	UARI.GuardedData.LoadPremadeAssetRegistry(EventContext, LoadResult, MoveTemp(ARState), FAssetRegistryState::EInitializationMode::OnlyUpdateNew);
+	UARI.GuardedData.LoadPremadeAssetRegistry(EventContext, LoadResult, MoveTemp(ARState));
 	check(ReferenceCount >= 1);
 	check(Consumed != nullptr);
 	Consumed->Trigger();
@@ -565,58 +559,54 @@ void FAsyncConsumer::Consume(UAssetRegistryImpl& UARI, UE::AssetRegistry::Impl::
 	}
 }
 
-#endif
-
 }
 
 namespace UE::AssetRegistry
 {
 void FAssetRegistryImpl::ConditionalLoadPremadeAssetRegistry(UAssetRegistryImpl& UARI, Impl::FEventContext& EventContext, FWriteScopeLock& ScopeLock)
 {
-#if ASSETREGISTRY_ENABLE_PREMADE_REGISTRY_IN_EDITOR
 	AsyncConsumer.Wait(UARI, ScopeLock);
-#endif
 }
 
 void FAssetRegistryImpl::ConsumeOrDeferPreloadedPremade(UAssetRegistryImpl& UARI, Impl::FEventContext& EventContext)
 {
 	// Called from inside WriteLock on InterfaceLock
 	using namespace UE::AssetRegistry::Premade;
-	if (!UE::AssetRegistry::Premade::IsEnabled())
+	if (!Premade::IsEnabled())
 	{
 		// if we aren't doing any preloading, then we can set the initial search is done right away.
 		// Otherwise, it is set from LoadPremadeAssetRegistry
 		bCanFinishInitialSearch = true;
 		return;
 	}
-#if ASSETREGISTRY_ENABLE_PREMADE_REGISTRY_IN_EDITOR
-	FPreloader::FConsumeFunction ConsumeFromAsyncThread = [this, &UARI](Premade::ELoadResult LoadResult, FAssetRegistryState&& ARState)
+
+	if (Premade::CanLoadAsync())
 	{
-		Impl::FEventContext EventContext;
+		FPreloader::FConsumeFunction ConsumeFromAsyncThread = [this, &UARI](Premade::ELoadResult LoadResult, FAssetRegistryState&& ARState)
 		{
-			FWriteScopeLock InterfaceScopeLock(UARI.InterfaceLock);
-			AsyncConsumer.Consume(UARI, EventContext, LoadResult, MoveTemp(ARState));
-		}
-		UARI.Broadcast(EventContext);
-	};
-	auto ConsumeOnCurrentThread = [ConsumeFromAsyncThread](Premade::ELoadResult LoadResult, FAssetRegistryState&& ARState) mutable
-	{
-		Async(EAsyncExecution::TaskGraph, [LoadResult, ARState=MoveTemp(ARState), ConsumeFromAsyncThread=MoveTemp(ConsumeFromAsyncThread)]() mutable
+			Impl::FEventContext EventContext;
+			{
+				FWriteScopeLock InterfaceScopeLock(UARI.InterfaceLock);
+				AsyncConsumer.Consume(UARI, EventContext, LoadResult, MoveTemp(ARState));
+			}
+			UARI.Broadcast(EventContext);
+		};
+		auto ConsumeOnCurrentThread = [ConsumeFromAsyncThread](Premade::ELoadResult LoadResult, FAssetRegistryState&& ARState) mutable
 		{
-			ConsumeFromAsyncThread(LoadResult, MoveTemp(ARState));
-		});
-	};
-	if (CanConsumeAsync())
-	{
+			Async(EAsyncExecution::TaskGraph, [LoadResult, ARState=MoveTemp(ARState), ConsumeFromAsyncThread=MoveTemp(ConsumeFromAsyncThread)]() mutable
+			{
+				ConsumeFromAsyncThread(LoadResult, MoveTemp(ARState));
+			});
+		};
+
 		AsyncConsumer.PrepareForConsume();
 		GPreloader.ConsumeOrDefer(MoveTemp(ConsumeOnCurrentThread), MoveTemp(ConsumeFromAsyncThread));
 	}
 	else
-#endif
 	{
 		GPreloader.Consume([this, &EventContext](Premade::ELoadResult LoadResult, FAssetRegistryState&& ARState)
 			{
-				LoadPremadeAssetRegistry(EventContext, LoadResult, MoveTemp(ARState), FAssetRegistryState::EInitializationMode::Rebuild);
+				LoadPremadeAssetRegistry(EventContext, LoadResult, MoveTemp(ARState));
 			});
 	}
 }
@@ -714,23 +704,31 @@ FAssetRegistryImpl::FAssetRegistryImpl()
 }
 
 void FAssetRegistryImpl::LoadPremadeAssetRegistry(Impl::FEventContext& EventContext,
-	Premade::ELoadResult LoadResult, FAssetRegistryState&& ARState, FAssetRegistryState::EInitializationMode Mode)
+	Premade::ELoadResult LoadResult, FAssetRegistryState&& ARState)
 {
 	SCOPED_BOOT_TIMING("LoadPremadeAssetRegistry");
 	UE_SCOPED_ENGINE_ACTIVITY("Loading premade asset registry");
 
 	if (SerializationOptions.bSerializeAssetRegistry)
 	{
+		SCOPED_BOOT_TIMING("LoadPremadeAssetRegistry_Main");
 		if (LoadResult == Premade::ELoadResult::Succeeded)
 		{
-			if (Mode == FAssetRegistryState::EInitializationMode::Rebuild)
+			if (State.GetNumAssets() == 0)
+			{
+				State = MoveTemp(ARState);
+				CachePathsFromState(EventContext, State);
+			}
+			else if (State.GetNumAssets() < ARState.GetNumAssets())
 			{
+				FAssetRegistryState ExistingState = MoveTemp(State);
 				State = MoveTemp(ARState);
 				CachePathsFromState(EventContext, State);
+				AppendState(EventContext, ExistingState);
 			}
 			else
 			{
-				InitializeState(EventContext, ARState, Mode);
+				AppendState(EventContext, ARState, FAssetRegistryState::EInitializationMode::OnlyUpdateNew);
 			}
 		}
 		else
@@ -740,20 +738,23 @@ void FAssetRegistryImpl::LoadPremadeAssetRegistry(Impl::FEventContext& EventCont
 		}
 	}
 
-	TArray<TSharedRef<IPlugin>> ContentPlugins = IPluginManager::Get().GetEnabledPluginsWithContent();
-	for (const TSharedRef<IPlugin>& ContentPlugin : ContentPlugins)
 	{
-		if (ContentPlugin->CanContainContent())
+		SCOPED_BOOT_TIMING("LoadPremadeAssetRegistry_Plugins");
+		TArray<TSharedRef<IPlugin>> ContentPlugins = IPluginManager::Get().GetEnabledPluginsWithContent();
+		for (const TSharedRef<IPlugin>& ContentPlugin : ContentPlugins)
 		{
-			FArrayReader SerializedAssetData;
-			FString PluginAssetRegistry = ContentPlugin->GetBaseDir() / TEXT("AssetRegistry.bin");
-			if (IFileManager::Get().FileExists(*PluginAssetRegistry) && FFileHelper::LoadFileToArray(SerializedAssetData, *PluginAssetRegistry))
+			if (ContentPlugin->CanContainContent())
 			{
-				SerializedAssetData.Seek(0);
-				FAssetRegistryState PluginState;
-				PluginState.Load(SerializedAssetData);
+				FArrayReader SerializedAssetData;
+				FString PluginAssetRegistry = ContentPlugin->GetBaseDir() / TEXT("AssetRegistry.bin");
+				if (IFileManager::Get().FileExists(*PluginAssetRegistry) && FFileHelper::LoadFileToArray(SerializedAssetData, *PluginAssetRegistry))
+				{
+					SerializedAssetData.Seek(0);
+					FAssetRegistryState PluginState;
+					PluginState.Load(SerializedAssetData);
 
-				AppendState(EventContext, PluginState);
+					AppendState(EventContext, PluginState);
+				}
 			}
 		}
 	}
@@ -4236,7 +4237,7 @@ void UAssetRegistryImpl::AppendState(const FAssetRegistryState& InState)
 namespace UE::AssetRegistry
 {
 
-void FAssetRegistryImpl::InitializeState(Impl::FEventContext& EventContext, const FAssetRegistryState& InState, FAssetRegistryState::EInitializationMode Mode)
+void FAssetRegistryImpl::AppendState(Impl::FEventContext& EventContext, const FAssetRegistryState& InState, FAssetRegistryState::EInitializationMode Mode)
 {
 	// @note Using this define to identify cooked editor for now. We might want to change the name later if we find more differences.
 	State.InitializeFromExisting(
@@ -4251,11 +4252,6 @@ void FAssetRegistryImpl::InitializeState(Impl::FEventContext& EventContext, cons
 	CachePathsFromState(EventContext, InState);
 }
 
-void FAssetRegistryImpl::AppendState(Impl::FEventContext& EventContext, const FAssetRegistryState& InState)
-{
-	InitializeState(EventContext, InState, FAssetRegistryState::EInitializationMode::Append);
-}
-
 void FAssetRegistryImpl::CachePathsFromState(Impl::FEventContext& EventContext, const FAssetRegistryState& InState)
 {
 	SCOPED_BOOT_TIMING("FAssetRegistryImpl::CachePathsFromState");
diff --git a/Engine/Source/Runtime/AssetRegistry/Private/AssetRegistryImpl.h b/Engine/Source/Runtime/AssetRegistry/Private/AssetRegistryImpl.h
--- a/Engine/Source/Runtime/AssetRegistry/Private/AssetRegistryImpl.h
+++ b/Engine/Source/Runtime/AssetRegistry/Private/AssetRegistryImpl.h
@@ -37,7 +37,6 @@ namespace UE::AssetRegistry::Impl { struct FClassInheritanceContext; }
 namespace UE::AssetDependencyGatherer::Private { class FRegisteredAssetDependencyGatherer; }
 #endif
 
-#if ASSETREGISTRY_ENABLE_PREMADE_REGISTRY_IN_EDITOR
 namespace UE::AssetRegistry::Premade
 {
 
@@ -68,7 +67,6 @@ private:
 };
 
 }
-#endif
 
 namespace UE::AssetRegistry
 {
@@ -143,8 +141,8 @@ public:
 	void ScanModifiedAssetFiles(Impl::FEventContext& EventContext, Impl::FClassInheritanceContext& InheritanceContext,
 		const TArray<FString>& InFilePaths);
 	void Serialize(FArchive& Ar, Impl::FEventContext& EventContext);
-	void InitializeState(Impl::FEventContext& EventContext, const FAssetRegistryState& InState, FAssetRegistryState::EInitializationMode Mode);
-	void AppendState(Impl::FEventContext& EventContext, const FAssetRegistryState& InState);
+	void AppendState(Impl::FEventContext& EventContext, const FAssetRegistryState& InState,
+		FAssetRegistryState::EInitializationMode Mode = FAssetRegistryState::EInitializationMode::Append);
 	void GetAllocatedSize(bool bLogDetailed, SIZE_T& StateSize, SIZE_T& StaticSize, SIZE_T& SearchSize) const;
 	bool IsLoadingAssets() const;
 	void SetManageReferences(const TMultiMap<FAssetIdentifier, FAssetIdentifier>& ManagerMap,
@@ -341,7 +339,7 @@ private:
 	void ConsumeOrDeferPreloadedPremade(UAssetRegistryImpl& UARI, Impl::FEventContext& EventContext);
 	/** Moves a premade asset registry state into this AR */
 	void LoadPremadeAssetRegistry(Impl::FEventContext& Context,
-		Premade::ELoadResult LoadResult, FAssetRegistryState&& ARState, FAssetRegistryState::EInitializationMode Mode);
+		Premade::ELoadResult LoadResult, FAssetRegistryState&& ARState);
 
 private:
 
@@ -454,9 +452,7 @@ private:
 	/** Class names that return true for IsAsset but which should not be treated as assets in cooked packages */
 	TSet<FTopLevelAssetPath> SkipCookedClasses;
 #endif
-#if ASSETREGISTRY_ENABLE_PREMADE_REGISTRY_IN_EDITOR
 	UE::AssetRegistry::Premade::FAsyncConsumer AsyncConsumer;
-#endif
 	friend struct UE::AssetRegistry::Premade::FAsyncConsumer;
 	friend struct Impl::FClassInheritanceContext;
 };
