AssetRegistry Gather: Fix dereference-after-free in FScanDir::Update.
#rnx

diff --git a/Engine/Source/Runtime/AssetRegistry/Private/AssetDataGatherer.cpp b/Engine/Source/Runtime/AssetRegistry/Private/AssetDataGatherer.cpp
--- a/Engine/Source/Runtime/AssetRegistry/Private/AssetDataGatherer.cpp
+++ b/Engine/Source/Runtime/AssetRegistry/Private/AssetDataGatherer.cpp
@@ -752,6 +752,10 @@ void FScanDir::Update(FScanDir*& OutCursor, FInherited& InOutParentData)
 		if (SubDirToScan && (!ThisCursor || SubDirToScan->AccumulatedPriority <= ThisCursor->AccumulatedPriority))
 		{
 			FInherited ThisCursorParentData(InOutParentData, DirectData);
+#if DO_CHECK
+			int32 LoopCount = 0;
+			const int32 MaxLoopCount = SubDirs.Num();
+#endif
 			for (;;) // Continue Updating SubDirs until we find one that needs scanning or all are complete or lower in priority than this
 			{
 				FInherited SubDirsParentData(ThisCursorParentData);
@@ -762,10 +766,15 @@ void FScanDir::Update(FScanDir*& OutCursor, FInherited& InOutParentData)
 					return;
 				}
 				// Otherwise the SubDir finished updating and we should move to the next one
-				check(SubDirToScan->IsComplete());
+				// The old value of SubDirToScan may have been deallocated; do not access it further.
+
 				SubDirToScan = FindHighestPrioritySubDir();
 				if (SubDirToScan && (!ThisCursor || SubDirToScan->AccumulatedPriority <= ThisCursor->AccumulatedPriority))
 				{
+					checkf(LoopCount++ < MaxLoopCount,
+						TEXT("Infinite loop: SubDirs return completion without actually completing. SubDirToScan=%s"),
+						*SubDirToScan->GetLocalAbsPath());
+
 					continue;
 				}
 				else
