AssetDataGatherer: comments and line length update.
#rnx

diff --git a/Engine/Source/Runtime/AssetRegistry/Private/AssetDataGatherer.cpp b/Engine/Source/Runtime/AssetRegistry/Private/AssetDataGatherer.cpp
--- a/Engine/Source/Runtime/AssetRegistry/Private/AssetDataGatherer.cpp
+++ b/Engine/Source/Runtime/AssetRegistry/Private/AssetDataGatherer.cpp
@@ -435,7 +435,8 @@ FScanDir* FScanDir::GetControllingDir(FStringView InRelPath, bool bIsDirectory,
 	}
 }
 
-bool FScanDir::TrySetDirectoryProperties(FStringView InRelPath, const FSetPathProperties& InProperties, bool bConfirmedExists)
+bool FScanDir::TrySetDirectoryProperties(FStringView InRelPath,
+	const FSetPathProperties& InProperties, bool bConfirmedExists)
 {
 	// TrySetDirectoryProperties can only be called on valid ScanDirs, which we rely on so we can call FindOrAddSubDir which requires that
 	check(IsValid()); 
@@ -2045,17 +2046,19 @@ bool FAssetDataDiscovery::IsMonitored(FStringView LocalAbsPath) const
 	return MountDir && MountDir->IsMonitored(LocalAbsPath);
 }
 
-SIZE_T FAssetDataDiscovery::GetAllocatedSize() const
+template <typename ContainerType>
+SIZE_T GetArrayRecursiveAllocatedSize(const ContainerType& Container)
 {
-	auto GetArrayRecursiveAllocatedSize = [](auto Container)
+	SIZE_T Result = Container.GetAllocatedSize();
+	for (const auto& Value : Container)
 	{
-		SIZE_T Result = Container.GetAllocatedSize();
-		for (const auto& Value : Container)
-		{
-			Result += Value.GetAllocatedSize();
-		}
-		return Result;
-	};
+		Result += Value.GetAllocatedSize();
+	}
+	return Result;
+};
+
+SIZE_T FAssetDataDiscovery::GetAllocatedSize() const
+{
 
 	FScopedPause ScopedPause(*this);
 	CHECK_IS_NOT_LOCKED_CURRENT_THREAD(TreeLock);
diff --git a/Engine/Source/Runtime/AssetRegistry/Private/AssetDataGatherer.h b/Engine/Source/Runtime/AssetRegistry/Private/AssetDataGatherer.h
--- a/Engine/Source/Runtime/AssetRegistry/Private/AssetDataGatherer.h
+++ b/Engine/Source/Runtime/AssetRegistry/Private/AssetDataGatherer.h
@@ -58,13 +58,14 @@ typedef FScopeLock FGathererScopeLock;
 class FAssetDataGatherer : public FRunnable
 {
 public:
-	FAssetDataGatherer(const TArray<FString>& InLongPackageNamesDenyList, const TArray<FString>& InMountRelativePathsDenyList, bool bInIsSynchronous);
+	FAssetDataGatherer(const TArray<FString>& InLongPackageNamesDenyList,
+		const TArray<FString>& InMountRelativePathsDenyList, bool bInIsSynchronous);
 	virtual ~FAssetDataGatherer();
 
 
 	// Extra at-construction configuration 
 
-	/** Configure the gatherer to use a single monolithic cache for all files, and read/write this cache during ticks. */
+	/** Configure the gatherer to use a single monolithic cache, and read/write this cache during ticks. */
 	void ActivateMonolithicCache();
 
 
@@ -93,7 +94,11 @@ public:
 		TRingBuffer<FString> CookedPackageNamesWithoutAssetData;
 		TRingBuffer<FName> VerseFiles;
 
-		SIZE_T GetAllocatedSize() const { return Assets.GetAllocatedSize() + Paths.GetAllocatedSize() + Dependencies.GetAllocatedSize() + CookedPackageNamesWithoutAssetData.GetAllocatedSize() + VerseFiles.GetAllocatedSize(); }
+		SIZE_T GetAllocatedSize() const
+		{
+			return Assets.GetAllocatedSize() + Paths.GetAllocatedSize() + Dependencies.GetAllocatedSize() +
+				CookedPackageNamesWithoutAssetData.GetAllocatedSize() + VerseFiles.GetAllocatedSize();
+		}
 		void Shrink()
 		{
 			Assets.Shrink();
@@ -115,8 +120,12 @@ public:
 	/** Gets search results from the data gatherer. */
 	void GetAndTrimSearchResults(FResults& InOutResults, FResultContext& OutContext);
 	/** Gets just the AssetResults and DependencyResults from the data gatherer. */
-	void GetPackageResults(TMultiMap<FName, FAssetData*>& OutAssetResults, TMultiMap<FName, FPackageDependencyData>& OutDependencyResults);
-	/** Wait for all monitored assets under the given path to be added to search results. Returns immediately if the given path is not monitored. */
+	void GetPackageResults(TMultiMap<FName, FAssetData*>& OutAssetResults,
+		TMultiMap<FName, FPackageDependencyData>& OutDependencyResults);
+	/**
+	 * Wait for all monitored assets under the given path to be added to search results.
+	 * Returns immediately if the given path is not monitored.
+	 */
 	void WaitOnPath(FStringView LocalPath);
 	/**
 	 * Empty the cache read from disk and the cache used to write to disk. Disable further caching.
@@ -136,7 +145,8 @@ public:
 	/** Wait for all monitored assets to be added to search results. */
 	void WaitForIdle();
 	/**
-	 * Report whether all monitored assets have been added to search results, AND these results have been gathered through a GetAndTrimSearchResults call.
+	 * Report whether all monitored assets have been added to search results, AND these results have been gathered
+	 * through a GetAndTrimSearchResults call.
 	 * This function can be used to check whether there is any work to be done on the gatherer.
 	 */
 	bool IsComplete() const;
@@ -144,7 +154,7 @@ public:
 
 	// Reading/writing/triggering depot-wide properties and events (possibly while tick is running)
 
-	/** Inform the gatherer that initial plugins have finished loading and it should not retry failed loads in hope of a missing custom version */
+	/** Set after initial plugins have loaded and we should not retry failed loads with missing custom versions. */
 	void SetInitialPluginsLoaded();
 	/** Report whether the gatherer is configured to load depends data in addition to asset data. */
 	bool IsGatheringDependencies() const;
@@ -156,25 +166,25 @@ public:
 	FString GetCacheFilename(TConstArrayView<FString> CacheFilePackagePaths);
 	/** Attempt to read the cache file at the given LocalPath, and store all of its results in the in-memory cache. */
 	void LoadCacheFile(FStringView CacheFilename);
-	/** Return the total amount of heap memory used by the gatherer (including not-yet-claimed search results). Used for performance metrics. */
+	/** Return the memory used by the gatherer. Used for performance metrics. */
 	SIZE_T GetAllocatedSize() const;
 
 
 	// Configuring mount points (possibly while tick is running)
 
-	/** Add a mountpoint to the gatherer. This should be the localpath to the packagepath/localpath pair that has been registered with FPackageName .*/
+	/** Add a mountpoint to the gatherer after it has been registered with FPackageName .*/
 	void AddMountPoint(FStringView LocalPath, FStringView LongPackageName);
 	/** Remove a previously added mountpoint. */
 	void RemoveMountPoint(FStringView LocalPath);
-	/** For each LocalPath in LocalPaths, query FPackageName for its mount point, and add this mountpoint to the gatherer. */
+	/** Add MountPoints in LocalPaths to the gatherer. */
 	void AddRequiredMountPoints(TArrayView<FString> LocalPaths);
 
 
 	// Reading/Writing properties of files and directories (possibly while tick is running)
 
-	/** Called from DirectoryWatcher. Update the directory in internal structures and report it as necessary in a future search result. */
+	/** Called from DirectoryWatcher. Update the directory for reporting in future search results. */
 	void OnDirectoryCreated(FStringView LocalPath);
-	/** Called from DirectoryWatcher. Update the files in internal structures and report them as necessary in a future search result. */
+	/** Called from DirectoryWatcher. Update the files for reporting in future search results. */
 	void OnFilesCreated(TConstArrayView<FString> LocalPaths);
 	/** Mark a file or directory to be scanned before unprioritized assets. */
 	void PrioritizeSearchPath(const FString& PathToPrioritize);
@@ -184,13 +194,13 @@ public:
 	 * By default no paths are scanned; adding a path to the allow list causes it and its subdirectories to be scanned.
 	 * Note that the deny list (InLongPackageNameDenyList) overrides the allow list.
 	 * Allow list settings are recursive. Attempting to mark a path as allowed if a parent path is on the allow list
-	 * will have no effect. This means the scenario (1) add allow list A (2) add allow list A/Child (3) remove allow list A will
-	 * therefore not result in A/Child being allowed.
+	 * will have no effect. This means the scenario ((1) add allow list A (2) add allow list A/Child (3) remove allow
+	 * list A) will therefore not result in A/Child being allowed.
 	 */
 	void SetIsOnAllowList(FStringView LocalPath, bool bIsAllowed);
-	/** Report whether the path is in the allow list. If not it will not be scanned. If so it will be scanned unless it is in the deny list. */
+	/** Report whether the path is in the allow list. Only paths in AllowList AND not in DenyList will be scanned. */
 	bool IsOnAllowList(FStringView LocalPath) const;
-	/** Report whether the path is in the deny list. If so it will not be scanned regardless of whether it is in the allow list. */
+	/** Report whether the path is in the deny list. Paths in DenyList are not scanned. */
 	bool IsOnDenyList(FStringView LocalPath) const;
 	/** Report whether the path is both in the allow list and not in the deny list. */
 	bool IsMonitored(FStringView LocalPath) const;
@@ -219,14 +229,14 @@ private:
 	 */
 	void InnerTickLoop(bool bInIsSynchronousTick, bool bContributeToCacheSave);
 	/**
-	 * Tick function to pump scanning and push results into the search results structure. May be called from devoted thread
-	 * or inline from synchronous functions on other threads.
+	 * Tick function to pump scanning and push results into the search results structure. May be called from devoted
+	 * thread or inline from synchronous functions on other threads.
 	 */
 	void TickInternal(bool& bOutIsIdle);
 	/** Add any new package files from the background directory scan to our work list **/
 	void IngestDiscoveryResults();
 
-	/** Helper for OnFilesCreated. Update the file in internal structures and report it as necessary in a future search result. */
+	/** Helper for OnFilesCreated. Update the file for reporting in future search results. */
 	void OnFileCreated(FStringView LocalPath);
 
 	/**
@@ -265,8 +275,8 @@ private:
 	/** Add the given AssetDatas into DiskCachedAssetDataMap and DiskCachedAssetBlocks. */
 	void ConsumeCacheFile(UE::AssetDataGather::Private::FCachePayload&& Payload);
 	/**
-	 * If a save of the monolithic cache has been triggered, get the cache filename and pointers to all elements that should be saved,
-	 * for later saving outside of the critical section.
+	 * If a save of the monolithic cache has been triggered, get the cache filename and pointers to all elements that
+	 * should be saved, for later saving outside of the critical section.
 	 */
 	void TryReserveSaveMonolithicCache(bool& bOutShouldSave, TArray<TPair<FName,FDiskCachedAssetData*>>& AssetsToSave);
 	/**
@@ -274,19 +284,21 @@ private:
 	 * This function reads the read-only-after-creation data from each FDiskCachedAssetData*, but otherwise does not use
 	 * data from this Gatherer and so can be run outside any critical section.
 	 */
-	void SaveCacheFileInternal(const FString& CacheFilename, const TArray<TPair<FName,FDiskCachedAssetData*>>& AssetsToSave, bool bIsAsyncCacheSave);
+	void SaveCacheFileInternal(const FString& CacheFilename,
+		const TArray<TPair<FName,FDiskCachedAssetData*>>& AssetsToSave, bool bIsAsyncCacheSave);
 	/**
 	 * Get the list of FDiskCachedAssetData* that have been loaded in the gatherer, for saving into a cachefile.
 	 * Filters the list of assets by child paths of the elements in SaveCacheLongPackageNameDirs, if it is non-empty.
 	 */
-	void GetAssetsToSave(TArrayView<const FString> SaveCacheLongPackageNameDirs, TArray<TPair<FName,FDiskCachedAssetData*>>& OutAssetsToSave);
+	void GetAssetsToSave(TArrayView<const FString> SaveCacheLongPackageNameDirs,
+		TArray<TPair<FName,FDiskCachedAssetData*>>& OutAssetsToSave);
 
-	/* Adds the given PackageName,DiskCachedAssetData pair into NewCachedAssetDataMap, and detects collisions for multiple files with the same PackageName */
+	/* Adds the given pair into NewCachedAssetDataMap. Detects collisions for multiple files with the same PackageName */
 	void AddToCache(FName PackageName, FDiskCachedAssetData* DiskCachedAssetData);
 
 	/**
-	 * Mark that the gatherer has become idle or has become active. Called from tick function and configuration functions when they note a possible state change.
-	 * Caller is responsible for holding the ResultsLock.
+	 * Mark that the gatherer has become idle or has become active. Called from tick function and configuration functions
+	 * when they note a possible state change. Caller is responsible for holding the ResultsLock.
 	 */
 	void SetIsIdle(bool IsIdle);
 
@@ -337,21 +349,18 @@ private:
 
 	// Variable section for variables that are atomics read/writable from outside critical sections.
 
-	/** > 0 if we've been asked to abort work in progress at the next opportunity. Read/writable anywhere. */
+	/** > 0 if we've been asked to abort work in progress at the next opportunity. */
 	std::atomic<uint32> IsStopped;
-	/** > 0 if we've been asked to pause the worker thread so a synchronous function can take over the tick. Read/writable anywhere */
+	/** > 0 if we've been asked to pause the worker thread so a synchronous function can take over the tick. */
 	mutable std::atomic<uint32> IsPaused;
 	/**
-	 * Discovery subsystem; this subsystem handles deciding which paths to search and running the FileManager calls to search directories.
-	 * Pointer itself is constant during threading. Object pointed to is read/writable anywhere; internally provides threadsafety.
+	 * Discovery subsystem; decides which paths to search and queries the FileManager to search directories.
+	 * Pointer is constant during threading. Object pointed to internally provides threadsafety.
 	 */
 	TUniquePtr<UE::AssetDataGather::Private::FAssetDataDiscovery> Discovery;
-	/** Async only. Set to true once initial plugins have been loaded. Read/writable anywhere. */
+	/** Async only. Set to true once initial plugins have been loaded. */
 	std::atomic<bool> bInitialPluginsLoaded;
-	/**
-	 * Set to true when periodic or final save of the async cache was detected during TickInternal and should be processed outside of the Tick.
-	 * Read/writable anywhere.
-	 */
+	/** True when TickInternal requests periodic or final save of the async cache. */
 	std::atomic<bool> bSaveAsyncCacheTriggered;
 	/** True if the current process allows reading of AssetDataGatherer cache files. */
 	std::atomic<bool> bCacheReadEnabled;
@@ -360,103 +369,86 @@ private:
 
 	// Variable section for variables that are read/writable only within ResultsLock.
 
-	/** List of files that need to be processed by the search. Read/writable only within ResultsLock. */
+	/** List of files that need to be processed by the search. */
 	TUniquePtr<UE::AssetDataGather::Private::FFilesToSearch> FilesToSearch;
 
-	/** The asset data gathered from the searched files. Read/writable only within ResultsLock. */
+	/** The asset data gathered from the searched files. */
 	TArray<FAssetData*> AssetResults;
-	/** Dependency data gathered from the searched files packages. Read/writable only within ResultsLock. */
+	/** Dependency data gathered from the searched files packages. */
 	TArray<FPackageDependencyData> DependencyResults;
 	/**
 	 * A list of cooked packages that did not have asset data in them.
-	 * These assets may still contain assets (if they were older for example). Read/writable only within ResultsLock.
+	 * These assets may still contain assets (if they were older for example). 
 	 */
 	TArray<FString> CookedPackageNamesWithoutAssetDataResults;
 	/** File paths (in UE LongPackagePath notation) of the Verse source code gathered from the searched files. */
 	TArray<FName> VerseResults;
 
-	/** All the search times since the last call to GetAndTrimSearchResults. Read/writable only within ResultsLock. */
+	/** All the search times since the last call to GetAndTrimSearchResults. */
 	TArray<double> SearchTimes;
-	/**
-	 * The directories found during the search, unless they are hidden by DirLongPackageNamesToNotReport.
-	 * Read/writable only within ResultsLock.
-	 */
+
+	/** The directories found during the search, unless they are hidden by DirLongPackageNamesToNotReport. */
 	TArray<FString> DiscoveredPaths;
 
-	/** The current search start time, set in the first tick after idle and used for performance metrics when reporting results. Read/writable only within ResultsLock. */
+	/** The search start time, set after resuming from idle and used for performance metrics when reporting results. */
 	double SearchStartTime;
-	/** The last time at which the cache file was written, used to periodically update the cache with the results of ticking. Read/writable only within ResultsLock. */
+	/** The last time at which the cache file was written, used to periodically update the cache. */
 	double LastCacheWriteTime;
-	/**
-	 * The cached value of the NumPathsToSearch returned by the Discovery subsystem the last time we synchronized with it.
-	 * Read/writable only within ResultsLock.
-	 */
+	/** The cached value of the NumPathsToSearch returned by Discovery the last time we synchronized with it. */
 	int32 NumPathsToSearchAtLastSyncPoint;
 	/**
 	 * Track whether we are allowed to read from a monolithic cache that should be loaded during tick.
-	 * Even if we are or not, if bCacheReadEnabled the AssetRegistry can also call LoadCacheFile/ScanPathsSynchronous to load/save smaller files.
-	 * Read/writable only within ResultsLock.
+	 * Even if we are or not, if bCacheReadEnabled the AssetRegistry can also call LoadCacheFile/ScanPathsSynchronous to
+	 * load/save smaller files.
 	 */
 	bool bReadMonolithicCache;
-	/**
-	 * Track whether we are allowed to write to the monolithic cache.
-	 * Read/writable only within ResultsLock.
-	 */
+	/** Track whether we are allowed to write to the monolithic cache. */
 	bool bWriteMonolithicCache;
-	/** If bHasLoadedMonolithicCache is true, track whether the cache has been loaded. Read/writable only within ResultsLock. */
+	/** If bHasLoadedMonolithicCache is true, track whether the cache has been loaded. */
 	bool bHasLoadedMonolithicCache;
-	/** Track whether the Discovery subsystem has gone idle and we have read all filenames from it. Read/writable only within ResultsLock. */
+	/** Track whether the Discovery subsystem has gone idle and we have read all filenames from it. */
 	bool bDiscoveryIsComplete;
-	/** Track whether this Gather has gone idle and a caller has read all search data from it. Read/writable only within ResultsLock. */
+	/** Track whether this Gather has gone idle and a caller has read all search data from it. */
 	bool bIsComplete;
-	/** Track whether this Gatherer has gone idle, either because of no more work or because its blocked on external events. Read/writable only within ResultsLock. */
+	/** Track whether this Gatherer has gone idle, either it has no more work or it's blocked on external events. */
 	bool bIsIdle;
-	/** Track the first tick after idle to set up e.g. timing data. Read/writable only within ResultsLock. */
+	/** Track the first tick after idle to set up e.g. timing data. */
 	bool bFirstTickAfterIdle;
-	/** True if we have finished discovering our first wave of files, to report metrics for that most-important wave. Read/writable only within ResultsLock. */
+	/** True if we have finished discovering our first wave of files, to report metrics for that most-important wave. */
 	bool bFinishedInitialDiscovery;
 
 
 	// Variable section for variables that are read/writable only within TickLock.
 
 	/**
-	 * An array of all cached data that was newly discovered this run. This array is just used to make sure they are all deleted at shutdown.
-	 * Read/writable only within TickLock.
+	 * An array of all cached data that was newly discovered this run. This array is just used to make sure they are all
+	 * deleted at shutdown.
 	 */
 	TArray<FDiskCachedAssetData*> NewCachedAssetData;
 	TArray<TPair<int32, FDiskCachedAssetData*>> DiskCachedAssetBlocks;
 	/**
 	 * Map of PackageName to cached discovered assets that were loaded from disk.
 	 * This should only be modified by ConsumeCacheFile.
-	 * Read/Writable only within TickLock.
 	 */
 	TMap<FName, FDiskCachedAssetData*> DiskCachedAssetDataMap;
-	/** Map of PackageName to cached discovered assets that will be written to disk at shutdown. Read/writable only within TickLock. */
+	/** Map of PackageName to cached discovered assets that will be written to disk at shutdown. */
 	TMap<FName, FDiskCachedAssetData*> NewCachedAssetDataMap;
-	/**
-	 * Used when we are blocking on gather results. If non-zero, we should end the tick when WaitBatchCount files have been processed.
-	 * Read/writable only within TickLock.
-	 */
+	/** Used to block on gather results. If non-zero, tick should end when WaitBatchCount files have been processed. */
 	int32 WaitBatchCount;
-	/** How many files in the search results were read from the cache. Read/writable only within TickLock. */
+	/** How many files in the search results were read from the cache. */
 	int32 NumCachedAssetFiles;
-	/** How many files in the search results were not in the cache and were read by parsing the file. Read/writable only within TickLock. */
+	/** How many files in the search results were not in the cache and were read by parsing the file. */
 	int32 NumUncachedAssetFiles;
 	/**
 	 * Incremented when a thread is in the middle of saving any cache and therefore the cache cannot be deleted,
-	 * decremented when the thread is done. Only allowed to be incremented when bCacheEnabled has been recently confirmed to be true.
-	 * Read/writable only within TickLock.
+	 * decremented when the thread is done. Only incremented when bCacheEnabled has been recently confirmed to be true.
 	 */
 	int32 CacheInUseCount;
 	/**
-	 * True if the current TickInternal is synchronous, which may be because bIsSynchronous or because the game thread has taken over the tick for a synchronous function.
-	 * Read/writable only within TickLock.
+	 * True if the current TickInternal is synchronous, which may be because bIsSynchronous or because the game thread has
+	 * taken over the tick for a synchronous function.
 	 */
 	bool bIsSynchronousTick;
-	/**
-	 * Set to true when a thread is in the middle of saving an async cache and so another save of the cache should not be triggered.
-	 * Read/writable only within TickLock.
-	 */
+	/** True when a thread is saving an async cache and so another save of the cache should not be triggered. */
 	bool bIsSavingAsyncCache;
-
 };
diff --git a/Engine/Source/Runtime/AssetRegistry/Private/AssetDataGathererPrivate.h b/Engine/Source/Runtime/AssetRegistry/Private/AssetDataGathererPrivate.h
--- a/Engine/Source/Runtime/AssetRegistry/Private/AssetDataGathererPrivate.h
+++ b/Engine/Source/Runtime/AssetRegistry/Private/AssetDataGathererPrivate.h
@@ -39,7 +39,7 @@ constexpr uint32 CountEPriority = static_cast<uint32>(EPriority::Normal) + 1;
 /** Fields being set in a call to SetProperty */
 struct FSetPathProperties
 {
-	/** The path (usually a plugin's root content path) has been requested for scanning through e.g. ScanPathsSynchronous */
+	/** The path (usually a plugin's root content path) was requested for scanning through e.g. ScanPathsSynchronous */
 	TOptional<bool> IsOnAllowList;
 	/**
 	 * Set whether the given directory matches a deny list entry.
@@ -118,10 +118,12 @@ struct FGatheredPathData
 	EGatherableFileType Type = EGatherableFileType::Invalid;
 
 	FGatheredPathData() = default;
-	FGatheredPathData(FStringView InLocalAbsPath, FStringView InLongPackageName, const FDateTime& InPackageTimestamp, EGatherableFileType InType);
+	FGatheredPathData(FStringView InLocalAbsPath, FStringView InLongPackageName, const FDateTime& InPackageTimestamp,
+		EGatherableFileType InType);
 	explicit FGatheredPathData(const FDiscoveredPathData& DiscoveredData);
 	explicit FGatheredPathData(FDiscoveredPathData&& DiscoveredData);
-	void Assign(FStringView InLocalAbsPath, FStringView InLongPackageName, const FDateTime& InPackageTimestamp, EGatherableFileType InType);
+	void Assign(FStringView InLocalAbsPath, FStringView InLongPackageName, const FDateTime& InPackageTimestamp,
+		EGatherableFileType InType);
 	void Assign(const FDiscoveredPathData& DiscoveredData);
 
 	/**
@@ -197,7 +199,7 @@ private:
 		/** Add the given FilePath as a direct file child of the Node's directory. */
 		void AddFile(FGatheredPathData&& FilePath);
 
-		/** Pop the given number of files out of this node and its children and decrement NumToPop by how many were popped. */
+		/** Pop the NumToPop files out of this node and its children and decrement NumToPop by how many were popped. */
 		template <typename RangeType>
 		void PopFiles(RangeType& Out, int32& NumToPop);
 		/** Pop all files out of this node and its children. */
@@ -236,7 +238,7 @@ private:
 };
 
 
-/** Stores a LocalAbsPath and its existence information. A file system query is issued the first time the data is needed. */
+/** Stores a LocalAbsPath and existence information. A file system query is issued the first time the data is needed. */
 struct FPathExistence
 {
 	/** What kind of thing we found at the given path */
@@ -271,7 +273,7 @@ private:
  * referenced from the SubDirs array. Directories are removed from the tree once their scans are finished to save memory.
  * Queries take into account that deleted nodes have been completed.
  *
- * This class is not ThreadSafe; The FAssetDataDiscovery reads/writes its data only while holding the TreeLock critical section.
+ * This class is not ThreadSafe; The FAssetDataDiscovery reads/writes its data only while holding TreeLock.
  */
 class FScanDir : public FRefCountBase
 {
@@ -304,7 +306,7 @@ public:
 	 */
 	void Shutdown();
 
-	/** Check whether this ScanDir is still alive or if it has been marked for destruction and cleared on another thread. */
+	/** Check whether this ScanDir is alive; it may have been marked for destruction and cleared on another thread. */
 	bool IsValid() const;
 
 	FMountDir* GetMountDir() const;
@@ -312,7 +314,7 @@ public:
 	/** Get this ScanDir's RelPath from its Parent */
 	FStringView GetRelPath() const;
 
-	/** Return this ScanDir's accumulated Priority: the maximum of its directly-set priority and its children's priority. */
+	/** Return this ScanDir's accumulated Priority: maximum of its directly-set priority and its children's priority. */
 	EPriority GetPriority() const;
 
 	/** Calculate this ScanDir's full absolute path by accumulating RelPaths from parents and append it. */
@@ -350,11 +352,12 @@ public:
 		FInherited& OutData, FString& OutRelPath);
 
 	/**
-	 * Set values of fields on the given directory indicated by InRelPath for all of the properties existing on InProperties.
+	 * Set values of fields on the given directory indicated by InRelPath for all properties existing on InProperties.
 	 * Returns whether the directory was found and its property was changed; returns false if InRelPath was not a
 	 * directory or the property did not need to be changed.
 	 */
-	bool TrySetDirectoryProperties(FStringView InRelPath, const FSetPathProperties& InProperties, bool bConfirmedExists);
+	bool TrySetDirectoryProperties(FStringView InRelPath,
+		const FSetPathProperties& InProperties, bool bConfirmedExists);
 	/**
 	 * Mark that the given file has already been scanned, so that it will not be double reporting in the upcoming
 	 * directory scan, if one is upcoming.
@@ -368,8 +371,8 @@ public:
 	 * Called from the Tick; handle the list of subdirs and files that were found from IterateDirectoryStat called
 	 * on this ScanDir, reporting discovered files and updating status variables.
 	 */
-	void SetScanResults(FStringView LocalAbsPath, const FInherited& ParentData, TArrayView<FDiscoveredPathData>& InOutSubDirs,
-		TArrayView<FDiscoveredPathData>& InOutFiles);
+	void SetScanResults(FStringView LocalAbsPath, const FInherited& ParentData,
+		TArrayView<FDiscoveredPathData>& InOutSubDirs, TArrayView<FDiscoveredPathData>& InOutFiles);
 	/**
 	 * Update the completion state of this ScanDir based on its scan status and its child dirs completion state.
 	 * Set OutCursor to the highest priority child dir that needs to be updated if any children need to be updated
@@ -441,7 +444,7 @@ protected:
 	 */
 	template <typename CallbackType> void ForEachDescendent(const CallbackType& Callback);
 
-	/** Set the AccumulatedPriority of this ScanDir to the maximum of the DirectPriorities of it and all its child dirs. */
+	/** Set the AccumulatedPriority of this ScanDir to the maximum of the DirectPriorities of it and its child dirs. */
 	void UpdateAccumulatedPriority();
 	/** Record a changed priority of a direct child ScanDir and update the accumulated priority of this ScanDir. */
 	void OnChildPriorityChanged(EPriority InPriority, int32 Delta);
@@ -465,10 +468,11 @@ protected:
 /**
  * Gather data about a MountPoint that has been registered with FPackageName
  * The FMountDir holds a FScanTree with information about each directory (that is pruned when not in use).
- * It also holds some data that is needed only per MountPoint, such as the packagename.
- * It also holds some data per subdirectory that is more performant to hold in a map rather than to require the FScanTrees to be kept.
+ * It also holds data that is needed only per MountPoint, such as the packagename.
+ * It also holds data per subdirectory that is more performant to hold in a map rather than to require the FScanTrees to
+ * be kept.
  *
- * This class is not ThreadSafe; The FAssetDataDiscovery reads/writes its data only while holding the TreeLock critical section.
+ * This class is not ThreadSafe; The FAssetDataDiscovery reads/writes its data only while holding TreeLock.
  */
 class FMountDir
 {
@@ -485,7 +489,8 @@ public:
 	FAssetDataDiscovery& GetDiscovery() const;
 
 	/** Find the direct parent of InRelPath, or the lowest fallback. See FScanDir::GetControllingDir. */
-	FScanDir* GetControllingDir(FStringView LocalAbsPath, bool bIsDirectory, FScanDir::FInherited& OutParentData, FString& OutRelPath);
+	FScanDir* GetControllingDir(FStringView LocalAbsPath, bool bIsDirectory, FScanDir::FInherited& OutParentData,
+		FString& OutRelPath);
 	/** Return the memory used by the tree under this MountDir, except that sizeof(*this) is excluded. */
 	SIZE_T GetAllocatedSize() const;
 
@@ -507,7 +512,8 @@ public:
 	 * Returns whether the directory was foundand its property was changed; returns false if LocalAbsPath was not a
 	 * directory under this MountDir or the property did not need to be changed.
 	 */
-	bool TrySetDirectoryProperties(FStringView LocalAbsPath, const FSetPathProperties& InProperties, bool bConfirmedExists);
+	bool TrySetDirectoryProperties(FStringView LocalAbsPath, const FSetPathProperties& InProperties,
+		bool bConfirmedExists);
 	/** 
 	 * Update all ScanDirs under this MountDir in priority order until one is found that needs to be scanned.
 	 * Set OutCursor to that ScanDir, or to nullptr if all are complete.
@@ -543,7 +549,7 @@ public:
 	TArray<FMountDir*> GetChildMounts() const;
 
 protected:
-	/** Inspect the Discovery's deny lists and add the ones applicable to this MountDir into this MountDir's set of deny lists. */
+	/** Inspect the Discovery's DenyLists and add the ones applicable to this into this MountDir's set of DenyLists. */
 	void UpdateDenyList();
 	/** Mark that given path needs to be reconsidered by Update. */
 	void MarkDirty(FStringView MountRelPath);
@@ -619,7 +625,7 @@ public:
 	bool IsOnAllowList(FStringView LocalAbsPath) const;
 	/** Return whether the given path matches the deny list and has not been marked IgnoreDenyList. */
 	bool IsOnDenyList(FStringView LocalAbsPath) const;
-	/** Return whether the given path should or has been scanned because it is on the allow list and not on the deny list. */
+	/** Return whether the path should or has been scanned because it is on the AllowList and not on the DenyList. */
 	bool IsMonitored(FStringView LocalAbsPath) const;
 	/** Return the memory used by *this. sizeof(*this) is not included. */
 	SIZE_T GetAllocatedSize() const;
@@ -664,7 +670,8 @@ private:
 	void SetIsIdle(bool bInIdle);
 
 	/** Store the given discovered files and directories in the results. */
-	void AddDiscovered(FStringView DirAbsPath, TConstArrayView<FDiscoveredPathData> SubDirs, TConstArrayView<FDiscoveredPathData> Files);
+	void AddDiscovered(FStringView DirAbsPath, TConstArrayView<FDiscoveredPathData> SubDirs,
+		TConstArrayView<FDiscoveredPathData> Files);
 	/** Store the given specially reported single file in the results. */
 	void AddDiscoveredFile(FDiscoveredPathData&& File);
 
@@ -673,7 +680,7 @@ private:
 
 	/**
 	 * Return whether a directory with the given LongPackageName should be reported to the AssetRegistry
-	 * We do not report some directories because they are paths that should not enter the AssetRegistry list of paths if empty,
+	 * We do not report some directories because they should not enter the AssetRegistry list of paths if empty,
 	 * and reporting a path to the AssetRegistry adds it unconditionally to the list of paths.
 	 * If ShouldDirBeReported returns false, the directory will still be added to the catalog if non-empty,
 	 * because the AssetRegistry adds the path of every added file.
@@ -723,21 +730,15 @@ private:
 
 	// Variable section for variables that are constant during threading.
 
-	/**
-	 * Deny list of full absolute paths. Child paths will not be scanned unless requested to ignore deny lists.
-	 * Constant during threading.
-	 */
+	/** Deny list of full absolute paths. Child paths will not be scanned unless requested to ignore DenyLists. */
 	TArray<FString> LongPackageNamesDenyList;
-	/**
-	 * Deny list of relative paths in each mount. Child paths will not be scanned unless requested to ignore deny lists.
-	 * Constant during threading.
-	 */
+	/** DenyList of relative paths in each mount. Child paths will not be scanned unless requested to ignore DenyLists. */
 	TArray<FString> MountRelativePathsDenyList;
-	/** LongPackageNames for directories that should not be reported, see ShouldDirBeReported. Constant during threading. */
+	/** LongPackageNames for directories that should not be reported, see ShouldDirBeReported. */
 	TSet<FString> DirLongPackageNamesToNotReport;
-	/** Thread to run the discovery FRunnable on. Read-only while threading is possible. Constant during threading. */
+	/** Thread to run the discovery FRunnable on. Read-only while threading is possible. */
 	FRunnableThread* Thread;
-	/** True if this gather request is synchronous (i.e, IsRunningCommandlet()). Constant during threading. */
+	/** True if this gather request is synchronous (i.e, IsRunningCommandlet()). */
 	bool bIsSynchronous;
 
 
@@ -748,26 +749,23 @@ private:
 	 * Readable anywhere. Writable only within TreeLock.
 	 */
 	std::atomic<bool> bIsIdle;
-	/** > 0 if we've been asked to abort work in progress at the next opportunity. Read/writable anywhere. */
+	/** > 0 if we've been asked to abort work in progress at the next opportunity. */
 	std::atomic<uint32> IsStopped;
-	/**
-	 * > 0 if we've been asked to pause the worker thread so a synchronous function can take over the tick.
-	 * Read/writable anywhere.
-	 */
+	/** > 0 if we've been asked to pause the worker thread so a synchronous function can take over the tick. */
 	mutable std::atomic<uint32> IsPaused;
 	/**
 	 * Number of directories that have been discovered and IsMonitored but have not yet been scanned.
-	 * Used for progress tracking. Read/writable anywhere.
+	 * Used for progress tracking.
 	 */
 	FThreadSafeCounter NumDirectoriesToScan;
 
 
 	// Variable section for variables that are read/writable only within ResultsLock.
 
-	/** Directories found in the scan; may be empty. Read/writable only within ResultsLock. */
+	/** Directories found in the scan; may be empty. */
 	TArray<FString> DiscoveredDirectories;
 
-	/** Files found found in the scan. Read/writable only within ResultsLock. */
+	/** Files found found in the scan. */
 	struct FDirectoryResult
 	{
 		FDirectoryResult(FStringView InDirAbsPath, TConstArrayView<FDiscoveredPathData> InFiles);
@@ -777,37 +775,31 @@ private:
 	};
 	TArray<FDirectoryResult> DiscoveredFiles;
 	TArray<FGatheredPathData> DiscoveredSingleFiles;
-	/**
-	 * Time at which the scan was started or last resumed from idle. Used for logging.
-	 * Read/writable only within ResultsLock.
-	 */
+	/** Time at which the scan was started or last resumed from idle. Used for logging. */
 	double DiscoverStartTime = 0.;
-	/**
-	 * Number of files discovered during scanning since start or resumed from idle. Used for logging.
-	 * Read/writable only within ResultsLock.
-	 */
+	/** Number of files discovered during scanning since start or resumed from idle. Used for logging. */
 	int32 NumDiscoveredFiles = 0;
 
 
 	// Variable section for variables that are read/writable only within TreeLock.
 
 	/**
-	 * Sorted list of MountDirs, sorted by FPackagePath::Less on the absolute paths. Read/writable only within TreeLock.
+	 * Sorted list of MountDirs, sorted by FPackagePath::Less on the absolute paths.
 	 * Each MountDir contains a ScanDir tree and other data that configures the scanning within that MountPoint.
 	 * Read/writable only with TreeLock, both the list and all data owned by each MountDir.
 	 */
 	TArray<TUniquePtr<FMountDir>> MountDirs;
-	/** The next ScanDir to update in Tick. Read/writable only within TreeLock. */
+	/** The next ScanDir to update in Tick. */
 	TRefCountPtr<FScanDir> Cursor = nullptr;
-	/** The value of Parent's InheritedData for Cursor (or false if Cursor is null). Read/writable only within TreeLock. */
+	/** The value of Parent's InheritedData for Cursor (or false if Cursor is null). */
 	FScanDir::FInherited CursorParentData;
 
 
 	// Variable section for variables that are read/writable only within TickLock.
 
-	/** Scratch space to store discovered subdirs during the tick, to avoid allocations. Read/writable only within TickLock. */
+	/** Scratch space to store discovered subdirs during the tick, to avoid allocations. */
 	TArray<FDiscoveredPathData> IteratedSubDirs;
-	/** Scratch space to store discovered files during the tick, to avoid allocations. Read/writable only within TickLock. */
+	/** Scratch space to store discovered files during the tick, to avoid allocations. */
 	TArray<FDiscoveredPathData> IteratedFiles;
 
 
