AssetDataDiscovery: Optimize TickInternal using a ParallelFor.
To support parallelization, all known directories to scan are scanned in a single ParallelFor. The discovered subdirectories from the scan are then scanned in the next ParallelFor until no more are found.
To support the scan of all known directories to scan, the concept of a single Cursor was removed, and the implementation of Priority for blocking scans was changed from a flag on each ScanDir to a global list of the prioritized ScanDirs.
Improved the diagnostic logs for AssetRegistry SearchTimes.
#rnx

diff --git a/Engine/Source/Runtime/AssetRegistry/Private/AssetDataGatherer.cpp b/Engine/Source/Runtime/AssetRegistry/Private/AssetDataGatherer.cpp
--- a/Engine/Source/Runtime/AssetRegistry/Private/AssetDataGatherer.cpp
+++ b/Engine/Source/Runtime/AssetRegistry/Private/AssetDataGatherer.cpp
@@ -5,6 +5,8 @@
 
 #include "Algo/AnyOf.h"
 #include "Algo/Find.h"
+#include "Algo/Sort.h"
+#include "Algo/Unique.h"
 #include "AssetRegistry/AssetData.h"
 #include "AssetRegistryArchive.h"
 #include "AssetRegistryPrivate.h"
@@ -41,6 +43,17 @@ namespace AssetDataGathererConstants
 	constexpr int32 MinSecondsToElapseBeforeCacheWrite = 60;
 	static constexpr uint32 CacheSerializationMagic = 0x3339D87B; // Versioning and integrity checking
 	static constexpr uint64 CurrentVersion = FAssetRegistryVersion::LatestVersion | (uint64(CacheSerializationMagic) << 32);
+	/**
+	 * The maximum number of threads used in the ParallelFor in FAssetDataDiscovery::TickInternal.
+	 * Higher values cause more time spent waking up workers. Lower values cause more time spent running the DirectoryEnumeration on each thread.
+	 * Current value was chosen based on experiments on an SSD drive with a warm disk cache.
+	 */
+	static int32 GARDiscoverThreads = 15;
+	/**
+	 * Minimum batch sized used used in the ParallelFor in FAssetDataDiscovery::TickInternal.
+	 * In theory higher minimum values could reduce the cost of waking up workers, but in practice it does not have a significant effect.
+	 */
+	static int32 GARDiscoverMinBatchSize = 1;
 }
 
 namespace UE::AssetDataGather::Private
@@ -197,15 +210,6 @@ void FScanDir::Shutdown()
 		MountDir->GetDiscovery().NumDirectoriesToScan.Decrement();
 	}
 
-	// Update Parent data that we influence
-	if (Parent) // Root ScanDir has no parent
-	{
-		if (AccumulatedPriority != EPriority::Normal)
-		{
-			Parent->OnChildPriorityChanged(AccumulatedPriority, -1);
-		}
-	}
-
 	// Clear backpointers (which also marks us as shutdown)
 	MountDir = nullptr;
 	Parent = nullptr;
@@ -226,11 +230,6 @@ FStringView FScanDir::GetRelPath() const
 	return RelPath;
 }
 
-EPriority FScanDir::GetPriority() const
-{
-	return AccumulatedPriority;
-}
-
 void FScanDir::AppendLocalAbsPath(FStringBuilderBase& OutFullPath) const
 {
 	if (!MountDir)
@@ -435,18 +434,18 @@ FScanDir* FScanDir::GetControllingDir(FStringView InRelPath, bool bIsDirectory,
 	}
 }
 
-bool FScanDir::TrySetDirectoryProperties(FStringView InRelPath,
-	const FSetPathProperties& InProperties, bool bConfirmedExists)
+bool FScanDir::TrySetDirectoryProperties(FStringView InRelPath, FInherited& ParentData,
+	const FSetPathProperties& InProperties, bool bConfirmedExists, FScanDirAndParentData& OutControllingDir)
 {
 	// TrySetDirectoryProperties can only be called on valid ScanDirs, which we rely on so we can call FindOrAddSubDir which requires that
 	check(IsValid()); 
 
-	SetComplete(false);
 	if (InRelPath.IsEmpty())
 	{
 		// The properties apply to this entire directory
 		if (InProperties.IsOnAllowList.IsSet() && DirectData.bIsOnAllowList != *InProperties.IsOnAllowList)
 		{
+			SetComplete(false);
 			if (bScanInFlight)
 			{
 				bScanInFlightInvalidated = true;
@@ -479,6 +478,7 @@ bool FScanDir::TrySetDirectoryProperties(FStringView InRelPath,
 		if ((InProperties.MatchesDenyList.IsSet() && DirectData.bMatchesDenyList != *InProperties.MatchesDenyList) ||
 			(InProperties.IgnoreDenyList.IsSet() && DirectData.bIgnoreDenyList != *InProperties.IgnoreDenyList))
 		{
+			SetComplete(false);
 			if (InProperties.MatchesDenyList.IsSet())
 			{
 				DirectData.bMatchesDenyList = *InProperties.MatchesDenyList;
@@ -513,6 +513,7 @@ bool FScanDir::TrySetDirectoryProperties(FStringView InRelPath,
 		}
 		if (InProperties.HasScanned.IsSet())
 		{
+			SetComplete(false);
 			bool bNewValue = *InProperties.HasScanned;
 			auto SetProperties = [bNewValue](FScanDir& ScanDir)
 			{
@@ -526,10 +527,9 @@ bool FScanDir::TrySetDirectoryProperties(FStringView InRelPath,
 			SetProperties(*this);
 			ForEachDescendent(SetProperties);
 		}
-		if (InProperties.Priority.IsSet() && DirectPriority != *InProperties.Priority)
-		{
-			SetDirectPriority(*InProperties.Priority);
-		}
+
+		OutControllingDir.ScanDir = this;
+		OutControllingDir.ParentData = ParentData;
 		return true;
 	}
 	else
@@ -543,6 +543,8 @@ bool FScanDir::TrySetDirectoryProperties(FStringView InRelPath,
 			ModifiedProperties->IsOnAllowList.Reset();
 			if (!ModifiedProperties->IsSet())
 			{
+				OutControllingDir.ScanDir = this;
+				OutControllingDir.ParentData = ParentData;
 				return false;
 			}
 			Properties = &ModifiedProperties.GetValue();
@@ -565,6 +567,8 @@ bool FScanDir::TrySetDirectoryProperties(FStringView InRelPath,
 			SubDir = FindSubDir(FirstComponent);
 			if (!SubDir)
 			{
+				OutControllingDir.ScanDir = this;
+				OutControllingDir.ParentData = ParentData;
 				return false;
 			}
 		}
@@ -583,14 +587,25 @@ bool FScanDir::TrySetDirectoryProperties(FStringView InRelPath,
 					{
 						UE_LOG(LogAssetRegistry, Warning, TEXT("SetDirectoryProperties called on %s path %.*s. Ignoring the call."),
 							StatData.bIsValid ? TEXT("file") : TEXT("non-existent"), LocalAbsPath.Len(), LocalAbsPath.GetData());
+						OutControllingDir.ScanDir = this;
+						OutControllingDir.ParentData = ParentData;
 						return false;
 					}
 					bConfirmedExists = true;
 				}
 				SubDir = &FindOrAddSubDir(FirstComponent);
+				SetComplete(false);
 			}
 		}
-		return SubDir->TrySetDirectoryProperties(Remainder, *Properties, bConfirmedExists);
+
+		FInherited ParentDataForSubDir(ParentData, DirectData);
+		bool bResult = SubDir->TrySetDirectoryProperties(Remainder, ParentDataForSubDir, *Properties, bConfirmedExists,
+			OutControllingDir);
+		if (OutControllingDir.ScanDir && !OutControllingDir.ScanDir->IsComplete())
+		{
+			SetComplete(false);
+		}
+		return bResult;
 	}
 }
 
@@ -611,68 +626,6 @@ void FScanDir::MarkFileAlreadyScanned(FStringView BaseName)
 	AlreadyScannedFiles.Emplace(BaseName);
 }
 
-EPriority FScanDir::GetDirectPriority() const
-{
-	return DirectPriority;
-}
-
-void FScanDir::SetDirectPriority(EPriority InPriority)
-{
-	DirectPriority = InPriority;
-	UpdateAccumulatedPriority();
-}
-
-void FScanDir::UpdateAccumulatedPriority()
-{
-	uint32 LocalAccumulated = static_cast<uint32>(DirectPriority);
-	for (uint32 PriorityLevel = 0; PriorityLevel < CountEPriority; ++PriorityLevel)
-	{
-		if (PriorityRefCounts[PriorityLevel] > 0 && PriorityLevel < LocalAccumulated)
-		{
-			LocalAccumulated = PriorityLevel;
-		}
-	}
-
-	EPriority LocalEPriority = static_cast<EPriority>(LocalAccumulated);
-	if (LocalEPriority != AccumulatedPriority)
-	{
-		if (Parent)
-		{
-			if (AccumulatedPriority != EPriority::Normal)
-			{
-				Parent->OnChildPriorityChanged(AccumulatedPriority, -1);
-			}
-			if (LocalEPriority != EPriority::Normal)
-			{
-				Parent->OnChildPriorityChanged(LocalEPriority, 1);
-			}
-		}
-		AccumulatedPriority = LocalEPriority;
-	}
-}
-
-void FScanDir::OnChildPriorityChanged(EPriority InPriority, int32 Delta)
-{
-	check(-(int32)TNumericLimits<uint8>::Max() < Delta && Delta < TNumericLimits<uint8>::Max());
-	uint8& PriorityRefCount = PriorityRefCounts[static_cast<uint32>(InPriority)];
-	check(Delta > 0 || PriorityRefCount >= -Delta);
-
-	if (Delta > 0 && PriorityRefCount >= TNumericLimits<uint8>::Max() - Delta)
-	{
-		// Mark that the count is now stuck
-		PriorityRefCount = TNumericLimits<uint8>::Max();
-	}
-	else if (Delta < 0 && PriorityRefCount == TNumericLimits<uint8>::Max())
-	{
-		// The count is stuck, do not decrement it
-	}
-	else
-	{
-		PriorityRefCount = (uint8)(PriorityRefCount + Delta);
-	}
-	UpdateAccumulatedPriority();
-}
-
 void FScanDir::SetScanResults(FStringView LocalAbsPath, const FInherited& ParentData, TArrayView<FDiscoveredPathData>& InOutSubDirs, TArrayView<FDiscoveredPathData>& InOutFiles)
 {
 	SetComplete(false);
@@ -731,7 +684,7 @@ void FScanDir::SetScanResults(FStringView LocalAbsPath, const FInherited& Parent
 	bHasScanned = true;
 }
 
-void FScanDir::Update(FScanDir*& OutCursor, FInherited& InOutParentData)
+void FScanDir::Update(TArray<FScanDirAndParentData>& OutScanRequests, const FScanDir::FInherited& ParentData)
 {
 	check(MountDir);
 	if (bIsComplete)
@@ -739,90 +692,36 @@ void FScanDir::Update(FScanDir*& OutCursor, FInherited& InOutParentData)
 		return;
 	}
 
-	FScanDir* ThisCursor = nullptr;
-	if (ShouldScan(InOutParentData))
+	bool bScanThis = ShouldScan(ParentData);
+	if (bScanThis)
 	{
-		ThisCursor = this;
+		OutScanRequests.Add(FScanDirAndParentData{ this, ParentData });
 	}
 
+	bool bAllSubDirsComplete = true;
 	if (SubDirs.Num())
 	{
-		FScanDir* SubDirToScan = FindHighestPrioritySubDir();
-		// If AccumulatedPriority of this and the subdir are the same, the tie goes to the subdir, since 
-		// the accumulatedpriority of a parentdir is clamped to be at least as strong (at least as small) as the subdir's accumulatedpriority
-		if (SubDirToScan && (!ThisCursor || SubDirToScan->AccumulatedPriority <= ThisCursor->AccumulatedPriority))
+		FInherited ParentDataForSubDirs(ParentData, DirectData);
+		TArray<TRefCountPtr<FScanDir>> CopySubDirs = SubDirs;
+		for (const TRefCountPtr<FScanDir>& SubDir : CopySubDirs)
 		{
-			FInherited ThisCursorParentData(InOutParentData, DirectData);
-#if DO_CHECK
-			int32 LoopCount = 0;
-			const int32 MaxLoopCount = SubDirs.Num();
-#endif
-			for (;;) // Continue Updating SubDirs until we find one that needs scanning or all are complete or lower in priority than this
+			if (SubDir->bIsComplete)
 			{
-				FInherited SubDirsParentData(ThisCursorParentData);
-				SubDirToScan->Update(OutCursor, SubDirsParentData);
-				if (OutCursor != this)
-				{
-					InOutParentData = SubDirsParentData;
-					return;
-				}
-				// Otherwise the SubDir finished updating and we should move to the next one
-				// The old value of SubDirToScan may have been deallocated; do not access it further.
-
-				SubDirToScan = FindHighestPrioritySubDir();
-				if (SubDirToScan && (!ThisCursor || SubDirToScan->AccumulatedPriority <= ThisCursor->AccumulatedPriority))
-				{
-					checkf(LoopCount++ < MaxLoopCount,
-						TEXT("Infinite loop: SubDirs return completion without actually completing. SubDirToScan=%s"),
-						*SubDirToScan->GetLocalAbsPath());
-
-					continue;
-				}
-				else
-				{
-					break;
-				}
+				continue;
 			}
+			int32 PreviousCount = OutScanRequests.Num();
+			SubDir->Update(OutScanRequests, ParentDataForSubDirs);
+			bool bSubDirComplete = SubDir->IsComplete();
+			check(OutScanRequests.Num() > PreviousCount || bSubDirComplete);
+			bAllSubDirsComplete &= bSubDirComplete;
 		}
 	}
 
-	if (ThisCursor)
+	if (bScanThis || !bAllSubDirsComplete)
 	{
-		OutCursor = ThisCursor;
 		return;
 	}
 
-	OutCursor = Parent; // Note this will be null for the root ScanDir
-	if (!Parent)
-	{
-		InOutParentData = FInherited();
-	}
-	else
-	{
-		if (Parent->DirectData.bIsOnAllowList)
-		{
-			// We have a contract that bIsOnAllowList is only set on the highest-level directory to monitor and
-			// applies to all directories under it. So we only need to change ParentData.bIsOnAllowList from true to
-			// false when we move up the tree into a parent directory with bIsOnAllowList = true.
-			check(!this->DirectData.bIsOnAllowList); // Verify children below allow list true are allow list false
-			check(!Parent->Parent || !Parent->Parent->DirectData.bIsOnAllowList); // Verify above allow list true is allow list false
-			check(InOutParentData.bIsOnAllowList); // Verify original InOutParentData matches parent's direct value
-			InOutParentData.bIsOnAllowList = false;
-		}
-		if (Parent->DirectData.bMatchesDenyList || Parent->DirectData.bIgnoreDenyList)
-		{
-			// We don't have the same set-once contract for deny list information, so when we find a direct deny list parent
-			// we have to recalculate the parent's deny listed information by checking all ancestors.
-			InOutParentData.bMatchesDenyList = false;
-			InOutParentData.bIgnoreDenyList = false;
-			for (FScanDir* Current = Parent->Parent; Current; Current = Current->Parent)
-			{
-				InOutParentData.bMatchesDenyList = InOutParentData.bMatchesDenyList || Current->DirectData.bMatchesDenyList;
-				InOutParentData.bIgnoreDenyList = InOutParentData.bIgnoreDenyList || Current->DirectData.bIgnoreDenyList;
-			}
-		}
-	}
-
 	SetComplete(true);
 	// After calling SetComplete, this may have been removed from tree and should no longer run calculations
 }
@@ -897,31 +796,7 @@ void FScanDir::SetComplete(bool bInIsComplete)
 	if (bIsComplete)
 	{
 		MountDir->GetDiscovery().NumDirectoriesToScan.Decrement();
-		// If we were given a priority, remove it when we complete
-		SetDirectPriority(EPriority::Normal);
-		// All subDirs are complete, so all of their priorities should be set back to normal, so we can unstick any stuck priorities now by setting them all to 0
-#if DO_CHECK
-		bool bHasPriority = false;
-		ForEachSubDir([&bHasPriority](FScanDir& SubDir) { if (SubDir.GetPriority() != EPriority::Normal) { bHasPriority = true; }});
-		if (bHasPriority)
-		{
-			UE_LOG(LogAssetRegistry, Warning, TEXT("ScanDir %s is marked complete, but it has subdirectories with still-set priorities."), *GetLocalAbsPath());
-		}
-		else
-#endif
-		{
-			bool bModifiedRefCount = false;
-			for (uint8& PriorityRefCount : PriorityRefCounts)
-			{
-				bModifiedRefCount = bModifiedRefCount | (PriorityRefCount != 0);
-				PriorityRefCount = 0;
-			}
-			if (bModifiedRefCount)
-			{
-				UpdateAccumulatedPriority();
-			}
-		}
-		// Upon completion, subdirs that do not need to be maintained are deleted, which is done by removing them from the parent
+		// Upon completion, subdirs that do not need to be maintained are deleted, which is done by removing them from the parent.
 		// ScanDirs need to be maintained if they are the root, or have persistent settings, or have child ScanDirs that need to be maintained,
 		// or the parent scan has not been done yet.
 		if ((Parent != nullptr && Parent->bHasScanned) && !HasPersistentSettings() && SubDirs.IsEmpty())
@@ -999,31 +874,6 @@ int32 FScanDir::FindLowerBoundSubDir(FStringView SubDirBaseName)
 	);
 }
 
-FScanDir* FScanDir::FindHighestPrioritySubDir()
-{
-	if (SubDirs.Num() == 0)
-	{
-		return nullptr;
-	}
-
-	FScanDir* WinningSubDir = nullptr;
-	EPriority WinningPriority = EPriority::Normal;
-
-	for (const TRefCountPtr<FScanDir>& SubDir : SubDirs)
-	{
-		if (SubDir->bIsComplete)
-		{
-			continue;
-		}
-		if (WinningSubDir == nullptr || SubDir->AccumulatedPriority < WinningPriority)
-		{
-			WinningSubDir = SubDir.GetReference();
-			WinningPriority = SubDir->AccumulatedPriority;
-		}
-	}
-	return WinningSubDir;
-}
-
 template <typename CallbackType>
 void FScanDir::ForEachSubDir(const CallbackType& Callback)
 {
@@ -1128,11 +978,6 @@ bool FMountDir::IsComplete() const
 	return Root->IsComplete();
 }
 
-EPriority FMountDir::GetPriority() const
-{
-	return Root->GetPriority();
-}
-
 void FMountDir::GetMonitorData(FStringView InLocalAbsPath, FScanDir::FInherited& OutData) const
 {
 	FStringView QueryRelPath;
@@ -1152,11 +997,16 @@ bool FMountDir::IsMonitored(FStringView InLocalAbsPath) const
 	return MonitorData.IsMonitored();
 }
 
-bool FMountDir::TrySetDirectoryProperties(FStringView InLocalAbsPath, const FSetPathProperties& InProperties, bool bConfirmedExists)
+bool FMountDir::TrySetDirectoryProperties(FStringView InLocalAbsPath, const FSetPathProperties& InProperties, bool bConfirmedExists,
+	FScanDirAndParentData* OutControllingDir)
 {
 	FStringView RelPath;
 	if (!ensure(FPathViews::TryMakeChildPathRelativeTo(InLocalAbsPath, GetLocalAbsPath(), RelPath)))
 	{
+		if (OutControllingDir)
+		{
+			OutControllingDir->ScanDir.SafeRelease();
+		}
 		return false;
 	}
 	if (InProperties.IgnoreDenyList.IsSet())
@@ -1171,12 +1021,22 @@ bool FMountDir::TrySetDirectoryProperties(FStringView InLocalAbsPath, const FSet
 			NewProperties.IgnoreDenyList.Reset();
 			if (!NewProperties.IsSet())
 			{
+				if (OutControllingDir)
+				{
+					OutControllingDir->ScanDir.SafeRelease();
+				}
 				return false;
 			}
-			return TrySetDirectoryProperties(InLocalAbsPath, NewProperties, bConfirmedExists);
+			return TrySetDirectoryProperties(InLocalAbsPath, NewProperties, bConfirmedExists, OutControllingDir);
 		}
 	}
-	return Root->TrySetDirectoryProperties(RelPath, InProperties, bConfirmedExists);
+	FScanDirAndParentData PlaceHolderOutControllingDir;
+	if (!OutControllingDir)
+	{
+		OutControllingDir = &PlaceHolderOutControllingDir;
+	}
+	FScanDir::FInherited ParentData;
+	return Root->TrySetDirectoryProperties(RelPath, ParentData, InProperties, bConfirmedExists, *OutControllingDir);
 }
 
 void FMountDir::UpdateDenyList()
@@ -1218,6 +1078,7 @@ void FMountDir::UpdateDenyList()
 	TStringBuilder<256> AbsPathDenyList;
 	IFileManager& FileManager = IFileManager::Get();
 	FSetPathProperties ChangeDenyList;
+	FScanDir::FInherited ParentData;
 	ChangeDenyList.MatchesDenyList = true;
 	for (const FString& RelPath : AddedDenyListPaths)
 	{
@@ -1226,21 +1087,23 @@ void FMountDir::UpdateDenyList()
 		FPathViews::AppendPath(AbsPathDenyList, RelPath);
 		if (FileManager.DirectoryExists(AbsPathDenyList.ToString()))
 		{
-			Root->TrySetDirectoryProperties(RelPath, ChangeDenyList, true /* bConfirmedExists */);
+			FScanDirAndParentData UnusedControllingDir;
+			Root->TrySetDirectoryProperties(RelPath, ParentData, ChangeDenyList, true /* bConfirmedExists */, UnusedControllingDir);
 		}
 	}
 	ChangeDenyList.MatchesDenyList = false;
 	for (const FString& RelPath : RemovedDenyLists)
 	{
 		// We don't need to check for existence when setting the removal property, because the scandir already exists
-		Root->TrySetDirectoryProperties(RelPath, ChangeDenyList, false /* bConfirmedExists */);
+		FScanDirAndParentData UnusedControllingDir;
+		Root->TrySetDirectoryProperties(RelPath, ParentData, ChangeDenyList, false /* bConfirmedExists */, UnusedControllingDir);
 	}
 }
 
-void FMountDir::Update(FScanDir*& OutCursor, FScanDir::FInherited& OutParentData)
+void FMountDir::Update(TArray<FScanDirAndParentData>& OutScanRequests)
 {
-	OutParentData = FScanDir::FInherited();
-	Root->Update(OutCursor, OutParentData);
+	FScanDir::FInherited ParentData = FScanDir::FInherited();
+	Root->Update(OutScanRequests, ParentData);
 }
 
 FScanDir* FMountDir::GetFirstIncompleteScanDir()
@@ -1393,12 +1256,18 @@ FAssetDataDiscovery::FAssetDataDiscovery(const TArray<FString>& InLongPackageNam
 		bIsSynchronous = true;
 		UE_LOG(LogAssetRegistry, Warning, TEXT("Requested asyncronous asset data discovery, but threading support is disabled. Performing a synchronous discovery instead!"));
 	}
+
+	PriorityDataUpdated->Trigger();
+
+	FParse::Value(FCommandLine::Get(), TEXT("-ARDiscoverThreads="), AssetDataGathererConstants::GARDiscoverThreads);
+	FParse::Value(FCommandLine::Get(), TEXT("-ARDiscoverMinBatchSize="), AssetDataGathererConstants::GARDiscoverMinBatchSize);
+	AssetDataGathererConstants::GARDiscoverThreads = FMath::Max(0, AssetDataGathererConstants::GARDiscoverThreads);
+	AssetDataGathererConstants::GARDiscoverMinBatchSize = FMath::Max(1, AssetDataGathererConstants::GARDiscoverMinBatchSize);
 }
 
 FAssetDataDiscovery::~FAssetDataDiscovery()
 {
 	EnsureCompletion();
-	Cursor.SafeRelease();
 	// Remove pointers to other MountDirs before we delete any of them
 	for (TUniquePtr<FMountDir>& MountDir : MountDirs)
 	{
@@ -1426,23 +1295,35 @@ uint32 FAssetDataDiscovery::Run()
 {
 	LLM_SCOPE(ELLMTag::AssetRegistry);
 	constexpr float IdleSleepTime = 0.1f;
-	{
-		FGathererScopeLock ResultsScopeLock(&ResultsLock);
-		DiscoverStartTime = FPlatformTime::Seconds();
-		NumDiscoveredFiles = 0;
-	}
 
+	CHECK_IS_NOT_LOCKED_CURRENT_THREAD(TreeLock);
+	CHECK_IS_NOT_LOCKED_CURRENT_THREAD(ResultsLock);
 	while (!IsStopped)
 	{
+		bool bTickOwner = false;
+		while (!IsStopped && !bIsIdle && !IsPaused)
 		{
-			CHECK_IS_NOT_LOCKED_CURRENT_THREAD(TreeLock);
-			CHECK_IS_NOT_LOCKED_CURRENT_THREAD(ResultsLock);
-			FGathererScopeLock TickScopeLock(&TickLock);
-			while (!IsStopped && !bIsIdle && !IsPaused)
+			if (!bTickOwner)
+			{
+				if (TickOwner.TryTakeOwnership(TreeLock))
+				{
+					bTickOwner = true;
+				}
+			}
+			if (bTickOwner)
+			{
+				TickInternal(true /* bTickAll */);
+			}
+			else
 			{
-				TickInternal();
+				FPlatformProcess::Sleep(IdleSleepTime);
 			}
 		}
+		if (bTickOwner)
+		{
+			TickOwner.ReleaseOwnershipChecked(TreeLock);
+			bTickOwner = false;
+		}
 
 		while (!IsStopped && (IsPaused || bIsIdle))
 		{
@@ -1461,10 +1342,17 @@ FAssetDataDiscovery::FScopedPause::FScopedPause(const FAssetDataDiscovery& InOwn
 	{
 		Owner.IsPaused++;
 	}
+	while (!Owner.TickOwner.TryTakeOwnership(Owner.TreeLock))
+	{
+		check(!Owner.TickOwner.IsOwnedByCurrentThread());
+		constexpr float BlockingSleepTime = 0.001f;
+		FPlatformProcess::Sleep(BlockingSleepTime);
+	}
 }
 
 FAssetDataDiscovery::FScopedPause::~FScopedPause()
 {
+	Owner.TickOwner.ReleaseOwnershipChecked(Owner.TreeLock);
 	if (!Owner.bIsSynchronous)
 	{
 		check(Owner.IsPaused > 0);
@@ -1494,224 +1382,330 @@ void FAssetDataDiscovery::EnsureCompletion()
 }
 
 
-void FAssetDataDiscovery::TickInternal()
+void FAssetDataDiscovery::TickInternal(bool bTickAll)
 {
 	LLM_SCOPE(ELLMTag::AssetRegistry);
+	CHECK_IS_NOT_LOCKED_CURRENT_THREAD(ResultsLock);
+	check(TickOwner.IsOwnedByCurrentThread());
 
-	TStringBuilder<256> DirLocalAbsPath;
-	TStringBuilder<128> DirLongPackageName;
+	TArray<FScanDirAndParentData> ScanRequests;
 	TStringBuilder<128> DirMountRelPath;
-	int32 DirLongPackageNameRootLen;
-	TRefCountPtr<FScanDir> LocalCursor = nullptr;
+
+	int32 DirToScanDatasNum = 0;
+	bool bUpdatedPriorityData = false;
+	double TickStartTime = FPlatformTime::Seconds();
+	ON_SCOPE_EXIT
+	{
+		if (TickStartTime >= 0.)
+		{
+			CurrentDiscoveryTime += FPlatformTime::Seconds() - TickStartTime;
+			TickStartTime = -1.;
+		}
+	};
+	for (;;)
 	{
-		CHECK_IS_NOT_LOCKED_CURRENT_THREAD(ResultsLock);
-		FGathererScopeLock TreeScopeLock(&TreeLock);
-		for (;;)
 		{
-			// Start at the existing cursor (initializing it if necessary) at call Update until we find a ScanTree that requires scanning
-			if (!Cursor || !Cursor->IsValid())
+			FGathererScopeLock TreeScopeLock(&TreeLock);
+
+			// Process scanned directories from the previous iteration of the for(;;) loop.
+			if (DirToScanDatasNum > 0)
 			{
-				FScanDir* NewCursor;
-				FindFirstCursor(NewCursor, CursorParentData);
-				Cursor = NewCursor;
-				if (!NewCursor)
+				for (FDirToScanData& Data : TArrayView<FDirToScanData>(DirToScanDatas.GetData(), DirToScanDatasNum))
 				{
-					SetIsIdle(true);
-					int32 LocalNumDirectoriesToScan = NumDirectoriesToScan.GetValue();
-					if (LocalNumDirectoriesToScan != 0)
+					if (Data.bScanned)
 					{
-						FScanDir* Incomplete = nullptr;
-						for (TUniquePtr<FMountDir>& MountDir : MountDirs)
+						TArrayView<FDiscoveredPathData> LocalSubDirs(Data.IteratedSubDirs.GetData(), Data.NumIteratedDirs);
+						TArrayView<FDiscoveredPathData> LocalDiscoveredFiles(Data.IteratedFiles.GetData(), Data.NumIteratedFiles);
+						if (!Data.ScanDir->IsValid())
+						{
+							// The ScanDir has been shutdown, and it is only still allocated to prevent us from crashing. Drop our reference and allow it to delete.
+						}
+						else if (Data.ScanDir->IsScanInFlightInvalidated())
 						{
-							Incomplete = MountDir->GetFirstIncompleteScanDir();
-							if (Incomplete)
+							// Some setting has been applied to the ScanDir that requires a new scan
+							// Consume the invalidated flag and ignore the results of our scan
+							Data.ScanDir->SetScanInFlightInvalidated(false);
+						}
+						else
+						{
+							Data.ScanDir->SetScanResults(Data.DirLocalAbsPath, Data.ParentData, LocalSubDirs, LocalDiscoveredFiles);
+							if (!LocalSubDirs.IsEmpty() || !LocalDiscoveredFiles.IsEmpty())
 							{
-								break;
+								AddDiscovered(Data.DirLocalAbsPath, LocalSubDirs, LocalDiscoveredFiles);
 							}
 						}
-						UE_LOG(LogAssetRegistry, Warning, TEXT("FAssetDataDiscovery::SetIsIdle(true) called when NumDirectoriesToScan == %d.\n")
-							TEXT("First incomplete scandir: %s"), LocalNumDirectoriesToScan, Incomplete ? *Incomplete->GetLocalAbsPath() : TEXT("<NoneFound>"));
+						Data.ScanDir->SetScanInFlight(false);
+						Data.ScanDir.SafeRelease();
 					}
-					return;
 				}
+				DirToScanDatasNum = 0;
 			}
-			if (Cursor->ShouldScan(CursorParentData))
+
+			// Look for new dirs to scan, break out of the for (;;) loop if we don't find any
+			auto AddScanRequest = [this, &DirToScanDatasNum, &DirMountRelPath](FScanDirAndParentData&& ScanRequest)
 			{
-				SetIsIdle(false);
-				break;
-			}
+				if (DirToScanDatas.Num() <= DirToScanDatasNum)
+				{
+					// We increment DirToScanDatasNum one at a time so DirToScanDatas.Num() should always be >= to it.
+					check(DirToScanDatas.Num() == DirToScanDatasNum);
+					DirToScanDatas.Emplace();
+				}
+				FScanDir* ScanDir = ScanRequest.ScanDir.GetReference();
+				FDirToScanData& ScanData = DirToScanDatas[DirToScanDatasNum++];
+				ScanData.Reset();
+				DirMountRelPath.Reset();
+				ScanDir->SetScanInFlight(true);
+				FMountDir* MountDir = ScanDir->GetMountDir();
+				check(MountDir);
+				ScanDir->AppendMountRelPath(DirMountRelPath);
+				ScanData.DirLocalAbsPath << MountDir->GetLocalAbsPath();
+				FPathViews::AppendPath(ScanData.DirLocalAbsPath, DirMountRelPath);
+				ScanData.DirLongPackageName << MountDir->GetLongPackageName();
+				FPathViews::AppendPath(ScanData.DirLongPackageName, DirMountRelPath);
+				ScanData.ScanDir = MoveTemp(ScanRequest.ScanDir);
+				ScanData.ParentData = MoveTemp(ScanRequest.ParentData);
+			};
 
-			FScanDir* NewCursor = Cursor.GetReference();
-			NewCursor->Update(NewCursor, CursorParentData);
-			check(NewCursor != Cursor);
-			Cursor = NewCursor;
-		}
-		// IsScanInFlight must be false, because it is not valid to have two TickInternals run at the same time, and we set ScanInFlight back to false after each TickInternal.
-		// If ScanInFlight were true here we would not be able to proceed since we currently don't have a way to find the next ScanTree to update without scanning the current one.
-		check(!Cursor->IsScanInFlight());
-
-		Cursor->SetScanInFlight(true);
-		FMountDir* MountDir = Cursor->GetMountDir();
-		check(MountDir);
-		Cursor->AppendMountRelPath(DirMountRelPath);
-		DirLocalAbsPath << MountDir->GetLocalAbsPath();
-		FPathViews::AppendPath(DirLocalAbsPath, DirMountRelPath);
-		DirLongPackageName << MountDir->GetLongPackageName();
-		FPathViews::AppendPath(DirLongPackageName, DirMountRelPath);
-		DirLongPackageNameRootLen = DirLongPackageName.Len();
-		LocalCursor = Cursor;
-	}
-
-	int32 NumIteratedDirs = 0;
-	int32 NumIteratedFiles = 0;
-	IFileManager::Get().IterateDirectoryStat(DirLocalAbsPath.ToString(), [this, &DirLocalAbsPath, &DirLongPackageName, DirLongPackageNameRootLen, &NumIteratedDirs, &NumIteratedFiles]
-		(const TCHAR* InPackageFilename, const FFileStatData& InPackageStatData)
-		{
-			FStringView LocalAbsPath(InPackageFilename);
-			FStringView RelPath;
-			FString Buffer;
-			if (!FPathViews::TryMakeChildPathRelativeTo(InPackageFilename, DirLocalAbsPath, RelPath))
+			bool bExitAfterPriorityUpdate = !bTickAll && bUpdatedPriorityData;
+			if (!PriorityScanDirs.IsEmpty())
 			{
-				// Try again with the path converted to the absolute path format that we passed in; some IFileManagers can send relative paths to the visitor even though the search path is absolute
-				Buffer = FPaths::ConvertRelativePathToFull(FString(InPackageFilename));
-				LocalAbsPath = Buffer;
-				if (!FPathViews::TryMakeChildPathRelativeTo(Buffer, DirLocalAbsPath, RelPath))
+				ScanRequests.Reset();
+				for (int32 PriorityIndex = 0; PriorityIndex < PriorityScanDirs.Num(); ++PriorityIndex)
 				{
-					UE_LOG(LogAssetRegistry, Warning, TEXT("IterateDirectoryStat returned unexpected result %s which is not a child of the requested path %s."), InPackageFilename, DirLocalAbsPath.ToString());
-					return true;
+					FPriorityScanDirData& PriorityData = PriorityScanDirs[PriorityIndex];
+					int32 OriginalScanRequestsNum = ScanRequests.Num();
+					if (PriorityData.ScanDir->IsValid() && !PriorityData.ScanDir->IsComplete())
+					{
+						PriorityData.ScanDir->Update(ScanRequests, PriorityData.ParentData);
+					}
+					if (PriorityData.ScanDir->IsComplete())
+					{
+						// Update should not add ScanRequests if it was already or transitioned to complete
+						check(ScanRequests.Num() == OriginalScanRequestsNum); 
+						if (PriorityData.bReleaseWhenComplete)
+						{
+							PriorityData.bReleaseWhenComplete = false;
+							check(PriorityData.RequestCount > 0);
+							PriorityData.RequestCount--;
+							if (PriorityData.RequestCount == 0)
+							{
+								PriorityScanDirs.RemoveAtSwap(PriorityIndex);
+								--PriorityIndex; // Counteract the ++PriorityIndex in the for loop iteration
+							}
+							continue;
+						}
+					}
 				}
+
+				if (!bExitAfterPriorityUpdate)
+				{
+					// A ScanDir and its parent can both be in PriorityScanDirs, and in that case we can get duplicates
+					// in the list of ScanRequests. Ensure uniqueness now.
+					Algo::Sort(ScanRequests,
+						[](const FScanDirAndParentData & A, const FScanDirAndParentData& B)
+						{
+							return A.ScanDir.GetReference() < B.ScanDir.GetReference();
+						});
+					ScanRequests.SetNum(Algo::Unique(ScanRequests,
+						[](const FScanDirAndParentData& A, const FScanDirAndParentData& B)
+						{
+							return A.ScanDir.GetReference() == B.ScanDir.GetReference();
+						}));
+
+					for (FScanDirAndParentData& ScanRequest : ScanRequests)
+					{
+						AddScanRequest(MoveTemp(ScanRequest));
+					}
+				}
+				ScanRequests.Reset();
 			}
-			if (FPathViews::GetPathLeaf(RelPath).Len() != RelPath.Len())
+			if (bUpdatedPriorityData)
 			{
-				UE_LOG(LogAssetRegistry, Warning, TEXT("IterateDirectoryStat returned unexpected result %s which is not a direct child of the requested path %s."), InPackageFilename, DirLocalAbsPath.ToString());
-				return true;
+				PriorityDataUpdated->Trigger();
+			}
+			if (bExitAfterPriorityUpdate)
+			{
+				return; // exit to check the done condition
 			}
-			ON_SCOPE_EXIT{ DirLongPackageName.RemoveSuffix(DirLongPackageName.Len() - DirLongPackageNameRootLen); };
+			bPriorityDirty = false;
+			bUpdatedPriorityData = DirToScanDatasNum > 0;
 
-			if (InPackageStatData.bIsDirectory)
+			if (bTickAll && DirToScanDatasNum == 0)
 			{
-				FPathViews::AppendPath(DirLongPackageName, RelPath);
-				// Don't enter directories that contain invalid packagepath characters (including '.'; extensions are not valid in content directories because '.' is not valid in a packagepath)
-				if (!FPackageName::DoesPackageNameContainInvalidCharacters(RelPath))
+				ScanRequests.Reset();
+				UpdateAll(ScanRequests);
+				for (FScanDirAndParentData& ScanRequest : ScanRequests)
 				{
-					if (IteratedSubDirs.Num() < NumIteratedDirs + 1)
-					{
-						check(IteratedSubDirs.Num() == NumIteratedDirs);
-						IteratedSubDirs.Emplace();
-					}
-					IteratedSubDirs[NumIteratedDirs++].Assign(LocalAbsPath, DirLongPackageName, RelPath, EGatherableFileType::Directory);
+					AddScanRequest(MoveTemp(ScanRequest));
 				}
+				ScanRequests.Reset();
 			}
-			else
+
+			if (DirToScanDatasNum == 0)
 			{
-				EGatherableFileType FileType = GetFileType(RelPath);
-				// Don't record files that contain invalid packagepath characters (not counting their extension) or that do not end with a recognized extension
-				if (FileType != EGatherableFileType::Invalid)
+				if (!bTickAll)
 				{
-					FStringView BaseName = FPathViews::GetBaseFilename(RelPath);
-					if (!FPackageName::DoesPackageNameContainInvalidCharacters(BaseName))
+					return;
+				}
+
+				int32 LocalNumDirectoriesToScan = NumDirectoriesToScan.GetValue();
+				if (LocalNumDirectoriesToScan != 0)
+				{
+					// We have some directories left to scan, but we were unable to find any of them. Print diagnostics.
+					FScanDir* Incomplete = nullptr;
+					for (TUniquePtr<FMountDir>& MountDir : MountDirs)
 					{
-						if (IteratedFiles.Num() < NumIteratedFiles + 1)
+						Incomplete = MountDir->GetFirstIncompleteScanDir();
+						if (Incomplete)
 						{
-							check(IteratedFiles.Num() == NumIteratedFiles);
-							IteratedFiles.Emplace();
+							break;
 						}
-						FPathViews::AppendPath(DirLongPackageName, BaseName);
-						IteratedFiles[NumIteratedFiles++].Assign(LocalAbsPath, DirLongPackageName, RelPath, InPackageStatData.ModificationTime, FileType);
 					}
+					UE_LOG(LogAssetRegistry, Warning, TEXT("FAssetDataDiscovery::SetIsIdle(true) called when NumDirectoriesToScan == %d.\n")
+						TEXT("First incomplete scandir: %s"), LocalNumDirectoriesToScan, Incomplete ? *Incomplete->GetLocalAbsPath() : TEXT("<NoneFound>"));
 				}
+				SetIsIdle(true, TickStartTime);
+				return;
 			}
-			return true;
-		});
+			SetIsIdle(false);
+		}
 
-	TArrayView<FDiscoveredPathData> LocalSubDirs(IteratedSubDirs.GetData(), NumIteratedDirs);
-	TArrayView<FDiscoveredPathData> LocalDiscoveredFiles(IteratedFiles.GetData(), NumIteratedFiles);
-	bool bValid = false;
-	{
-		CHECK_IS_NOT_LOCKED_CURRENT_THREAD(ResultsLock);
-		FGathererScopeLock TreeScopeLock(&TreeLock);
-		if (!LocalCursor->IsValid())
+		// Outside of the TreeLock critical section, scan the directories
+		int32 NumThreads = FMath::Max(FTaskGraphInterface::Get().GetNumWorkerThreads(), 1);
+		if (AssetDataGathererConstants::GARDiscoverThreads > 0)
 		{
-			// The ScanDir has been shutdown, and it is only still allocated to prevent us from crashing. Drop our reference and allow it to delete.
+			NumThreads = FMath::Min(NumThreads, AssetDataGathererConstants::GARDiscoverThreads);
 		}
-		else if (LocalCursor->IsScanInFlightInvalidated())
+		int32 DirToScanBuffersNum = FMath::Min(NumThreads, DirToScanDatasNum);
+		if (DirToScanBuffers.Num() < DirToScanBuffersNum)
 		{
-			// Some setting has been applied to the ScanDir that requires a new scan
-			// Consume the invalidated flag and ignore the results of our scan
-			LocalCursor->SetScanInFlightInvalidated(false);
+			DirToScanBuffers.SetNum(DirToScanBuffersNum);
 		}
-		else
+		TArrayView<FDirToScanBuffer> LocalScanBuffers(DirToScanBuffers.GetData(), DirToScanBuffersNum);
+		for (FDirToScanBuffer& ScanBuffer : LocalScanBuffers)
 		{
-			LocalCursor->SetScanResults(DirLocalAbsPath, CursorParentData, LocalSubDirs, LocalDiscoveredFiles);
-			bValid = true;
-			FScanDir* NewCursor = Cursor.GetReference();
-			// Other thread may have set the cursor to a new spot; in that case do not update and on the next tick start at the new cursor
-			if (LocalCursor == NewCursor)
-			{
-				LocalCursor->Update(NewCursor, CursorParentData);
-				Cursor = NewCursor;
-			}
+			ScanBuffer.Reset();
 		}
-		LocalCursor->SetScanInFlight(false);
-	}
-
-	if (bValid && (!LocalSubDirs.IsEmpty() || !LocalDiscoveredFiles.IsEmpty()))
-	{
-		AddDiscovered(DirLocalAbsPath, LocalSubDirs, LocalDiscoveredFiles);
-	}
-}
 
-void FAssetDataDiscovery::FindFirstCursor(FScanDir*& OutCursor, FScanDir::FInherited& OutParentData)
-{
-	CHECK_IS_LOCKED_CURRENT_THREAD(TreeLock);
-	OutCursor = nullptr;
-	while (!OutCursor)
-	{
-		EPriority WinningPriority = EPriority::Normal;
-		FMountDir* WinningMountDir = nullptr;
-		for (TUniquePtr<FMountDir>& MountDir : MountDirs)
+		ParallelForWithExistingTaskContext(LocalScanBuffers, DirToScanDatasNum, AssetDataGathererConstants::GARDiscoverMinBatchSize,
+		[this](FDirToScanBuffer& ScanBuffer, int32 DirToScanDatasIndex)
 		{
-			if (MountDir->IsComplete())
+			FDirToScanData& Data = DirToScanDatas[DirToScanDatasIndex];
+			if (ScanBuffer.bAbort)
 			{
-				continue;
+				return;
 			}
-			if (WinningMountDir == nullptr || MountDir->GetPriority() < WinningPriority)
+			if (bPriorityDirty.load(std::memory_order_relaxed))
 			{
-				WinningMountDir = MountDir.Get();
-				WinningPriority = MountDir->GetPriority();
+				ScanBuffer.bAbort = true;
+				return;
 			}
-		}
 
-		if (!WinningMountDir)
-		{
-			OutCursor = nullptr;
-			OutParentData = FScanDir::FInherited();
-			break;
-		}
+			IFileManager::Get().IterateDirectoryStat(Data.DirLocalAbsPath.ToString(),
+			[this, &Data]
+			(const TCHAR* InPackageFilename, const FFileStatData& InPackageStatData)
+			{
+				FStringView LocalAbsPath(InPackageFilename);
+				FStringView RelPath;
+				FString Buffer;
+				if (!FPathViews::TryMakeChildPathRelativeTo(InPackageFilename, Data.DirLocalAbsPath, RelPath))
+				{
+					// Try again with the path converted to the absolute path format that we passed in; some
+					// IFileManagers can send relative paths to the visitor even though the search path is absolute
+					Buffer = FPaths::ConvertRelativePathToFull(FString(InPackageFilename));
+					LocalAbsPath = Buffer;
+					if (!FPathViews::TryMakeChildPathRelativeTo(Buffer, Data.DirLocalAbsPath, RelPath))
+					{
+						UE_LOG(LogAssetRegistry, Warning,
+							TEXT("IterateDirectoryStat returned unexpected result %s which is not a child of the requested path %s."),
+							InPackageFilename, Data.DirLocalAbsPath.ToString());
+						return true;
+					}
+				}
+				if (FPathViews::GetPathLeaf(RelPath).Len() != RelPath.Len())
+				{
+					UE_LOG(LogAssetRegistry, Warning,
+						TEXT("IterateDirectoryStat returned unexpected result %s which is not a direct child of the requested path %s."),
+						InPackageFilename, Data.DirLocalAbsPath.ToString());
+					return true;
+				}
+				int32 DirLongPackageRootNameLen = Data.DirLongPackageName.Len();
+				ON_SCOPE_EXIT
+				{
+					Data.DirLongPackageName.RemoveSuffix(Data.DirLongPackageName.Len() - DirLongPackageRootNameLen);
+				};
 
-		WinningMountDir->Update(OutCursor, OutParentData);
-		check(OutCursor != nullptr || WinningMountDir->IsComplete()); // The WinningMountDir's update should either return something to update or it should mark itself complete
+				if (InPackageStatData.bIsDirectory)
+				{
+					FPathViews::AppendPath(Data.DirLongPackageName, RelPath);
+					// Don't enter directories that contain invalid packagepath characters (including '.';
+					// extensions are not valid in content directories because '.' is not valid in a packagepath)
+					if (!FPackageName::DoesPackageNameContainInvalidCharacters(RelPath))
+					{
+						if (Data.IteratedSubDirs.Num() < Data.NumIteratedDirs + 1)
+						{
+							check(Data.IteratedSubDirs.Num() == Data.NumIteratedDirs);
+							Data.IteratedSubDirs.Emplace();
+						}
+						Data.IteratedSubDirs[Data.NumIteratedDirs++].Assign(LocalAbsPath, Data.DirLongPackageName, RelPath,
+							EGatherableFileType::Directory);
+					}
+				}
+				else
+				{
+					EGatherableFileType FileType = GetFileType(RelPath);
+					// Don't record files that contain invalid packagepath characters (not counting their extension)
+					// or that do not end with a recognized extension
+					if (FileType != EGatherableFileType::Invalid)
+					{
+						FStringView BaseName = FPathViews::GetBaseFilename(RelPath);
+						if (!FPackageName::DoesPackageNameContainInvalidCharacters(BaseName))
+						{
+							if (Data.IteratedFiles.Num() < Data.NumIteratedFiles + 1)
+							{
+								check(Data.IteratedFiles.Num() == Data.NumIteratedFiles);
+								Data.IteratedFiles.Emplace();
+							}
+							FPathViews::AppendPath(Data.DirLongPackageName, BaseName);
+							Data.IteratedFiles[Data.NumIteratedFiles++].Assign(LocalAbsPath, Data.DirLongPackageName, RelPath,
+								InPackageStatData.ModificationTime, FileType);
+						}
+					}
+				}
+				return true;
+			});
+			Data.bScanned = true;
+		});
 	}
 }
 
-void FAssetDataDiscovery::InvalidateCursor()
+void FAssetDataDiscovery::UpdateAll(TArray<FScanDirAndParentData>& OutScanRequests)
 {
-	if (Cursor)
+	CHECK_IS_LOCKED_CURRENT_THREAD(TreeLock);
+	for (TUniquePtr<FMountDir>& MountDirOwner : MountDirs)
 	{
-		if (Cursor->IsScanInFlight())
+		FMountDir* MountDir = MountDirOwner.Get();
+		if (MountDir->IsComplete())
 		{
-			Cursor->SetScanInFlightInvalidated(true);
+			continue;
 		}
-		Cursor.SafeRelease();
+
+		MountDir->Update(OutScanRequests);
 	}
 }
 
 void FAssetDataDiscovery::SetIsIdle(bool bInIsIdle)
+{
+	double TickStartTime = -1.;
+	SetIsIdle(bInIsIdle, TickStartTime);
+}
+
+void FAssetDataDiscovery::SetIsIdle(bool bInIsIdle, double& TickStartTime)
 {
 	CHECK_IS_LOCKED_CURRENT_THREAD(TreeLock);
 
 	// Caller is responsible for holding TreeLock around this function; writes of SetIsIdle are done inside the TreeLock
-	// If bIsIdle is true, caller holds TickLock and TreeLock
+	// If bIsIdle is true, caller holds TickOwner and TreeLock
 	if (bIsIdle == bInIsIdle)
 	{
 		return;
@@ -1722,18 +1716,28 @@ void FAssetDataDiscovery::SetIsIdle(bool bInIsIdle)
 	{
 		if (bIsIdle)
 		{
-			UE_LOG(LogAssetRegistry, Verbose, TEXT("Discovery took %0.6f seconds and found %d files to process"), FPlatformTime::Seconds() - DiscoverStartTime, NumDiscoveredFiles);
+			if (TickStartTime >= 0.)
+			{
+				CurrentDiscoveryTime += FPlatformTime::Seconds() - TickStartTime;
+				TickStartTime = -1.;
+			}
+
+			CumulativeDiscoveryTime += static_cast<float>(CurrentDiscoveryTime);
+			CumulativeDiscoveredFiles += NumDiscoveredFiles;
+			UE_LOG(LogAssetRegistry, Verbose,
+				TEXT("Discovery took %0.4f seconds to add %d files, Cumulative=%0.4f seconds to add %d."),
+				CurrentDiscoveryTime, NumDiscoveredFiles, CumulativeDiscoveryTime, CumulativeDiscoveredFiles);
+			CurrentDiscoveryTime = 0.;
 		}
 		else
 		{
-			DiscoverStartTime = FPlatformTime::Seconds();
 			NumDiscoveredFiles = 0;
 		}
 	}
 
 	if (bIsIdle)
 	{
-		CHECK_IS_LOCKED_CURRENT_THREAD(TickLock);
+		check(TickOwner.IsOwnedByCurrentThread());
 		Shrink();
 	}
 }
@@ -1768,19 +1772,42 @@ void FAssetDataDiscovery::GetAndTrimSearchResults(bool& bOutIsComplete, TArray<F
 	}
 }
 
+void FAssetDataDiscovery::GetDiagnostics(float& OutCumulativeDiscoveryTime)
+{
+	FGathererScopeLock ResultsScopeLock(&ResultsLock);
+	OutCumulativeDiscoveryTime = CumulativeDiscoveryTime;
+}
+
+
 void FAssetDataDiscovery::WaitForIdle()
 {
 	if (bIsIdle)
 	{
 		return;
 	}
-	FScopedPause ScopedPause(*this);
 	CHECK_IS_NOT_LOCKED_CURRENT_THREAD(TreeLock);
 	CHECK_IS_NOT_LOCKED_CURRENT_THREAD(ResultsLock);
-	FGathererScopeLock TickScopeLock(&TickLock);
+
+	constexpr float IdleSleepTime = 0.1f;
+	bool bTickOwner = false;
 	while (!bIsIdle)
 	{
-		TickInternal();
+		if (!bTickOwner)
+		{
+			bTickOwner = TickOwner.TryTakeOwnership(TreeLock);
+		}
+		if (bTickOwner)
+		{
+			TickInternal(true /* bTickAll */);
+		}
+		else
+		{
+			FPlatformProcess::Sleep(IdleSleepTime);
+		}
+	}
+	if (bTickOwner)
+	{
+		TickOwner.ReleaseOwnershipChecked(TreeLock);
 	}
 }
 
@@ -1843,143 +1870,246 @@ void FPathExistence::LoadExistenceData()
 	bHasExistenceData = true;
 }
 
-void FAssetDataDiscovery::SetPropertiesAndWait(FPathExistence& QueryPath, bool bAddToAllowList, bool bForceRescan, bool bIgnoreDenyListScanFilters)
+void FAssetDataDiscovery::SetPropertiesAndWait(TArrayView<FPathExistence> QueryPaths, bool bAddToAllowList,
+	bool bForceRescan, bool bIgnoreDenyListScanFilters)
 {
-	// We might have been asked to wait on a filename missing the extension, in which case QueryPath.GetType() == MissingButDirExists
-	// We need to handle Directory, File, and MissingButDirExists in unique ways
-	FPathExistence::EType PathType = QueryPath.GetType();
-	if (PathType == FPathExistence::EType::MissingParentDir)
+	for (FPathExistence& QueryPath : QueryPaths)
 	{
-		// SetPropertiesAndWait is called for every ScanPathsSynchronous, and this is the first spot that checks for existence. 
-		// Some systems call ScanPathsSynchronous speculatively to scan whatever is present, so this log is verbose-only.
-		UE_LOG(LogAssetRegistry, Verbose, TEXT("SetPropertiesAndWait called on non-existent path %s. Call will be ignored."),
-			*QueryPath.GetLocalAbsPath());
-		return;
+		QueryPath.LoadExistenceData();
 	}
-	FStringView SearchPath = QueryPath.GetLowestExistingPath();
 
+	struct FScanDirAndQueryPath
+	{
+		TRefCountPtr<FScanDir> ScanDir;
+		bool bScanEntireTree = false;
+	};
+	TArray<FScanDirAndQueryPath> DirsToScan;
+	bool bTickOwner = false;
 	{
 		CHECK_IS_NOT_LOCKED_CURRENT_THREAD(ResultsLock);
 		FGathererScopeLock TreeScopeLock(&TreeLock);
-		FMountDir* MountDir = FindContainingMountPoint(SearchPath);
-		if (!MountDir)
+		for (FPathExistence& QueryPath : QueryPaths)
 		{
-			UE_LOG(LogAssetRegistry, Log, TEXT("SetPropertiesAndWait called on %s which is not in a mounted directory. Call will be ignored."),
-				*QueryPath.GetLocalAbsPath());
-			return;
-		}
-
-		if (PathType == FPathExistence::EType::Directory)
-		{
-			FSetPathProperties Properties;
-			if (bAddToAllowList)
+			// We might have been asked to wait on a filename missing the extension, in which case QueryPath.GetType() == MissingButDirExists
+			// We need to handle Directory, File, and MissingButDirExists in unique ways
+			FPathExistence::EType PathType = QueryPath.GetType();
+			if (PathType == FPathExistence::EType::MissingParentDir)
 			{
-				Properties.IsOnAllowList = bAddToAllowList;
+				// SetPropertiesAndWait is called for every ScanPathsSynchronous, and this is the first spot that checks for existence. 
+				// Some systems call ScanPathsSynchronous speculatively to scan whatever is present, so this log is verbose-only.
+				UE_LOG(LogAssetRegistry, Verbose, TEXT("SetPropertiesAndWait called on non-existent path %s. Call will be ignored."),
+					*QueryPath.GetLocalAbsPath());
+				continue;
 			}
-			if (bForceRescan)
+			FStringView SearchPath = QueryPath.GetLowestExistingPath();
+
+			FMountDir* MountDir = FindContainingMountPoint(SearchPath);
+			if (!MountDir)
 			{
-				Properties.HasScanned = false;
+				UE_LOG(LogAssetRegistry, Log, TEXT("SetPropertiesAndWait called on %s which is not in a mounted directory. Call will be ignored."),
+					*QueryPath.GetLocalAbsPath());
+				continue;
 			}
-			if (bIgnoreDenyListScanFilters)
+
+			if (PathType == FPathExistence::EType::Directory)
 			{
-				Properties.IgnoreDenyList = true;
+				FSetPathProperties Properties;
+				if (bAddToAllowList)
+				{
+					Properties.IsOnAllowList = bAddToAllowList;
+				}
+				if (bForceRescan)
+				{
+					Properties.HasScanned = false;
+				}
+				if (bIgnoreDenyListScanFilters)
+				{
+					Properties.IgnoreDenyList = true;
+				}
+				if (Properties.IsSet())
+				{
+					SetIsIdle(false);
+					MountDir->TrySetDirectoryProperties(SearchPath, Properties, true /* bConfirmedExists */, nullptr);
+				}
 			}
-			if (Properties.IsSet())
+
+			FString RelPath;
+			FScanDir::FInherited MonitorData;
+			bool bSearchPathIsDirectory = PathType == FPathExistence::EType::Directory || PathType == FPathExistence::EType::MissingButDirExists;
+			TRefCountPtr<FScanDir> ScanDir = MountDir->GetControllingDir(SearchPath, bSearchPathIsDirectory, MonitorData, RelPath);
+			bool bIsAllowedInThisCall = MonitorData.IsOnAllowList() || bAddToAllowList;
+			bool bIsDeniedInThisCall = MonitorData.IsOnDenyList() && !bIgnoreDenyListScanFilters;
+			bool bIsMonitoredInThisCall = bIsAllowedInThisCall && !bIsDeniedInThisCall;
+			if (!ScanDir || !bIsMonitoredInThisCall)
 			{
-				SetIsIdle(false);
-				MountDir->TrySetDirectoryProperties(SearchPath, Properties, true /* bConfirmedExists */);
+				UE_LOG(LogAssetRegistry, Log, TEXT("SetPropertiesAndWait called on %s which is not monitored. Call will be ignored."),
+					*QueryPath.GetLocalAbsPath());
+				continue;
 			}
-		}
 
-		FString RelPath;
-		FScanDir::FInherited MonitorData;
-		bool bSearchPathIsDirectory = PathType == FPathExistence::EType::Directory || PathType == FPathExistence::EType::MissingButDirExists;
-		TRefCountPtr<FScanDir> ScanDir = MountDir->GetControllingDir(SearchPath, bSearchPathIsDirectory, MonitorData, RelPath);
-		bool bIsAllowedInThisCall = MonitorData.IsOnAllowList() || bAddToAllowList;
-		bool bIsDeniedInThisCall = MonitorData.IsOnDenyList() && !bIgnoreDenyListScanFilters;
-		bool bIsMonitoredInThisCall = bIsAllowedInThisCall && !bIsDeniedInThisCall;
-		if (!ScanDir || !bIsMonitoredInThisCall)
-		{
-			UE_LOG(LogAssetRegistry, Log, TEXT("SetPropertiesAndWait called on %s which is not monitored. Call will be ignored."),
-				*QueryPath.GetLocalAbsPath());
-			return;
-		}
+			if (bSearchPathIsDirectory)
+			{
+				// If Relpath from the controlling dir to the requested dir is not empty then we have found a parent directory rather than the requested directory.
+				// This can only occur for a monitored directory when the requested directory is already complete and we do not need to wait on it.
+				if (!RelPath.IsEmpty() || ScanDir->IsComplete())
+				{
+					continue;
+				}
 
-		if (bSearchPathIsDirectory)
-		{
-			// If Relpath from the controlling dir to the requested dir is not empty then we have found a parent directory rather than the requested directory.
-			// This can only occur for a monitored directory when the requested directory is already complete and we do not need to wait on it.
-			if (RelPath.IsEmpty() && !ScanDir->IsComplete())
+				FScanDirAndQueryPath& ScanQuery = DirsToScan.Emplace_GetRef();
+				ScanQuery.ScanDir = ScanDir;
+				ScanQuery.bScanEntireTree = PathType == FPathExistence::EType::Directory;
+				FPriorityScanDirData* PriorityData = PriorityScanDirs.FindByPredicate(
+					[&ScanDir](const FPriorityScanDirData& ScanDirData) { return ScanDirData.ScanDir == ScanDir; });
+				if (!PriorityData)
+				{
+					PriorityData = &PriorityScanDirs.Add_GetRef(FPriorityScanDirData{ ScanDir });
+				}
+				++PriorityData->RequestCount;
+				PriorityData->ParentData = MonitorData;
+			}
+			else
 			{
-				// We are going to wait on the path, so set its priority to blocking
-				SetIsIdle(false);
-				EPriority OriginalPriority = ScanDir->GetDirectPriority();
-				ScanDir->SetDirectPriority(EPriority::Blocking);
-				InvalidateCursor();
-
-				FScopedPause ScopedPause(*this);
-				TreeScopeLock.Unlock(); // Entering the ticklock, as well as any long duration task such as a tick, has to be done outside of any locks
-				// If the query path is MissingButDirExists, we assume it is a file missing the extension, and scan (or confirm already scanned) its directory
-				bool bScanEntireTree = PathType == FPathExistence::EType::Directory;
-
-				CHECK_IS_NOT_LOCKED_CURRENT_THREAD(TreeLock);
-				CHECK_IS_NOT_LOCKED_CURRENT_THREAD(ResultsLock);
-				FGathererScopeLock TickScopeLock(&TickLock);
-				for (;;)
+				check(PathType == FPathExistence::EType::File);
+				bool bAlreadyScanned = ScanDir->HasScanned() && MonitorData.IsMonitored();
+				if (!bAlreadyScanned || bForceRescan)
 				{
-					TickInternal();
-					FGathererScopeLock LoopTreeScopeLock(&TreeLock);
-					if (!ScanDir->IsValid())
-					{
-						break;
-					}
-					if (ScanDir->IsComplete() || (!bScanEntireTree && ScanDir->HasScanned()))
-					{
-						// Restore the Priority.
-						// TODO: We dropped the lock so something else could have changed the value of ScanDir's priority and we are now clobbering it.
-						// This can be fixed by adding a counter for each level of the priority. That's complex and cpu expensive, so not doing it until we need it.
-						ScanDir->SetDirectPriority(OriginalPriority);
-						break;
-					}
-					else if (!ensureMsgf(!bIsIdle, TEXT("It should not be possible for the Discovery to go idle while there is an incomplete ScanDir.")))
+					FStringView RelPathFromParentDir = FPathViews::GetCleanFilename(RelPath);
+					EGatherableFileType FileType = GetFileType(RelPathFromParentDir);
+					if (FileType != EGatherableFileType::Invalid)
 					{
-						// Restore the Priority
-						ScanDir->SetDirectPriority(OriginalPriority);
-						break;
+						FStringView FileRelPathNoExt = FPathViews::GetBaseFilenameWithPath(RelPath);
+						if (!FPackageName::DoesPackageNameContainInvalidCharacters(FileRelPathNoExt))
+						{
+							TStringBuilder<256> LongPackageName;
+							LongPackageName << MountDir->GetLongPackageName();
+							FPathViews::AppendPath(LongPackageName, ScanDir->GetMountRelPath());
+							FPathViews::AppendPath(LongPackageName, FileRelPathNoExt);
+							AddDiscoveredFile(FDiscoveredPathData(SearchPath, LongPackageName, RelPathFromParentDir, QueryPath.GetModificationTime(), FileType));
+							if (FPathViews::IsPathLeaf(RelPath) && !ScanDir->HasScanned())
+							{
+								SetIsIdle(false);
+								ScanDir->MarkFileAlreadyScanned(RelPath);
+							}
+						}
 					}
 				}
 			}
 		}
+
+		if (!DirsToScan.IsEmpty())
+		{
+			bPriorityDirty = true;
+			PriorityDataUpdated->Reset();
+			bTickOwner = TickOwner.TryTakeOwnership(TreeScopeLock);
+		}
+	}
+
+	while (!DirsToScan.IsEmpty())
+	{
+		if (bTickOwner)
+		{
+			TickInternal(false /* bTickAll */);
+		}
 		else
 		{
-			check(PathType == FPathExistence::EType::File);
-			bool bAlreadyScanned = ScanDir->HasScanned() && MonitorData.IsMonitored();
-			if (!bAlreadyScanned || bForceRescan)
+			constexpr uint32 WaitTimeMilliseconds = 100;
+			PriorityDataUpdated->Wait(WaitTimeMilliseconds);
+		}
+
+		{
+			FGathererScopeLock LoopTreeScopeLock(&TreeLock);
+			for (int32 Index = 0; Index < DirsToScan.Num(); ++Index)
 			{
-				FStringView RelPathFromParentDir = FPathViews::GetCleanFilename(RelPath);
-				EGatherableFileType FileType = GetFileType(RelPathFromParentDir);
-				if (FileType != EGatherableFileType::Invalid)
+				FScanDirAndQueryPath& ScanQuery = DirsToScan[Index];
+				FScanDir* ScanDir = ScanQuery.ScanDir.GetReference();
+				auto RemoveCurrent = [&DirsToScan, &Index, ScanDir, this]()
 				{
-					FStringView FileRelPathNoExt = FPathViews::GetBaseFilenameWithPath(RelPath);
-					if (!FPackageName::DoesPackageNameContainInvalidCharacters(FileRelPathNoExt))
+					DirsToScan.RemoveAtSwap(Index);
+					int32 PriorityDataIndex = PriorityScanDirs.IndexOfByPredicate(
+						[ScanDir](const FPriorityScanDirData& ScanDirData) { return ScanDirData.ScanDir.GetReference() == ScanDir; });
+
+					check(PriorityDataIndex != INDEX_NONE); // Nothing should be able to remove it until we remove our RequestCount
+					FPriorityScanDirData& PriorityData = PriorityScanDirs[PriorityDataIndex];
+					check(PriorityData.RequestCount > 0);
+					--PriorityData.RequestCount;
+					if (PriorityData.RequestCount == 0)
 					{
-						TStringBuilder<256> LongPackageName;
-						LongPackageName << MountDir->GetLongPackageName();
-						FPathViews::AppendPath(LongPackageName, ScanDir->GetMountRelPath());
-						FPathViews::AppendPath(LongPackageName, FileRelPathNoExt);
-						AddDiscoveredFile(FDiscoveredPathData(SearchPath, LongPackageName, RelPathFromParentDir, QueryPath.GetModificationTime(), FileType));
-						if (FPathViews::IsPathLeaf(RelPath) && !ScanDir->HasScanned())
-						{
-							SetIsIdle(false);
-							ScanDir->MarkFileAlreadyScanned(RelPath);
-						}
+						PriorityScanDirs.RemoveAtSwap(PriorityDataIndex);
 					}
+				};
+
+				if (!ScanDir->IsValid())
+				{
+					RemoveCurrent();
+					continue;
+				}
+				if (ScanDir->IsComplete() || (!ScanQuery.bScanEntireTree && ScanDir->HasScanned()))
+				{
+					RemoveCurrent();
+					continue;
+				}
+				else if (!ensureMsgf(!bIsIdle, TEXT("It should not be possible for the Discovery to go idle while there is an incomplete ScanDir.")))
+				{
+					RemoveCurrent();
+					continue;
+				}
+			}
+
+			if (DirsToScan.IsEmpty())
+			{
+				if (bTickOwner)
+				{
+					TickOwner.ReleaseOwnershipChecked(LoopTreeScopeLock);
+					bTickOwner = false;
+				}
+			}
+			else
+			{
+				PriorityDataUpdated->Reset();
+				if (!bTickOwner)
+				{
+					bTickOwner = TickOwner.TryTakeOwnership(LoopTreeScopeLock);
 				}
 			}
 		}
 	}
 }
 
+void FAssetDataDiscovery::PrioritizeSearchPath(const FString& LocalAbsPath, EPriority Priority)
+{
+	CHECK_IS_NOT_LOCKED_CURRENT_THREAD(ResultsLock);
+	FGathererScopeLock TreeScopeLock(&TreeLock);
+	SetIsIdle(false);
+	FMountDir* MountDir = FindContainingMountPoint(LocalAbsPath);
+	if (!MountDir)
+	{
+		UE_LOG(LogAssetRegistry, Warning, TEXT("FAssetDataGatherer::PrioritizeSearchPath called on unmounted path %.*s. Call will be ignored."),
+			LocalAbsPath.Len(), *LocalAbsPath);
+		return;
+	}
+
+	FSetPathProperties EmptyProperties;
+	FScanDirAndParentData ScanDirAndParent;
+	MountDir->TrySetDirectoryProperties(LocalAbsPath, EmptyProperties, false, &ScanDirAndParent);
+	FScanDir* ScanDir = ScanDirAndParent.ScanDir.GetReference();
+	if (ScanDir && ScanDir->IsValid() && !ScanDir->IsComplete())
+	{
+		FPriorityScanDirData* PriorityData = PriorityScanDirs.FindByPredicate(
+			[&ScanDir](const FPriorityScanDirData& ScanDirData) { return ScanDirData.ScanDir.GetReference() == ScanDir; });
+		if (!PriorityData)
+		{
+			PriorityData = &PriorityScanDirs.Add_GetRef(FPriorityScanDirData{ ScanDir });
+		}
+		if (!PriorityData->bReleaseWhenComplete)
+		{
+			PriorityData->bReleaseWhenComplete = true;
+			++PriorityData->RequestCount;
+		}
+		PriorityData->ParentData = ScanDirAndParent.ParentData;
+	}
+}
+
 bool FAssetDataDiscovery::TrySetDirectoryProperties(const FString& LocalAbsPath, const FSetPathProperties& InProperties, bool bConfirmedExists)
 {
 	if (!InProperties.IsSet())
@@ -1993,7 +2123,6 @@ bool FAssetDataDiscovery::TrySetDirectoryProperties(const FString& LocalAbsPath,
 	{
 		return false;
 	}
-	InvalidateCursor();
 	return true;
 }
 
@@ -2007,7 +2136,7 @@ bool FAssetDataDiscovery::TrySetDirectoryPropertiesInternal(const FString& Local
 		return false;
 	}
 
-	return MountDir->TrySetDirectoryProperties(LocalAbsPath, InProperties, bConfirmedExists);
+	return MountDir->TrySetDirectoryProperties(LocalAbsPath, InProperties, bConfirmedExists, nullptr);
 }
 
 bool FAssetDataDiscovery::IsOnAllowList(FStringView LocalAbsPath) const
@@ -2059,11 +2188,10 @@ SIZE_T GetArrayRecursiveAllocatedSize(const ContainerType& Container)
 
 SIZE_T FAssetDataDiscovery::GetAllocatedSize() const
 {
-
-	FScopedPause ScopedPause(*this);
 	CHECK_IS_NOT_LOCKED_CURRENT_THREAD(TreeLock);
 	CHECK_IS_NOT_LOCKED_CURRENT_THREAD(ResultsLock);
-	FGathererScopeLock TickScopeLock(&TickLock);
+	check(!TickOwner.IsOwnedByCurrentThread());
+	FScopedPause ScopedPause(*this);
 	FGathererScopeLock TreeScopeLock(&TreeLock);
 	FGathererScopeLock ResultsScopeLock(&ResultsLock);
 
@@ -2087,11 +2215,35 @@ SIZE_T FAssetDataDiscovery::GetAllocatedSize() const
 		Result += sizeof(*Value);
 		Result += Value->GetAllocatedSize();
 	}
+	Result += GetArrayRecursiveAllocatedSize(DirToScanDatas);
+	Result += DirToScanBuffers.GetAllocatedSize();
+	return Result;
+}
+
+void FAssetDataDiscovery::FDirToScanData::Reset()
+{
+	DirLocalAbsPath.Reset();
+	DirLongPackageName.Reset();
+	NumIteratedDirs = 0;
+	NumIteratedFiles = 0;
+	bScanned = false;
+}
+
+SIZE_T FAssetDataDiscovery::FDirToScanData::GetAllocatedSize() const
+{
+	SIZE_T Result = 0;
+	Result += DirLocalAbsPath.GetAllocatedSize();
+	Result += DirLongPackageName.GetAllocatedSize();
 	Result += GetArrayRecursiveAllocatedSize(IteratedSubDirs);
 	Result += GetArrayRecursiveAllocatedSize(IteratedFiles);
 	return Result;
 }
 
+void FAssetDataDiscovery::FDirToScanBuffer::Reset()
+{
+	bAbort = false;
+}
+
 SIZE_T FAssetDataDiscovery::FDirectoryResult::GetAllocatedSize() const
 {
 	return DirAbsPath.GetAllocatedSize() + Files.GetAllocatedSize();
@@ -2099,7 +2251,7 @@ SIZE_T FAssetDataDiscovery::FDirectoryResult::GetAllocatedSize() const
 
 void FAssetDataDiscovery::Shrink()
 {
-	CHECK_IS_LOCKED_CURRENT_THREAD(TickLock);
+	check(TickOwner.IsOwnedByCurrentThread());
 	CHECK_IS_LOCKED_CURRENT_THREAD(TreeLock);
 	CHECK_IS_LOCKED_CURRENT_THREAD(ResultsLock);
 	DirLongPackageNamesToNotReport.Shrink();
@@ -2111,8 +2263,8 @@ void FAssetDataDiscovery::Shrink()
 	{
 		MountDir->Shrink();
 	}
-	IteratedSubDirs.Shrink();
-	IteratedFiles.Shrink();
+	DirToScanDatas.Empty();
+	DirToScanBuffers.Empty();
 }
 
 void FAssetDataDiscovery::AddMountPoint(const FString& LocalAbsPath, FStringView LongPackageName)
@@ -2121,8 +2273,6 @@ void FAssetDataDiscovery::AddMountPoint(const FString& LocalAbsPath, FStringView
 	FGathererScopeLock TreeScopeLock(&TreeLock);
 	SetIsIdle(false);
 	AddMountPointInternal(LocalAbsPath, LongPackageName);
-
-	InvalidateCursor();
 }
 
 void FAssetDataDiscovery::AddMountPointInternal(const FString& LocalAbsPath, FStringView LongPackageName)
@@ -2183,8 +2333,6 @@ void FAssetDataDiscovery::RemoveMountPoint(const FString& LocalAbsPath)
 	FGathererScopeLock TreeScopeLock(&TreeLock);
 	SetIsIdle(false);
 	RemoveMountPointInternal(LocalAbsPath);
-
-	InvalidateCursor();
 }
 
 void FAssetDataDiscovery::RemoveMountPointInternal(const FString& LocalAbsPath)
@@ -2391,8 +2539,11 @@ void FAssetDataDiscovery::AddDiscovered(FStringView DirAbsPath, TConstArrayView<
 	{
 		DiscoveredDirectories.Add(FString(SubDir.LongPackageName));
 	}
-	DiscoveredFiles.Emplace(DirAbsPath, Files);
-	NumDiscoveredFiles += Files.Num();
+	if (Files.Num())
+	{
+		DiscoveredFiles.Emplace(DirAbsPath, Files);
+		NumDiscoveredFiles += Files.Num();
+	}
 }
 
 void FAssetDataDiscovery::AddDiscoveredFile(FDiscoveredPathData&& File)
@@ -2948,7 +3099,7 @@ FAssetDataGatherer::FAssetDataGatherer(const TArray<FString>& InLongPackageNames
 	, IsPaused(0)
 	, bInitialPluginsLoaded(false)
 	, bSaveAsyncCacheTriggered(false)
-	, SearchStartTime(0)
+	, CurrentSearchTime(0.)
 	, LastCacheWriteTime(0.0)
 	, bHasLoadedMonolithicCache(false)
 	, bDiscoveryIsComplete(false)
@@ -3084,9 +3235,15 @@ void FAssetDataGatherer::InnerTickLoop(bool bInIsSynchronousTick, bool bContribu
 		TGuardValue<bool> ScopeSynchronousTick(bIsSynchronousTick, bInIsSynchronousTick);
 		TRACE_CPUPROFILER_EVENT_SCOPE(FAssetDataGatherer::Tick);
 		bool bTickInterruptionEvent = false;
+		double TickStartTime = FPlatformTime::Seconds();
 		while (!IsStopped && (bInIsSynchronousTick || !IsPaused) && !bTickInterruptionEvent)
 		{
-			TickInternal(bTickInterruptionEvent);
+			TickInternal(bTickInterruptionEvent, TickStartTime);
+		}
+		if (TickStartTime >= 0.)
+		{
+			CurrentSearchTime += FPlatformTime::Seconds() - TickStartTime;
+			TickStartTime = 0.;
 		}
 
 		if (bContributeToCacheSave)
@@ -3154,7 +3311,7 @@ void FAssetDataGatherer::EnsureCompletion()
 	}
 }
 
-void FAssetDataGatherer::TickInternal(bool& bOutIsTickInterrupt)
+void FAssetDataGatherer::TickInternal(bool& bOutIsTickInterrupt, double& TickStartTime)
 {
 	LLM_SCOPE(ELLMTag::AssetRegistry);
 
@@ -3181,7 +3338,6 @@ void FAssetDataGatherer::TickInternal(bool& bOutIsTickInterrupt)
 		{
 			bFirstTickAfterIdle = false;
 			LastCacheWriteTime = FPlatformTime::Seconds();
-			SearchStartTime = LastCacheWriteTime;
 		}
 
 		IngestDiscoveryResults();
@@ -3212,7 +3368,7 @@ void FAssetDataGatherer::TickInternal(bool& bOutIsTickInterrupt)
 			{
 				bSaveAsyncCacheTriggered = true;
 			}
-			SetIsIdle(true);
+			SetIsIdle(true, TickStartTime);
 			return;
 		}
 		if (bReadMonolithicCache && !bHasLoadedMonolithicCache)
@@ -3597,6 +3753,13 @@ void FAssetDataGatherer::GetAndTrimSearchResults(FResults& InOutResults, FResult
 	OutContext.bAbleToProgress = !bIsIdle;
 }
 
+void FAssetDataGatherer::GetDiagnostics(float& OutGatherTimeSeconds, float& OutDiscoverTimeSeconds)
+{
+	Discovery->GetDiagnostics(OutDiscoverTimeSeconds);
+	FGathererScopeLock ResultsScopeLock(&ResultsLock);
+	OutGatherTimeSeconds = CumulativeGatherTime;
+}
+
 void FAssetDataGatherer::GetPackageResults(TMultiMap<FName, FAssetData*>& OutAssetResults, TMultiMap<FName, FPackageDependencyData>& OutDependencyResults)
 {
 	FGathererScopeLock ResultsScopeLock(&ResultsLock);
@@ -3625,7 +3788,8 @@ void FAssetDataGatherer::WaitOnPath(FStringView InPath)
 	}
 	FString LocalAbsPath = NormalizeLocalPath(InPath);
 	UE::AssetDataGather::Private::FPathExistence QueryPath(LocalAbsPath);
-	Discovery->SetPropertiesAndWait(QueryPath, false /* bAddToAllowList */, false /* bForceRescan */, false /* bIgnoreDenyListScanFilters */);
+	Discovery->SetPropertiesAndWait(TArrayView<UE::AssetDataGather::Private::FPathExistence>(&QueryPath, 1),
+		false /* bAddToAllowList */, false /* bForceRescan */, false /* bIgnoreDenyListScanFilters */);
 	WaitOnPathsInternal(TArrayView<UE::AssetDataGather::Private::FPathExistence>(&QueryPath, 1), FString(), TArray<FString>());
 }
 
@@ -3691,10 +3855,7 @@ void FAssetDataGatherer::ScanPathsSynchronous(const TArray<FString>& InLocalPath
 		QueryPaths.Add(UE::AssetDataGather::Private::FPathExistence(NormalizeLocalPath(LocalPath)));
 	}
 
-	for (UE::AssetDataGather::Private::FPathExistence& QueryPath: QueryPaths)
-	{
-		Discovery->SetPropertiesAndWait(QueryPath, true /* bAddToAllowList */, bForceRescan, bIgnoreDenyListScanFilters);
-	}
+	Discovery->SetPropertiesAndWait(QueryPaths, true /* bAddToAllowList */, bForceRescan, bIgnoreDenyListScanFilters);
 
 	{
 		FGathererScopeLock ResultsScopeLock(&ResultsLock);
@@ -4258,9 +4419,22 @@ void FAssetDataGatherer::PrioritizeSearchPath(const FString& PathToPrioritize)
 	FString LocalFilenamePathToPrioritize;
 	if (FPackageName::TryConvertLongPackageNameToFilename(PathToPrioritize, LocalFilenamePathToPrioritize))
 	{
-		FSetPathProperties Properties;
-		Properties.Priority = EPriority::High;
-		SetDirectoryProperties(LocalFilenamePathToPrioritize, Properties);
+		LocalFilenamePathToPrioritize = NormalizeLocalPath(LocalFilenamePathToPrioritize);
+		if (LocalFilenamePathToPrioritize.Len() == 0)
+		{
+			return;
+		}
+		EPriority Priority = EPriority::High;
+		Discovery->PrioritizeSearchPath(LocalFilenamePathToPrioritize, Priority);
+
+		{
+			FGathererScopeLock ResultsScopeLock(&ResultsLock);
+			SetIsIdle(false);
+			int32 NumPrioritizedPaths;
+			UE::AssetDataGather::Private::FPathExistence QueryPath(PathToPrioritize);
+			SortPathsByPriority(TArrayView<UE::AssetDataGather::Private::FPathExistence>(&QueryPath, 1),
+				Priority, NumPrioritizedPaths);
+		}
 	}
 }
 
@@ -4277,13 +4451,6 @@ void FAssetDataGatherer::SetDirectoryProperties(FStringView LocalPath, const UE:
 	{
 		FGathererScopeLock ResultsScopeLock(&ResultsLock);
 		SetIsIdle(false);
-		if (InProperties.Priority.IsSet())
-		{
-			int32 NumPrioritizedPaths;
-			UE::AssetDataGather::Private::FPathExistence QueryPath(LocalAbsPath);
-			SortPathsByPriority(TArrayView<UE::AssetDataGather::Private::FPathExistence>(&QueryPath, 1),
-				*InProperties.Priority, NumPrioritizedPaths);
-		}
 	}
 }
 
@@ -4345,6 +4512,12 @@ bool FAssetDataGatherer::IsVerseFile(FStringView FilePath)
 }
 
 void FAssetDataGatherer::SetIsIdle(bool bInIsIdle)
+{
+	double TickStartTime = -1.;
+	SetIsIdle(bInIsIdle, TickStartTime);
+}
+
+void FAssetDataGatherer::SetIsIdle(bool bInIsIdle, double& TickStartTime)
 {
 	CHECK_IS_LOCKED_CURRENT_THREAD(ResultsLock);
 	if (bInIsIdle == bIsIdle)
@@ -4356,14 +4529,21 @@ void FAssetDataGatherer::SetIsIdle(bool bInIsIdle)
 	if (bIsIdle)
 	{
 		// bIsComplete will be set in GetAndTrimSearchResults
-		double SearchTime = FPlatformTime::Seconds() - SearchStartTime;
+		if (TickStartTime >= 0.)
+		{
+			CurrentSearchTime += FPlatformTime::Seconds() - TickStartTime;
+			TickStartTime = -1.;
+		}
 		if (!bFinishedInitialDiscovery)
 		{
 			bFinishedInitialDiscovery = true;
 
-			UE_LOG(LogAssetRegistry, Verbose, TEXT("Initial scan took %0.6f seconds (found %d cached assets, and loaded %d)"), (float)SearchTime, NumCachedAssetFiles, NumUncachedAssetFiles);
+			UE_LOG(LogAssetRegistry, Verbose, TEXT("Initial scan took %0.6f seconds (found %d cached assets, and loaded %d)"),
+				(float)CurrentSearchTime, NumCachedAssetFiles, NumUncachedAssetFiles);
 		}
-		SearchTimes.Add(SearchTime);
+		SearchTimes.Add(CurrentSearchTime);
+		CumulativeGatherTime += static_cast<float>(CurrentSearchTime);
+		CurrentSearchTime = 0.;
 	}
 	else
 	{
diff --git a/Engine/Source/Runtime/AssetRegistry/Private/AssetDataGatherer.h b/Engine/Source/Runtime/AssetRegistry/Private/AssetDataGatherer.h
--- a/Engine/Source/Runtime/AssetRegistry/Private/AssetDataGatherer.h
+++ b/Engine/Source/Runtime/AssetRegistry/Private/AssetDataGatherer.h
@@ -119,6 +119,8 @@ public:
 	};
 	/** Gets search results from the data gatherer. */
 	void GetAndTrimSearchResults(FResults& InOutResults, FResultContext& OutContext);
+	/** Get diagnostics for telemetry or logging. */
+	void GetDiagnostics(float& OutGatherTimeSeconds, float& OutDiscoverTimeSeconds);
 	/** Gets just the AssetResults and DependencyResults from the data gatherer. */
 	void GetPackageResults(TMultiMap<FName, FAssetData*>& OutAssetResults,
 		TMultiMap<FName, FPackageDependencyData>& OutDependencyResults);
@@ -232,7 +234,7 @@ private:
 	 * Tick function to pump scanning and push results into the search results structure. May be called from devoted
 	 * thread or inline from synchronous functions on other threads.
 	 */
-	void TickInternal(bool& bOutIsIdle);
+	void TickInternal(bool& bOutIsIdle, double& TickStartTime);
 	/** Add any new package files from the background directory scan to our work list **/
 	void IngestDiscoveryResults();
 
@@ -301,6 +303,7 @@ private:
 	 * when they note a possible state change. Caller is responsible for holding the ResultsLock.
 	 */
 	void SetIsIdle(bool IsIdle);
+	void SetIsIdle(bool IsIdle, double& TickStartTime);
 
 	/** Minimize memory usage in the buffers used during gathering. */
 	void Shrink();
@@ -386,12 +389,14 @@ private:
 
 	/** All the search times since the last call to GetAndTrimSearchResults. */
 	TArray<double> SearchTimes;
+	/** Sum of all SearchTimes. */
+	float CumulativeGatherTime = 0.f;
 
 	/** The directories found during the search, unless they are hidden by DirLongPackageNamesToNotReport. */
 	TArray<FString> DiscoveredPaths;
 
-	/** The search start time, set after resuming from idle and used for performance metrics when reporting results. */
-	double SearchStartTime;
+	/** The time spent in TickInternal since the last idle time. Used for performance metrics when reporting results. */
+	double CurrentSearchTime = 0.;
 	/** The last time at which the cache file was written, used to periodically update the cache. */
 	double LastCacheWriteTime;
 	/** The cached value of the NumPathsToSearch returned by Discovery the last time we synchronized with it. */
diff --git a/Engine/Source/Runtime/AssetRegistry/Private/AssetDataGathererPrivate.h b/Engine/Source/Runtime/AssetRegistry/Private/AssetDataGathererPrivate.h
--- a/Engine/Source/Runtime/AssetRegistry/Private/AssetDataGathererPrivate.h
+++ b/Engine/Source/Runtime/AssetRegistry/Private/AssetDataGathererPrivate.h
@@ -20,6 +20,7 @@ namespace Private
 {
 
 class FMountDir;
+struct FScanDirAndParentData;
 
 /** Enum to specify files and directories that should be completed earlier than others */
 enum class EPriority : uint8
@@ -34,7 +35,6 @@ enum class EPriority : uint8
 	Highest = Blocking,
 	Lowest = Normal,
 };
-constexpr uint32 CountEPriority = static_cast<uint32>(EPriority::Normal) + 1;
 
 /** Fields being set in a call to SetProperty */
 struct FSetPathProperties
@@ -54,17 +54,12 @@ struct FSetPathProperties
 	 * IFileManager::IterateDirectoryStat after process start or the last request to rescan it
 	 */
 	TOptional<bool> HasScanned;
-	/**
-	 * Priority that determines which directories are scanned first, for when a directory is being waited on or are
-	 * not blocking but are wanted for optional features
-	 */
-	TOptional<EPriority> Priority;
 
 	/** Used to early-exit from tree traversal when all properties have finished being handled */
 	bool IsSet() const
 	{
 		return IsOnAllowList.IsSet() | HasScanned.IsSet() | MatchesDenyList.IsSet() | //-V792
-			IgnoreDenyList.IsSet() | Priority.IsSet(); //-V792
+			IgnoreDenyList.IsSet(); //-V792
 	}
 };
 
@@ -314,9 +309,6 @@ public:
 	/** Get this ScanDir's RelPath from its Parent */
 	FStringView GetRelPath() const;
 
-	/** Return this ScanDir's accumulated Priority: maximum of its directly-set priority and its children's priority. */
-	EPriority GetPriority() const;
-
 	/** Calculate this ScanDir's full absolute path by accumulating RelPaths from parents and append it. */
 	void AppendLocalAbsPath(FStringBuilderBase& OutFullPath) const;
 	/** Calculate this ScanDir's full absolute path by accumulating RelPaths from parents and return it as a string. */
@@ -356,16 +348,13 @@ public:
 	 * Returns whether the directory was found and its property was changed; returns false if InRelPath was not a
 	 * directory or the property did not need to be changed.
 	 */
-	bool TrySetDirectoryProperties(FStringView InRelPath,
-		const FSetPathProperties& InProperties, bool bConfirmedExists);
+	bool TrySetDirectoryProperties(FStringView InRelPath, FInherited& ParentData,
+		const FSetPathProperties& InProperties, bool bConfirmedExists, FScanDirAndParentData& OutControllingDir);
 	/**
 	 * Mark that the given file has already been scanned, so that it will not be double reporting in the upcoming
 	 * directory scan, if one is upcoming.
 	 */
 	void MarkFileAlreadyScanned(FStringView BaseName);
-	/** Set the direct priority of this ScanDir and update the accumulated priority to match it. */
-	void SetDirectPriority(EPriority InPriority);
-	EPriority GetDirectPriority() const;
 
 	/**
 	 * Called from the Tick; handle the list of subdirs and files that were found from IterateDirectoryStat called
@@ -374,12 +363,10 @@ public:
 	void SetScanResults(FStringView LocalAbsPath, const FInherited& ParentData,
 		TArrayView<FDiscoveredPathData>& InOutSubDirs, TArrayView<FDiscoveredPathData>& InOutFiles);
 	/**
-	 * Update the completion state of this ScanDir based on its scan status and its child dirs completion state.
-	 * Set OutCursor to the highest priority child dir that needs to be updated if any children need to be updated
-	 * and/or scanned first.
-	 * Set OutCursor to the parent ScanDir (or null if there is no parent) if this ScanDir is now complete.
+	 * Update the completion state of this ScanDir and all ScanDirs under it based on each dir's scan status and its
+	 * child dirs' completion state. Add any ScanDirs that need to be scanned to OutScanRequests.
 	 */
-	void Update(FScanDir*& OutCursor, FInherited& InOutParentData);
+	void Update(TArray<FScanDirAndParentData>& OutScanRequests, const FScanDir::FInherited& ParentData);
 
 	FScanDir* GetFirstIncompleteScanDir();
 
@@ -432,9 +419,6 @@ protected:
 	/** Find the index of the subdir with the given Relative path. */
 	int32 FindLowerBoundSubDir(FStringView SubDirBaseName);
 
-	/** Return the the highest-priority not-yet-complete SubDir, or null if all are complete. */
-	FScanDir* FindHighestPrioritySubDir();
-
 	/** Call the given lambda void(FScanDir&) on each existing SubDir. */
 	template <typename CallbackType> void ForEachSubDir(const CallbackType& Callback);
 	/**
@@ -444,19 +428,11 @@ protected:
 	 */
 	template <typename CallbackType> void ForEachDescendent(const CallbackType& Callback);
 
-	/** Set the AccumulatedPriority of this ScanDir to the maximum of the DirectPriorities of it and its child dirs. */
-	void UpdateAccumulatedPriority();
-	/** Record a changed priority of a direct child ScanDir and update the accumulated priority of this ScanDir. */
-	void OnChildPriorityChanged(EPriority InPriority, int32 Delta);
-
 	TArray<TRefCountPtr<FScanDir>> SubDirs; // Sorted
 	TArray<FString> AlreadyScannedFiles; // Unsorted
 	FMountDir* MountDir = nullptr;
 	FScanDir* Parent = nullptr;
 	FString RelPath;
-	uint8 PriorityRefCounts[CountEPriority] = {}; // Initialize all elements to 0
-	EPriority DirectPriority = EPriority::Normal;
-	EPriority AccumulatedPriority = EPriority::Normal;
 	/** Whether each piece of the inherited data has been set directly on this directory */
 	FInherited DirectData;
 	bool bHasScanned = false;
@@ -465,6 +441,16 @@ protected:
 	bool bIsComplete = false;
 };
 
+/**
+ *  A refcounted pointer to a ScanDir, along with the ParentData found when tracing down from the root
+ * path to the ScanDir. Used to return list of directories needing scanning from Update.
+ */
+struct FScanDirAndParentData
+{
+	TRefCountPtr<FScanDir> ScanDir;
+	FScanDir::FInherited ParentData;
+};
+
 /**
  * Gather data about a MountPoint that has been registered with FPackageName
  * The FMountDir holds a FScanTree with information about each directory (that is pruned when not in use).
@@ -496,8 +482,6 @@ public:
 
 	/** Report whether this MountDir is complete: all ScanDirs under it either have scanned or should not scan. */
 	bool IsComplete() const;
-	/** Return this MountDir's Priority, which is the maximum of any of its ScanDirs' priorities. */
-	EPriority GetPriority() const;
 
 	/** Report the collapsed data for the scandir - allow list, deny list, etc. Returns false data for non child paths. */
 	void GetMonitorData(FStringView InLocalAbsPath, FScanDir::FInherited& OutData) const;
@@ -513,14 +497,9 @@ public:
 	 * directory under this MountDir or the property did not need to be changed.
 	 */
 	bool TrySetDirectoryProperties(FStringView LocalAbsPath, const FSetPathProperties& InProperties,
-		bool bConfirmedExists);
-	/** 
-	 * Update all ScanDirs under this MountDir in priority order until one is found that needs to be scanned.
-	 * Set OutCursor to that ScanDir, or to nullptr if all are complete.
-	 * Set OutData for the inherited parent data of the OutCursor, set to false data if 
-	 * if OutCursor is null.
-	 */
-	void Update(FScanDir*& OutCursor, FScanDir::FInherited& OutParentData);
+		bool bConfirmedExists, FScanDirAndParentData* OutControllingDir);
+	/** Update all incomplete ScanDirs under this MountDir and add any that need to be scanned to OutScanRequests. */
+	void Update(TArray<FScanDirAndParentData>& OutScanRequests);
 
 	FScanDir* GetFirstIncompleteScanDir();
 
@@ -616,10 +595,12 @@ public:
 	/** Gets search results from the file discovery. */
 	void GetAndTrimSearchResults(bool& bOutIsComplete, TArray<FString>& OutDiscoveredPaths,
 		FFilesToSearch& OutFilesToSearch, int32& OutNumPathsToSearch);
+	/** Get diagnostics for telemetry or logging. */
+	void GetDiagnostics(float& OutCumulativeDiscoveryTime);
 	/** Wait (joining in on the tick) until all currently monitored paths have been scanned. */
 	void WaitForIdle();
-	/** Optionally set some scan properties for the given path and then wait for the scan of it to finish. */
-	void SetPropertiesAndWait(FPathExistence& QueryPath, bool bAddToAllowList, bool bForceRescan,
+	/** Optionally set some scan properties for the given paths and then wait for their scans to finish. */
+	void SetPropertiesAndWait(TArrayView<FPathExistence> QueryPaths, bool bAddToAllowList, bool bForceRescan,
 		bool bIgnoreDenyListScanFilters);
 	/** Return whether the given path is allowed due to e.g. TrySetDirectoryProperties with IsOnAllowList. */
 	bool IsOnAllowList(FStringView LocalAbsPath) const;
@@ -638,6 +619,8 @@ public:
 	void AddMountPoint(const FString& LocalAbsPath, FStringView LongPackageName);
 	/** Remove the mountpoint because FPackageName has removed it. */
 	void RemoveMountPoint(const FString& LocalAbsPath);
+	/** Raise the priority until completion of scans of the given path and its subdirs. */
+	void PrioritizeSearchPath(const FString& LocalAbsPath, EPriority Priority);
 	/** Set properties on the directory, called when files are requested to be on an allow/deny list or rescanned. */
 	bool TrySetDirectoryProperties(const FString& LocalAbsPath,
 		const UE::AssetDataGather::Private::FSetPathProperties& Properties, bool bConfirmedExists);
@@ -661,13 +644,12 @@ private:
 	int32 FindLowerBoundMountPoint(FStringView LocalAbsPath) const;
 
 	/** Run the tick, either called from the async Run or called on thread from a thread executing a synchronous wait. */
-	void TickInternal();
-	/** Search all of the MountDirs for a highest priority ScanDir that needs to be Updated. */
-	void FindFirstCursor(FScanDir*& OutCursor, FScanDir::FInherited& OutParentData);
-	/** Invalidate the tick's current cursor, called when we need to reconsider whether directories ShouldBeScanned. */
-	void InvalidateCursor();
+	void TickInternal(bool bTickAll);
+	/** Update all incomplete ScanDirs under all MountDirs and add any that need to be scanned to OutScanRequests. */
+	void UpdateAll(TArray<FScanDirAndParentData>& OutScanRequests);
 	/** Mark whether this discoverer has finished and is idle. Update properties dependent upon the idle state. */
 	void SetIsIdle(bool bInIdle);
+	void SetIsIdle(bool bInIdle, double& TickStartTime);
 
 	/** Store the given discovered files and directories in the results. */
 	void AddDiscovered(FStringView DirAbsPath, TConstArrayView<FDiscoveredPathData> SubDirs,
@@ -710,20 +692,19 @@ private:
 	};
 
 private:
-	/**
-	 * Prevent simultaneous ticks from two different threads and protect access to Tick-specific data.
-	 * To prevent DeadLocks, TickLock can not be entered from within any of the other locks on this class.
-	 */
-	mutable FGathererCriticalSection TickLock;
 	/**
 	 * Protect access to data in the ScanDir tree which can be read/write from the tick or from SetProperties.
-	 * To prevent DeadLocks, TreeLock can be entered while holding TickLock, but can not be entered from
-	 * within any of the other locks on this class.
+	 * To prevent DeadLocks, TreeLock must not be entered while holding ResultsLock.
 	 */
 	mutable FGathererCriticalSection TreeLock;
+	/**
+	 * Prevent simultaneous ticks from two different threads and protect access to Tick-specific data.
+	 * Can only be queried or taken while holding TreeLock, but can continue to be held after exiting TreeLock.
+	 */
+	mutable FThreadOwnerSection TickOwner;
 	/**
 	 * Protect access to the data written from tick and read/written from GetAndTrimSearchResults.
-	 * ResultsLock can be entered while holding TickLock or TreeLock or both.
+	 * ResultsLock can be entered while holding TreeLock.
 	 */
 	mutable FGathererCriticalSection ResultsLock;
 
@@ -749,6 +730,11 @@ private:
 	 * Readable anywhere. Writable only within TreeLock.
 	 */
 	std::atomic<bool> bIsIdle;
+	/**
+	 * Set to true after editing PriorityScanDirs to inform a running Tick that it should abort its
+	 * current batch to reduce latency in handling the new priorities.
+	 */
+	std::atomic<bool> bPriorityDirty;
 	/** > 0 if we've been asked to abort work in progress at the next opportunity. */
 	std::atomic<uint32> IsStopped;
 	/** > 0 if we've been asked to pause the worker thread so a synchronous function can take over the tick. */
@@ -758,7 +744,8 @@ private:
 	 * Used for progress tracking.
 	 */
 	FThreadSafeCounter NumDirectoriesToScan;
-
+	/** Triggered after scanning PriorityDirs to inform threads waiting on a PriorityDir to recheck conditions. */
+	FEventRef PriorityDataUpdated{ EEventMode::ManualReset };
 
 	// Variable section for variables that are read/writable only within ResultsLock.
 
@@ -775,11 +762,14 @@ private:
 	};
 	TArray<FDirectoryResult> DiscoveredFiles;
 	TArray<FGatheredPathData> DiscoveredSingleFiles;
-	/** Time at which the scan was started or last resumed from idle. Used for logging. */
-	double DiscoverStartTime = 0.;
+	/** Time spent in TickInternal since the last time cook was idle. Used for logging. */
+	double CurrentDiscoveryTime = 0.;
 	/** Number of files discovered during scanning since start or resumed from idle. Used for logging. */
 	int32 NumDiscoveredFiles = 0;
-
+	/** Cumulative total of NumDiscoveredFiles across all non-idle periods. */
+	int32 CumulativeDiscoveredFiles = 0;
+	/** Cumulative total of time spent in all non-idle periods. */
+	float CumulativeDiscoveryTime = 0.f;
 
 	// Variable section for variables that are read/writable only within TreeLock.
 
@@ -789,19 +779,51 @@ private:
 	 * Read/writable only with TreeLock, both the list and all data owned by each MountDir.
 	 */
 	TArray<TUniquePtr<FMountDir>> MountDirs;
-	/** The next ScanDir to update in Tick. */
-	TRefCountPtr<FScanDir> Cursor = nullptr;
-	/** The value of Parent's InheritedData for Cursor (or false if Cursor is null). */
-	FScanDir::FInherited CursorParentData;
 
+	/** A ScanDir and referencecount for directories that have been prioritized by ScanPathsSynchronous. */
+	struct FPriorityScanDirData
+	{
+		TRefCountPtr<FScanDir> ScanDir;
+		FScanDir::FInherited ParentData;
+		int32 RequestCount = 0;
+		/**
+		 * Fire and forget prioritization increments RequestCount without ever decrementing it, and
+		 * sets bReleaseWhenComplete to have it be decremented by TickInternal.
+		 */
+		bool bReleaseWhenComplete = false;
+	};
+	TArray<FPriorityScanDirData> PriorityScanDirs;
+
+	// Variable section for variables that are read/writable only by the TickOwner thread
 
-	// Variable section for variables that are read/writable only within TickLock.
+	/** Scratch space to store scan results during the tick to avoid allocations. */
+	struct FDirToScanData
+	{
+		TRefCountPtr<FScanDir> ScanDir;
+		FScanDir::FInherited ParentData;
+		// This type is used in a TArray, and TArray assumes elements can be memmoved. 
+		// TStringBuilder<N> cannot be memmoved, so use FStringBuilderBase without a buffer.
+		FStringBuilderBase DirLocalAbsPath;
+		FStringBuilderBase DirLongPackageName;
+		TArray<FDiscoveredPathData> IteratedSubDirs;
+		TArray<FDiscoveredPathData> IteratedFiles;
+		int32 NumIteratedDirs = 0;
+		int32 NumIteratedFiles = 0;
+		bool bScanned = false;
+
+		void Reset();
+		SIZE_T GetAllocatedSize() const;
+	};
+	TArray<FDirToScanData> DirToScanDatas;
 
-	/** Scratch space to store discovered subdirs during the tick, to avoid allocations. */
-	TArray<FDiscoveredPathData> IteratedSubDirs;
-	/** Scratch space to store discovered files during the tick, to avoid allocations. */
-	TArray<FDiscoveredPathData> IteratedFiles;
+	/** Scratch space used per helper thread in the forloop to avoid allocations. */
+	struct FDirToScanBuffer
+	{
+		bool bAbort = false;
 
+		void Reset();
+	};
+	TArray<FDirToScanBuffer> DirToScanBuffers;
 
 	friend class FMountDir;
 	friend class FScanDir;
diff --git a/Engine/Source/Runtime/AssetRegistry/Private/AssetRegistry.cpp b/Engine/Source/Runtime/AssetRegistry/Private/AssetRegistry.cpp
--- a/Engine/Source/Runtime/AssetRegistry/Private/AssetRegistry.cpp
+++ b/Engine/Source/Runtime/AssetRegistry/Private/AssetRegistry.cpp
@@ -771,8 +771,7 @@ void FAssetRegistryImpl::Initialize(Impl::FInitializeContext& Context)
 	bInitialSearchCompleted = true;
 	GatherStatus = Impl::EGatherStatus::Active;
 	bSearchAllAssets = false;
-	AmortizeStartTime = 0;
-	TotalAmortizeTime = 0;
+	StoreGatherResultsTimeSeconds = 0.f;
 
 	// By default update the disk cache once on asset load, to incorporate changes made in PostLoad. This only happens in editor builds
 	bUpdateDiskCacheAfterLoad = true;
@@ -1532,7 +1531,6 @@ void FAssetRegistryImpl::SearchAllAssetsInitialAsync(Impl::FEventContext& EventC
 {
 	bInitialSearchStarted = true;
 	bInitialSearchCompleted = false;
-	FullSearchStartTime = FPlatformTime::Seconds();
 	SearchAllAssets(EventContext, InheritanceContext, false /* bSynchronousSearch */);
 }
 
@@ -1555,6 +1553,10 @@ void UAssetRegistryImpl::SearchAllAssets(bool bSynchronousSearch)
 			GuardedData.ConditionalLoadPremadeAssetRegistry(*this, EventContext, InterfaceScopeLock);
 		}
 		GuardedData.SearchAllAssets(EventContext, InheritanceContext, bSynchronousSearch);
+		if (bSynchronousSearch)
+		{
+			GuardedData.LogSearchDiagnostics();
+		}
 	}
 #if WITH_EDITOR
 	if (bSynchronousSearch)
@@ -3857,6 +3859,26 @@ Impl::EGatherStatus FAssetRegistryImpl::TickGatherer(Impl::FEventContext& EventC
 	{
 		return OutStatus;
 	}
+	double TimingStartTime = -1.;
+	auto LazyStartTimer = [&TimingStartTime]()
+	{
+		if (TimingStartTime <= 0.)
+		{
+			TimingStartTime = FPlatformTime::Seconds();
+		}
+	};
+	auto RecordTimer = [this, &TimingStartTime]()
+	{
+		if (TimingStartTime > 0.)
+		{
+			StoreGatherResultsTimeSeconds += static_cast<float>(FPlatformTime::Seconds() - TimingStartTime);
+			TimingStartTime = -1.;
+		}
+	};
+	ON_SCOPE_EXIT
+	{
+		RecordTimer();
+	};
 
 	// Gather results from the background search
 	FAssetDataGatherer::FResultContext ResultContext;
@@ -3908,44 +3930,37 @@ Impl::EGatherStatus FAssetRegistryImpl::TickGatherer(Impl::FEventContext& EventC
 		this->GatherStatus = OutStatus;
 	};
 
-
 	// Add discovered paths
 	if (BackgroundResults.Paths.Num())
 	{
+		LazyStartTimer();
 		PathDataGathered(EventContext, TickStartTime, BackgroundResults.Paths);
 	}
 
 	// Process the asset results
 	if (BackgroundResults.Assets.Num())
 	{
+		LazyStartTimer();
 		// Mark the first amortize time
-		if (AmortizeStartTime == 0)
-		{
-			AmortizeStartTime = FPlatformTime::Seconds();
-		}
 		if (AssetsFoundCallback.IsSet())
 		{
 			AssetsFoundCallback.GetValue()(BackgroundResults.Assets);
 		}
 
 		AssetSearchDataGathered(EventContext, TickStartTime, BackgroundResults.Assets);
-
-		if (BackgroundResults.Assets.Num() == 0)
-		{
-			TotalAmortizeTime += FPlatformTime::Seconds() - AmortizeStartTime;
-			AmortizeStartTime = 0;
-		}
 	}
 
 	// Add dependencies
 	if (BackgroundResults.Dependencies.Num())
 	{
+		LazyStartTimer();
 		DependencyDataGathered(TickStartTime, BackgroundResults.Dependencies);
 	}
 
 	// Load cooked packages that do not have asset data
 	if (BackgroundResults.CookedPackageNamesWithoutAssetData.Num())
 	{
+		LazyStartTimer();
 		CookedPackageNamesWithoutAssetDataGathered(EventContext, TickStartTime, BackgroundResults.CookedPackageNamesWithoutAssetData, bOutInterrupted);
 		if (bOutInterrupted)
 		{
@@ -3957,6 +3972,7 @@ Impl::EGatherStatus FAssetRegistryImpl::TickGatherer(Impl::FEventContext& EventC
 	// Add Verse files
 	if (BackgroundResults.VerseFiles.Num())
 	{
+		LazyStartTimer();
 		VerseFilesGathered(TickStartTime, BackgroundResults.VerseFiles);
 	}
 
@@ -3966,6 +3982,7 @@ Impl::EGatherStatus FAssetRegistryImpl::TickGatherer(Impl::FEventContext& EventC
 	bool bDiskGatherComplete = !ResultContext.bIsSearching && NumGatherFromDiskPending == 0;
 	if (bDiskGatherComplete && PackagesNeedingDependencyCalculation.Num())
 	{
+		LazyStartTimer();
 		LoadCalculatedDependencies(nullptr, TickStartTime, InheritanceContext, bOutInterrupted);
 		if (bOutInterrupted)
 		{
@@ -3989,8 +4006,8 @@ Impl::EGatherStatus FAssetRegistryImpl::TickGatherer(Impl::FEventContext& EventC
 			// update redirectors
 			UpdateRedirectCollector();
 #endif
-			UE_LOG(LogAssetRegistry, Verbose, TEXT("### Time spent amortizing search results: %0.4f seconds"), TotalAmortizeTime);
-			UE_LOG(LogAssetRegistry, Log, TEXT("Asset discovery search completed in %0.4f seconds"), FPlatformTime::Seconds() - FullSearchStartTime);
+			RecordTimer();
+			LogSearchDiagnostics();
 
 			bInitialSearchCompleted = true;
 
@@ -4001,6 +4018,16 @@ Impl::EGatherStatus FAssetRegistryImpl::TickGatherer(Impl::FEventContext& EventC
 	return OutStatus;
 }
 
+void FAssetRegistryImpl::LogSearchDiagnostics() const
+{
+	float GatherTimeSeconds;
+	float DiscoveryTimeSeconds;
+	GlobalGatherer->GetDiagnostics(GatherTimeSeconds, DiscoveryTimeSeconds);
+	float Total = DiscoveryTimeSeconds + GatherTimeSeconds + StoreGatherResultsTimeSeconds;
+	UE_LOG(LogAssetRegistry, Log, TEXT("AssetRegistryGather time %.4fs: AssetDataDiscovery %0.4fs, AssetDataGather %0.4fs, StoreResults %0.4fs."),
+		Total, DiscoveryTimeSeconds, GatherTimeSeconds, StoreGatherResultsTimeSeconds);
+}
+
 void FAssetRegistryImpl::TickGatherPackage(Impl::FEventContext& EventContext, const FString& PackageName, const FString& LocalPath)
 {
 	if (!GlobalGatherer.IsValid())
@@ -4008,6 +4035,22 @@ void FAssetRegistryImpl::TickGatherPackage(Impl::FEventContext& EventContext, co
 		return;
 	}
 	GlobalGatherer->WaitOnPath(LocalPath);
+	double TimingStartTime = -1.;
+	auto LazyStartTimer = [&TimingStartTime]()
+	{
+		if (TimingStartTime <= 0.)
+		{
+			TimingStartTime = FPlatformTime::Seconds();
+		}
+	};
+	ON_SCOPE_EXIT
+	{
+		if (TimingStartTime > 0.)
+		{
+			StoreGatherResultsTimeSeconds += static_cast<float>(FPlatformTime::Seconds() - TimingStartTime);
+			TimingStartTime = -1.;
+		}
+	};
 
 	FName PackageFName(PackageName);
 
@@ -4022,6 +4065,7 @@ void FAssetRegistryImpl::TickGatherPackage(Impl::FEventContext& EventContext, co
 	BackgroundResults.Dependencies.Remove(PackageFName);
 	if (PackageAssets.Num() > 0)
 	{
+		LazyStartTimer();
 		TMultiMap<FName, FAssetData*> PackageAssetsMap;
 		PackageAssetsMap.Reserve(PackageAssets.Num());
 		for (FAssetData* PackageAsset : PackageAssets)
@@ -4032,6 +4076,7 @@ void FAssetRegistryImpl::TickGatherPackage(Impl::FEventContext& EventContext, co
 	}
 	if (PackageDependencyDatas.Num() > 0)
 	{
+		LazyStartTimer();
 		TMultiMap<FName, FPackageDependencyData> PackageDependencyDatasMap;
 		PackageDependencyDatasMap.Reserve(PackageDependencyDatas.Num());
 		for (FPackageDependencyData& DependencyData : PackageDependencyDatas)
diff --git a/Engine/Source/Runtime/AssetRegistry/Private/AssetRegistryImpl.h b/Engine/Source/Runtime/AssetRegistry/Private/AssetRegistryImpl.h
--- a/Engine/Source/Runtime/AssetRegistry/Private/AssetRegistryImpl.h
+++ b/Engine/Source/Runtime/AssetRegistry/Private/AssetRegistryImpl.h
@@ -188,6 +188,8 @@ public:
 	Impl::EGatherStatus TickGatherer(Impl::FEventContext& EventContext,
 		Impl::FClassInheritanceContext& InheritanceContext, const double TickStartTime, bool& bOutInterrupted,
 		TOptional<FAssetsFoundCallback> AssetsFoundCallback = TOptional<FAssetsFoundCallback>());
+	/** Send a log message with the search statistics. */
+	void LogSearchDiagnostics() const;
 	/** Look for and load a single AssetData result from the gatherer. */
 	void TickGatherPackage(Impl::FEventContext& EventContext, const FString& PackageName, const FString& LocalPath);
 	void ClearGathererCache();
@@ -386,10 +388,8 @@ private:
 	UE_DEPRECATED(5.0, "DelayDelete is only intended to support deprecated functions.")
 	TArray<TUniqueFunction<void()>> DeleteActions;
 
-	/** The start time of the full asset search */
-	double FullSearchStartTime;
-	double AmortizeStartTime;
-	double TotalAmortizeTime;
+	/** Time spent processing Gather results */
+	float StoreGatherResultsTimeSeconds;
 	/** The highest number of pending results observed during initial gathering */
 	int32 HighestPending = 0;
 
diff --git a/Engine/Source/Runtime/AssetRegistry/Private/CriticalSectionQueryable.h b/Engine/Source/Runtime/AssetRegistry/Private/CriticalSectionQueryable.h
--- a/Engine/Source/Runtime/AssetRegistry/Private/CriticalSectionQueryable.h
+++ b/Engine/Source/Runtime/AssetRegistry/Private/CriticalSectionQueryable.h
@@ -4,6 +4,7 @@
 
 #include "HAL/CriticalSection.h"
 #include "HAL/PlatformTLS.h"
+#include "Misc/ScopeLock.h"
 
 #include <atomic>
 
@@ -103,3 +104,89 @@ private:
 	// Holds the synchronization object to aggregate and scope manage.
 	FCriticalSectionQueryable* SynchObject;
 };
+
+/**
+ * Like a critical section: only a single thread can posess it at once. But provides a different API than a critical
+ * section to more flexibly interact with another critical section. With multiple normal critical sections they must
+ * be entered in the same order to prevent deadlocks. A ThreadOwnerSection is instead tested and claimed for ownership
+ * only while inside its collaborating critical section (called the AegisLock), and ownership can then continue to be
+ * held even after releasing the AegisLock.
+ */
+class FThreadOwnerSection
+{
+public:
+	bool TryTakeOwnership(FCriticalSectionQueryable& NotYetEnteredAegisLock)
+	{
+		FScopeLockQueryable ScopeLock(&NotYetEnteredAegisLock);
+		return TryTakeOwnershipWithinLock();
+	}
+	bool TryTakeOwnership(FCriticalSection& NotYetEnteredAegisLock)
+	{
+		FScopeLock ScopeLock(&NotYetEnteredAegisLock);
+		return TryTakeOwnershipWithinLock();
+	}
+	bool TryTakeOwnership(FScopeLockQueryable& AlreadyEnteredAegisScopeLock)
+	{
+		return TryTakeOwnershipWithinLock();
+	}
+	bool TryTakeOwnership(FScopeLock& AlreadyEnteredAegisScopeLock)
+	{
+		return TryTakeOwnershipWithinLock();
+	}
+
+	void ReleaseOwnershipChecked(FCriticalSectionQueryable& NotYetEnteredAegisLock)
+	{
+		FScopeLockQueryable ScopeLock(&NotYetEnteredAegisLock);
+		ReleaseOwnershipCheckedWithinLock();
+	}
+	void ReleaseOwnershipChecked(FCriticalSection& NotYetEnteredAegisLock)
+	{
+		FScopeLock ScopeLock(&NotYetEnteredAegisLock);
+		ReleaseOwnershipCheckedWithinLock();
+	}
+	void ReleaseOwnershipChecked(FScopeLockQueryable& AlreadyEnteredAegisScopeLock)
+	{
+		ReleaseOwnershipCheckedWithinLock();
+	}
+	void ReleaseOwnershipChecked(FScopeLock& AlreadyEnteredAegisScopeLock)
+	{
+		ReleaseOwnershipCheckedWithinLock();
+	}
+
+	bool IsOwned(FScopeLockQueryable& AlreadyEnteredAegisScopeLock) const
+	{
+		return bHasOwner;
+	}
+	bool IsOwned(FScopeLock& AlreadyEnteredAegisScopeLock) const
+	{
+		return bHasOwner;
+	}
+
+	bool IsOwnedByCurrentThread() const
+	{
+		return OwnerThreadId.load(std::memory_order_relaxed) == FPlatformTLS::GetCurrentThreadId();
+	}
+
+private:
+	bool TryTakeOwnershipWithinLock()
+	{
+		if (!bHasOwner)
+		{
+			bHasOwner = true;
+			OwnerThreadId.store(FPlatformTLS::GetCurrentThreadId(), std::memory_order_relaxed);
+			return true;
+		}
+		return false;
+	}
+	void ReleaseOwnershipCheckedWithinLock()
+	{
+		check(IsOwnedByCurrentThread());
+		bHasOwner = false;
+		OwnerThreadId.store(0, std::memory_order_relaxed);
+	}
+
+private:
+	/** OwnerThreadId can be checked while outside of the lock so it is an atomic to prevent dataraces. */
+	std::atomic<uint32> OwnerThreadId{ 0 };
+	bool bHasOwner = false;
+};
\ No newline at end of file
