Remove unused header

diff --git a/Engine/Source/Runtime/AssetRegistry/Public/AssetRegistry/IAssetRegistry.h b/Engine/Source/Runtime/AssetRegistry/Public/AssetRegistry/IAssetRegistry.h
--- a/Engine/Source/Runtime/AssetRegistry/Public/AssetRegistry/IAssetRegistry.h
+++ b/Engine/Source/Runtime/AssetRegistry/Public/AssetRegistry/IAssetRegistry.h
@@ -4,7 +4,6 @@
 
 #include "AssetRegistry/ARFilter.h"
 #include "AssetRegistry/AssetIdentifier.h"
-#include "AssetRegistry/AssetDataTagMap.h"
 #include "Containers/Array.h"
 #include "Containers/BitArray.h"
 #include "Containers/StringFwd.h"
