AssetDataGatherer: Remove TFuture.Wait that can occur during c++ global shutdown; waiting on a future is invalid at that point because systems used by TFuture.Wait are shutdown at engine exit.
#rnx

diff --git a/Engine/Source/Runtime/AssetRegistry/Private/AssetDataGatherer.cpp b/Engine/Source/Runtime/AssetRegistry/Private/AssetDataGatherer.cpp
--- a/Engine/Source/Runtime/AssetRegistry/Private/AssetDataGatherer.cpp
+++ b/Engine/Source/Runtime/AssetRegistry/Private/AssetDataGatherer.cpp
@@ -3044,7 +3044,10 @@ public:
 
 	~FPreloader()
 	{
-		PreloadReady.Wait();
+		// This destructor is called during c++ global shutdown after main, and it is not valid to call TFuture::Wait
+		// after Engine shutdown. TaskGraph has already shutdown at this point, so we do not need to worry about the 
+		// Async thread still running and accessing *this.
+		// PreloadReady.Wait();
 	}
 
 	FCachePayload Consume()
