AssetRegistry: Add Engine.ini:[AssetRegistry]:CommandletsUsingAR to identify commandlets that should load the AssetRegistry at startup.
#rn Minor, Cooking

diff --git a/Engine/Config/BaseEngine.ini b/Engine/Config/BaseEngine.ini
--- a/Engine/Config/BaseEngine.ini
+++ b/Engine/Config/BaseEngine.ini
@@ -2472,6 +2472,10 @@ bFilterAssetDataWithNoTags=false
 bFilterDependenciesWithNoTags=false
 bFilterSearchableNames=true
 
+; Commandlets by default do not SearchAllAssets because many do not need it. A commandlet can be opted in
+; using CommandletsUsingAR. Format is the shortname of the commandlet, optionally without the "Commandlet" suffix.
++CommandletsUsingAR=DataValidation
+
 [AutomationTesting]
 ImportTestPath=Content/EditorAutomation/
 ImportTestPackagePath=/Engine/Content/EditorAutomation
diff --git a/Engine/Source/Runtime/AssetRegistry/Private/AssetDataGatherer.cpp b/Engine/Source/Runtime/AssetRegistry/Private/AssetDataGatherer.cpp
--- a/Engine/Source/Runtime/AssetRegistry/Private/AssetDataGatherer.cpp
+++ b/Engine/Source/Runtime/AssetRegistry/Private/AssetDataGatherer.cpp
@@ -2979,7 +2979,7 @@ struct FPreloadSettings
 
 		MonolithicCacheFilename = FPaths::ProjectIntermediateDir() / (bGatherDependsData ? TEXT("CachedAssetRegistry.bin") : TEXT("CachedAssetRegistryNoDeps.bin"));
 #if UE_EDITOR // See note on FPreloader for why we only allow preloading if UE_EDITOR
-		bMonolithicCacheActivatedDuringPreload = bAsyncEnabled && GIsEditor && (!IsRunningCommandlet() || IsRunningCookCommandlet());
+		bMonolithicCacheActivatedDuringPreload = bAsyncEnabled && UE::AssetRegistry::ShouldSearchAllAssetsAtStart();
 #else
 		bMonolithicCacheActivatedDuringPreload = false;
 #endif
diff --git a/Engine/Source/Runtime/AssetRegistry/Private/AssetRegistry.cpp b/Engine/Source/Runtime/AssetRegistry/Private/AssetRegistry.cpp
--- a/Engine/Source/Runtime/AssetRegistry/Private/AssetRegistry.cpp
+++ b/Engine/Source/Runtime/AssetRegistry/Private/AssetRegistry.cpp
@@ -7050,9 +7050,43 @@ void GetAssetForPackages(TConstArrayView<FName> PackageNames, TMap<FName, FAsset
 
 bool ShouldSearchAllAssetsAtStart()
 {
-	// If in the editor or cookcommandlet, we start the GlobalGatherer now
+	// If in the editor or cookcommandlet, or an allowlist commandlet, we start the GlobalGatherer now
 	// In the game or other commandlets, we do not construct it until project or commandlet code calls SearchAllAssets or ScanPathsSynchronous
-	bool bSearchAllAssetsAtStart = GIsEditor && (!IsRunningCommandlet() || IsRunningCookCommandlet());
+	bool bSearchAllAssetsAtStart = false;
+	if (GIsEditor)
+	{
+		if (!IsRunningCommandlet() || IsRunningCookCommandlet())
+		{
+			bSearchAllAssetsAtStart = true;
+		}
+		else
+		{
+			TArray<FString> CommandletsUsingAR;
+			GConfig->GetArray(TEXT("AssetRegistry"), TEXT("CommandletsUsingAR"), CommandletsUsingAR, GEngineIni);
+			FString CommandlineCommandlet;
+			FString CommandletToken(TEXT("commandlet"));
+			if (!CommandletsUsingAR.IsEmpty() &&
+				FParse::Value(FCommandLine::Get(), TEXT("-run="), CommandlineCommandlet))
+			{
+				if (CommandlineCommandlet.EndsWith(CommandletToken))
+				{
+					CommandlineCommandlet.LeftChopInline(CommandletToken.Len(), false /* bAllowShrinking */);
+				}
+				for (FString& CommandletUsingAR : CommandletsUsingAR)
+				{
+					if (CommandletUsingAR.EndsWith(CommandletToken))
+					{
+						CommandletUsingAR.LeftChopInline(CommandletToken.Len(), false /* bAllowShrinking */);
+					}
+					if (CommandletUsingAR == CommandlineCommandlet)
+					{
+						bSearchAllAssetsAtStart = true;
+						break;
+					}
+				}
+			}
+		}
+	}
 #if !UE_BUILD_SHIPPING
 	bool bCommandlineAllAssetsAtStart;
 	if (FParse::Bool(FCommandLine::Get(), TEXT("AssetGatherAll="), bCommandlineAllAssetsAtStart))
