#UE Add a commandline option to specify the folder to discover the asset registry cache files. Useful in cases where a workspace is frequently wiped, but you want to retain the cache files to keep using them between wipes.

diff --git a/Engine/Source/Runtime/AssetRegistry/Private/AssetDataGatherer.cpp b/Engine/Source/Runtime/AssetRegistry/Private/AssetDataGatherer.cpp
--- a/Engine/Source/Runtime/AssetRegistry/Private/AssetDataGatherer.cpp
+++ b/Engine/Source/Runtime/AssetRegistry/Private/AssetDataGatherer.cpp
@@ -2963,8 +2963,10 @@ struct FPreloadSettings
 			return;
 		}
 		bInitialized = true;
-
-		FString ProjectIntermediateDir = FPaths::ProjectIntermediateDir();
+		if (!FParse::Value(FCommandLine::Get(), TEXT("AssetRegistryCacheRootFolder="), AssetRegistryCacheRootFolder))
+		{
+			AssetRegistryCacheRootFolder = FPaths::ProjectIntermediateDir();
+		}
 		bForceDependsGathering = FParse::Param(FCommandLine::Get(), TEXT("ForceDependsGathering"));
 		bGatherDependsData = (GIsEditor && !FParse::Param(FCommandLine::Get(), TEXT("NoDependsGathering"))) || bForceDependsGathering;
 		bool bNoAssetRegistryCache = FParse::Param(FCommandLine::Get(), TEXT("NoAssetRegistryCache"));
@@ -2977,7 +2979,7 @@ struct FPreloadSettings
 		bCacheWriteEnabled = !bNoAssetRegistryCache && !bNoAssetRegistryCacheWrite && !bMultiprocess;
 		bool bAsyncEnabled = FPlatformProcess::SupportsMultithreading() && FTaskGraphInterface::IsRunning();
 
-		MonolithicCacheFilename = FPaths::ProjectIntermediateDir() / (bGatherDependsData ? TEXT("CachedAssetRegistry.bin") : TEXT("CachedAssetRegistryNoDeps.bin"));
+		MonolithicCacheFilename = AssetRegistryCacheRootFolder / (bGatherDependsData ? TEXT("CachedAssetRegistry.bin") : TEXT("CachedAssetRegistryNoDeps.bin"));
 #if UE_EDITOR // See note on FPreloader for why we only allow preloading if UE_EDITOR
 		bMonolithicCacheActivatedDuringPreload = bAsyncEnabled && UE::AssetRegistry::ShouldSearchAllAssetsAtStart();
 #else
@@ -3012,8 +3014,13 @@ struct FPreloadSettings
 	{
 		return MonolithicCacheFilename;
 	}
+	const FString& GetAssetRegistryCacheRootFolder() const
+	{
+		return AssetRegistryCacheRootFolder;
+	}
 private:
 	FString MonolithicCacheFilename;
+	FString AssetRegistryCacheRootFolder;
 	bool bForceDependsGathering = false;
 	bool bGatherDependsData = false;
 	bool bCacheReadEnabled = false;
@@ -4004,7 +4011,7 @@ FString FAssetDataGatherer::GetCacheFilename(TConstArrayView<FString> CacheFileP
 		CacheHash = HashCombine(CacheHash, GetTypeHash(SortedPaths[PathIndex]));
 	}
 
-	return FPaths::ProjectIntermediateDir() / TEXT("AssetRegistryCache") / FString::Printf(TEXT("%08x%s.bin"), CacheHash, bGatherDependsData ? TEXT("") : TEXT("NoDeps"));
+	return UE::AssetDataGather::Private::GPreloadSettings.GetAssetRegistryCacheRootFolder() / TEXT("AssetRegistryCache") / FString::Printf(TEXT("%08x%s.bin"), CacheHash, bGatherDependsData ? TEXT("") : TEXT("NoDeps"));
 }
 
 void FAssetDataGatherer::LoadCacheFile(FStringView CacheFilename)
