#UE Write AssetRegistry cache files in situations where you are synchronously scanning, but there is no background thread to do the write for you

diff --git a/Engine/Source/Runtime/AssetRegistry/Private/AssetDataGatherer.cpp b/Engine/Source/Runtime/AssetRegistry/Private/AssetDataGatherer.cpp
--- a/Engine/Source/Runtime/AssetRegistry/Private/AssetDataGatherer.cpp
+++ b/Engine/Source/Runtime/AssetRegistry/Private/AssetDataGatherer.cpp
@@ -3233,9 +3233,9 @@ void FAssetDataGatherer::InnerTickLoop(bool bInIsSynchronousTick, bool bContribu
 {
 	using namespace UE::AssetDataGather::Private;
 
-	// Synchronous ticks during Wait contribute to saving of the async cache only if there is no dedicated async thread to do it (bIsSynchronous is true)
+	// Synchronous ticks during Wait contribute to saving of the async cache only if there is no dedicated async thread to do it (Thread == nullptr)
 	// The dedicated async thread always contributes
-	bContributeToCacheSave = !bInIsSynchronousTick || (bIsSynchronous && bContributeToCacheSave);
+	bContributeToCacheSave = !bInIsSynchronousTick || (Thread == nullptr && bContributeToCacheSave);
 
 	bool bShouldSaveMonolithicCache = false;
 	TArray<TPair<FName, FDiskCachedAssetData*>> AssetsToSave;
