#UE Silently skip empty paths that are passed in to scan. These are not useful or desirable, but they are sometimes hard to track down and do not matter much

diff --git a/Engine/Source/Runtime/AssetRegistry/Private/AssetRegistry.cpp b/Engine/Source/Runtime/AssetRegistry/Private/AssetRegistry.cpp
--- a/Engine/Source/Runtime/AssetRegistry/Private/AssetRegistry.cpp
+++ b/Engine/Source/Runtime/AssetRegistry/Private/AssetRegistry.cpp
@@ -4533,7 +4533,11 @@ FScanPathContext::FScanPathContext(FEventContext& InEventContext, FClassInherita
 	PackageFiles.Reserve(InFiles.Num());
 	for (const FString& InFile : InFiles)
 	{
-		if (!FPackageName::TryConvertToMountedPath(InFile, &LocalPath, &PackageName, nullptr, nullptr, &Extension, &FlexNameType))
+		if (InFile.IsEmpty())
+		{
+			continue;
+		}
+		else if (!FPackageName::TryConvertToMountedPath(InFile, &LocalPath, &PackageName, nullptr, nullptr, &Extension, &FlexNameType))
 		{
 			UE_LOG(LogAssetRegistry, Warning, TEXT("ScanPathsSynchronous: %s is not in a mounted path, will not scan."), *InFile);
 			continue;
@@ -4560,7 +4564,11 @@ FScanPathContext::FScanPathContext(FEventContext& InEventContext, FClassInherita
 	PackageDirs.Reserve(InDirs.Num());
 	for (const FString& InDir : InDirs)
 	{
-		if (!FPackageName::TryConvertToMountedPath(InDir, &LocalPath, &PackageName, nullptr, nullptr, &Extension, &FlexNameType))
+		if (InDir.IsEmpty())
+		{
+			continue;
+		}
+		else if (!FPackageName::TryConvertToMountedPath(InDir, &LocalPath, &PackageName, nullptr, nullptr, &Extension, &FlexNameType))
 		{
 			UE_LOG(LogAssetRegistry, Warning, TEXT("ScanPathsSynchronous: %s is not in a mounted path, will not scan."), *InDir);
 			continue;
