AssetRegistryState: GetAllAssets now calls reserve to reduce reallocation costs.
#rnx

diff --git a/Engine/Source/Runtime/AssetRegistry/Private/AssetRegistryState.cpp b/Engine/Source/Runtime/AssetRegistry/Private/AssetRegistryState.cpp
--- a/Engine/Source/Runtime/AssetRegistry/Private/AssetRegistryState.cpp
+++ b/Engine/Source/Runtime/AssetRegistry/Private/AssetRegistryState.cpp
@@ -807,6 +807,7 @@ bool FAssetRegistryState::EnumerateAssets(const FARCompiledFilter& Filter, const
 
 bool FAssetRegistryState::GetAllAssets(const TSet<FName>& PackageNamesToSkip, TArray<FAssetData>& OutAssetData, bool bSkipARFilteredAssets) const
 {
+	OutAssetData.Reserve(OutAssetData.Num() + CachedAssets.Num() - PackageNamesToSkip.Num());
 	return EnumerateAllAssets(PackageNamesToSkip, [&OutAssetData](const FAssetData& AssetData)
 	{
 		OutAssetData.Emplace(AssetData);
