AssetRegistry ExtendedTags: When gathering assets from disk, if the Asset has already been loaded in memory, add the ExtendedTags onto the regular tags parsed from the in-memory Asset.
#rnx

diff --git a/Engine/Source/Runtime/AssetRegistry/Private/AssetRegistry.cpp b/Engine/Source/Runtime/AssetRegistry/Private/AssetRegistry.cpp
--- a/Engine/Source/Runtime/AssetRegistry/Private/AssetRegistry.cpp
+++ b/Engine/Source/Runtime/AssetRegistry/Private/AssetRegistry.cpp
@@ -4800,14 +4800,20 @@ void FAssetRegistryImpl::AssetSearchDataGathered(Impl::FEventContext& EventConte
 			// If this ensure fires then we've somehow processed the same result more than once, and that should never happen
 			if (ensure(ExistingAssetData != BackgroundResult.Get()))
 			{
-				// If the current AssetData came from a loaded asset, don't overwrite it with the new one from disk; loaded asset is more authoritative because it has run the postload steps
 #if WITH_EDITOR
-				if (!AssetDataObjectPathsUpdatedOnLoad.Contains(BackgroundResult->GetSoftObjectPath()))
-#endif
+				PostLoadAssetRegistryTags(BackgroundResult.Get());
+				if (AssetDataObjectPathsUpdatedOnLoad.Contains(BackgroundResult->GetSoftObjectPath()))
 				{
-#if WITH_EDITOR
-					PostLoadAssetRegistryTags(BackgroundResult.Get());
+					// If the current AssetData came from a loaded asset, don't overwrite it with the new one from disk
+					// The loaded asset is more authoritative because it has run the postload steps.
+					// However, the loaded asset is missing the extended tags. Our contract for extended tags is to keep any 
+					// that do not exist in the non-extended tags. So add on any tags from the BackgroundResult that
+					// are not already on the existing asset.
+					AddNonOverlappingTags(EventContext, *ExistingAssetData, *BackgroundResult);
+				}
+				else
 #endif
+				{
 					// The asset exists in the cache from disk and has not yet been loaded into memory, update it with the new background data
 					UpdateAssetData(EventContext, ExistingAssetData, MoveTemp(*BackgroundResult), false /* bKeepDeletedTags */);
 				}
@@ -5224,6 +5230,38 @@ void FAssetRegistryImpl::UpdateAssetData(Impl::FEventContext& EventContext, FAss
 	}
 }
 
+void FAssetRegistryImpl::AddNonOverlappingTags(Impl::FEventContext& EventContext, FAssetData& ExistingAssetData,
+	const FAssetData& NewAssetData)
+{
+	TOptional<FAssetDataTagMap> ModifiedTags;
+	NewAssetData.TagsAndValues.ForEach([&ExistingAssetData, &ModifiedTags](const TPair<FName, FAssetTagValueRef>& TagPair)
+		{
+			if (ModifiedTags)
+			{
+				if (!ModifiedTags->Contains(TagPair.Key))
+				{
+					ModifiedTags->Add(TagPair.Key, TagPair.Value.GetStorageString());
+				}
+			}
+			else
+			{
+				if (!ExistingAssetData.TagsAndValues.Contains(TagPair.Key))
+				{
+					ModifiedTags.Emplace(ExistingAssetData.TagsAndValues.CopyMap());
+					ModifiedTags->Add(TagPair.Key, TagPair.Value.GetStorageString());
+				}
+			}
+		});
+	if (ModifiedTags)
+	{
+		State.SetTagsOnExistingAsset(&ExistingAssetData, MoveTemp(*ModifiedTags));
+		if (!ShouldSkipAsset(ExistingAssetData.AssetClassPath, ExistingAssetData.PackageFlags))
+		{
+			EventContext.AssetEvents.Emplace(ExistingAssetData, Impl::FEventContext::EEvent::Updated);
+		}
+	}
+}
+
 bool FAssetRegistryImpl::RemoveAssetData(Impl::FEventContext& EventContext, FAssetData* AssetData)
 {
 	bool bRemoved = false;
diff --git a/Engine/Source/Runtime/AssetRegistry/Private/AssetRegistryImpl.h b/Engine/Source/Runtime/AssetRegistry/Private/AssetRegistryImpl.h
--- a/Engine/Source/Runtime/AssetRegistry/Private/AssetRegistryImpl.h
+++ b/Engine/Source/Runtime/AssetRegistry/Private/AssetRegistryImpl.h
@@ -309,6 +309,10 @@ private:
 	/** Updates an existing asset data with the new value and updates lookup maps */
 	void UpdateAssetData(Impl::FEventContext& EventContext, FAssetData* AssetData, FAssetData&& NewAssetData, bool bKeepDeletedTags);
 
+	/** Updates the tags on an existing asset data by adding the tags from NewAssetData that do not already exist. */
+	void AddNonOverlappingTags(Impl::FEventContext& EventContext, FAssetData& ExistingAssetData,
+		const FAssetData& NewAssetData);
+
 	/** Removes the asset data from the lookup maps */
 	bool RemoveAssetData(Impl::FEventContext& EventContext, FAssetData* AssetData);
 
