Cooker/AssetRegistry: Fix performance bug by new calls to GetAssetsByPackageName: optimize GetAssetsByPackageName to use FindObject instead of a TObjectIterator.
#rnx

diff --git a/Engine/Source/Runtime/AssetRegistry/Private/AssetRegistry.cpp b/Engine/Source/Runtime/AssetRegistry/Private/AssetRegistry.cpp
--- a/Engine/Source/Runtime/AssetRegistry/Private/AssetRegistry.cpp
+++ b/Engine/Source/Runtime/AssetRegistry/Private/AssetRegistry.cpp
@@ -2045,6 +2045,27 @@ void EnumerateMemoryAssetsHelper(const FARCompiledFilter& InFilter, TSet<FName>&
 			}
 		}
 	}
+	else if (InFilter.PackageNames.Num() > 0)
+	{
+		for (FName PackageName : InFilter.PackageNames)
+		{
+			UPackage* Package = FindObjectFast<UPackage>(nullptr, PackageName);
+			if (Package)
+			{
+				bool bContinue = true;
+				ForEachObjectWithPackage(Package, [&FilterInMemoryObjectLambda, &bContinue](UObject* Object)
+					{
+						FilterInMemoryObjectLambda(Object, bContinue);
+						return bContinue;
+					});
+				if (!bContinue)
+				{
+					bOutStopIteration = true;
+					return;
+				}
+			}
+		}
+	}
 	else
 	{
 		for (FThreadSafeObjectIterator ObjIt; ObjIt; ++ObjIt)
