AssetRegistry gather: Downgrade corrupt cachefile messages from Error to Warning, because the editor/commandlet can recover from them and it is a performance issue only.
#rnx

diff --git a/Engine/Source/Runtime/AssetRegistry/Private/AssetDataGatherer.cpp b/Engine/Source/Runtime/AssetRegistry/Private/AssetDataGatherer.cpp
--- a/Engine/Source/Runtime/AssetRegistry/Private/AssetDataGatherer.cpp
+++ b/Engine/Source/Runtime/AssetRegistry/Private/AssetDataGatherer.cpp
@@ -2645,7 +2645,7 @@ TOptional<FBlockHeader> LoadBlockHeader(FMemoryView Data)
 
 	if (Header.Magic != BlockMagic)
 	{
-		UE_LOG(LogAssetRegistry, Error, TEXT("Wrong block magic (0x%x)"), Header.Magic);
+		UE_LOG(LogAssetRegistry, Warning, TEXT("Wrong block magic (0x%x)"), Header.Magic);
 		return {};
 	}
 
@@ -2733,7 +2733,7 @@ class FChecksumArchiveBase : public FArchiveProxy
 		InnerArchive.Serialize(HeaderData, sizeof(HeaderData));
 		if (InnerArchive.IsError())
 		{
-			UE_LOG(LogAssetRegistry, Error, TEXT("Couldn't read block header"));
+			UE_LOG(LogAssetRegistry, Warning, TEXT("Couldn't read block header"));
 			return false;
 		}
 
@@ -2745,12 +2745,12 @@ class FChecksumArchiveBase : public FArchiveProxy
 
 			if (InnerArchive.IsError())
 			{
-				UE_LOG(LogAssetRegistry, Error, TEXT("Couldn't read block data"));
+				UE_LOG(LogAssetRegistry, Warning, TEXT("Couldn't read block data"));
 				return false;
 			}
 			else if (CalculateBlockChecksum(Block.GetRemaining()) != Header->Checksum)
 			{
-				UE_LOG(LogAssetRegistry, Error, TEXT("Wrong block checksum"));
+				UE_LOG(LogAssetRegistry, Warning, TEXT("Wrong block checksum"));
 				return false;
 			}
 
@@ -2805,7 +2805,7 @@ protected:
 
 			if (!LoadBlock())
 			{
-				UE_LOG(LogAssetRegistry, Error, TEXT("Integrity check failed, '%s' cache will be discarded"), *InnerArchive.GetArchiveName());
+				UE_LOG(LogAssetRegistry, Warning, TEXT("Integrity check failed, '%s' cache will be discarded"), *InnerArchive.GetArchiveName());
 				SetError();
 				return;
 			}
@@ -2892,7 +2892,7 @@ private:
 			TOptional<FMemoryView> NextBlock = LoadNextBlock(RemainingBlocks);
 			if (!NextBlock)
 			{
-				UE_LOG(LogAssetRegistry, Error, TEXT("Integrity check failed, '%s' cache will be discarded"), *FileName);
+				UE_LOG(LogAssetRegistry, Warning, TEXT("Integrity check failed, '%s' cache will be discarded"), *FileName);
 				SetError();
 				return;
 			}
@@ -2906,7 +2906,7 @@ private:
 	{
 		if (In.GetRemainingSize() < sizeof(FBlockHeader))
 		{
-			UE_LOG(LogAssetRegistry, Error, TEXT("Couldn't read block header"));
+			UE_LOG(LogAssetRegistry, Warning, TEXT("Couldn't read block header"));
 			return {};
 		}
 
@@ -2914,14 +2914,14 @@ private:
 		{
 			if (Header->Size > In.GetRemainingSize())
 			{
-				UE_LOG(LogAssetRegistry, Error, TEXT("Incomplete block"));
+				UE_LOG(LogAssetRegistry, Warning, TEXT("Incomplete block"));
 				return {};
 			}
 
 			FMemoryView Block = In.Load(Header->Size);
 			if (CalculateBlockChecksum(Block) != Header->Checksum)
 			{
-				UE_LOG(LogAssetRegistry, Error, TEXT("Wrong block checksum"));
+				UE_LOG(LogAssetRegistry, Warning, TEXT("Wrong block checksum"));
 				return {};
 			}
 
@@ -4273,7 +4273,7 @@ FCachePayload LoadCacheFile(FStringView InCacheFilename)
 		{
 			FChecksumViewReader ChecksummingReader(MoveTemp(FileReader), CacheFilename);
 			Payload = DoLoad(ChecksummingReader);
-			UE_CLOG(!Payload.bSucceeded, LogAssetRegistry, Error, TEXT("There was an error loading the asset registry cache using memory mapping"));
+			UE_CLOG(!Payload.bSucceeded, LogAssetRegistry, Warning, TEXT("There was an error loading the asset registry cache using memory mapping"));
 		}
 
 		Preload.Wait();
@@ -4290,7 +4290,7 @@ FCachePayload LoadCacheFile(FStringView InCacheFilename)
 			{
 				FChecksumArchiveReader ChecksummingReader(*FileAr);
 				Payload = DoLoad(ChecksummingReader);
-				UE_CLOG(!Payload.bSucceeded, LogAssetRegistry, Error, TEXT("There was an error loading the asset registry cache"));
+				UE_CLOG(!Payload.bSucceeded, LogAssetRegistry, Warning, TEXT("There was an error loading the asset registry cache"));
 			}
 		}
 	}
