Add support for blocking packages (in the asset registry) if they are downloaded from the internet.

API changes:
- Add IAssetRegistry::OnFilesBlocked
- Add IPlatformFile::HasMarkOfTheWeb
- Add IPlatformFile::SetMarkOfTheWeb
- Add CVar AssetRegistry.BlockPackagesWithMarkOfTheWeb (default false)

diff --git a/Engine/Source/Runtime/AssetRegistry/Private/AssetDataGatherer.cpp b/Engine/Source/Runtime/AssetRegistry/Private/AssetDataGatherer.cpp
--- a/Engine/Source/Runtime/AssetRegistry/Private/AssetDataGatherer.cpp
+++ b/Engine/Source/Runtime/AssetRegistry/Private/AssetDataGatherer.cpp
@@ -59,6 +59,18 @@ namespace AssetDataGathererConstants
 namespace UE::AssetDataGather::Private
 {
 
+static bool BlockPackagesWithMarkOfTheWeb = false;
+static FAutoConsoleVariableRef CVarBlockPackagesWithMarkOfTheWeb(
+	TEXT("AssetRegistry.BlockPackagesWithMarkOfTheWeb"),
+	BlockPackagesWithMarkOfTheWeb,
+	TEXT("Whether package files with mark of the web are blocked from the asset registry"),
+	ECVF_SetBySystemSettingsIni);
+
+bool IsPackageBlocked(const FStringView& FilePath)
+{
+	return BlockPackagesWithMarkOfTheWeb && IPlatformFile::GetPlatformPhysical().HasMarkOfTheWeb(*FString(FilePath));
+}
+
 /** A structure to hold serialized cache data from async loads before adding it to the Gatherer's main cache. */
 struct FCachePayload
 {
@@ -96,34 +108,37 @@ void AssignStringWithoutShrinking(FString& InOutResult, FStringView Value)
 	}
 }
 
-FDiscoveredPathData::FDiscoveredPathData(FStringView InLocalAbsPath, FStringView InLongPackageName, FStringView InRelPath, const FDateTime& InPackageTimestamp, EGatherableFileType InType)
+FDiscoveredPathData::FDiscoveredPathData(FStringView InLocalAbsPath, FStringView InLongPackageName, FStringView InRelPath, const FDateTime& InPackageTimestamp, EGatherableFileType InType, bool bInBlocked)
 	: LocalAbsPath(InLocalAbsPath)
 	, LongPackageName(InLongPackageName)
 	, RelPath(InRelPath)
 	, PackageTimestamp(InPackageTimestamp)
 	, Type(InType)
+	, bBlocked(bInBlocked)
 {
 }
 
-FDiscoveredPathData::FDiscoveredPathData(FStringView InLocalAbsPath, FStringView InLongPackageName, FStringView InRelPath, EGatherableFileType InType)
+FDiscoveredPathData::FDiscoveredPathData(FStringView InLocalAbsPath, FStringView InLongPackageName, FStringView InRelPath, EGatherableFileType InType, bool bInBlocked)
 	: LocalAbsPath(InLocalAbsPath)
 	, LongPackageName(InLongPackageName)
 	, RelPath(InRelPath)
 	, Type(InType)
+	, bBlocked(bInBlocked)
 {
 }
 
-void FDiscoveredPathData::Assign(FStringView InLocalAbsPath, FStringView InLongPackageName, FStringView InRelPath, EGatherableFileType InType)
+void FDiscoveredPathData::Assign(FStringView InLocalAbsPath, FStringView InLongPackageName, FStringView InRelPath, EGatherableFileType InType, bool bInBlocked)
 {
 	AssignStringWithoutShrinking(LocalAbsPath, InLocalAbsPath);
 	AssignStringWithoutShrinking(LongPackageName, InLongPackageName);
 	AssignStringWithoutShrinking(RelPath, InRelPath);
 	Type = InType;
+	bBlocked = bInBlocked;
 }
 
-void FDiscoveredPathData::Assign(FStringView InLocalAbsPath, FStringView InLongPackageName, FStringView InRelPath, const FDateTime& InPackageTimestamp, EGatherableFileType InType)
+void FDiscoveredPathData::Assign(FStringView InLocalAbsPath, FStringView InLongPackageName, FStringView InRelPath, const FDateTime& InPackageTimestamp, EGatherableFileType InType, bool bInBlocked)
 {
-	Assign(InLocalAbsPath, InLongPackageName, InRelPath, InType);
+	Assign(InLocalAbsPath, InLongPackageName, InRelPath, InType, bInBlocked);
 	PackageTimestamp = InPackageTimestamp;
 }
 
@@ -132,16 +147,17 @@ SIZE_T FDiscoveredPathData::GetAllocatedSize() const
 	return LocalAbsPath.GetAllocatedSize() + LongPackageName.GetAllocatedSize() + RelPath.GetAllocatedSize();
 }
 
-FGatheredPathData::FGatheredPathData(FStringView InLocalAbsPath, FStringView InLongPackageName, const FDateTime& InPackageTimestamp, EGatherableFileType InType)
+FGatheredPathData::FGatheredPathData(FStringView InLocalAbsPath, FStringView InLongPackageName, const FDateTime& InPackageTimestamp, EGatherableFileType InType, bool bInBlocked)
 	: LocalAbsPath(InLocalAbsPath)
 	, LongPackageName(InLongPackageName)
 	, PackageTimestamp(InPackageTimestamp)
 	, Type(InType)
+	, bBlocked(bInBlocked)
 {
 }
 
 FGatheredPathData::FGatheredPathData(const FDiscoveredPathData& DiscoveredData)
-	: FGatheredPathData(DiscoveredData.LocalAbsPath, DiscoveredData.LongPackageName, DiscoveredData.PackageTimestamp, DiscoveredData.Type)
+	: FGatheredPathData(DiscoveredData.LocalAbsPath, DiscoveredData.LongPackageName, DiscoveredData.PackageTimestamp, DiscoveredData.Type, DiscoveredData.bBlocked)
 {
 }
 
@@ -150,20 +166,22 @@ FGatheredPathData::FGatheredPathData(FDiscoveredPathData&& DiscoveredData)
 	, LongPackageName(MoveTemp(DiscoveredData.LongPackageName))
 	, PackageTimestamp(MoveTemp(DiscoveredData.PackageTimestamp))
 	, Type(DiscoveredData.Type)
+	, bBlocked(DiscoveredData.bBlocked)
 {
 }
 
-void FGatheredPathData::Assign(FStringView InLocalAbsPath, FStringView InLongPackageName, const FDateTime& InPackageTimestamp, EGatherableFileType InType)
+void FGatheredPathData::Assign(FStringView InLocalAbsPath, FStringView InLongPackageName, const FDateTime& InPackageTimestamp, EGatherableFileType InType, bool bInBlocked)
 {
 	AssignStringWithoutShrinking(LocalAbsPath, InLocalAbsPath);
 	AssignStringWithoutShrinking(LongPackageName, InLongPackageName);
 	PackageTimestamp = InPackageTimestamp;
 	Type = InType;
+	bBlocked = bInBlocked;
 }
 
 void FGatheredPathData::Assign(const FDiscoveredPathData& DiscoveredData)
 {
-	Assign(DiscoveredData.LocalAbsPath, DiscoveredData.LongPackageName, DiscoveredData.PackageTimestamp, DiscoveredData.Type);
+	Assign(DiscoveredData.LocalAbsPath, DiscoveredData.LongPackageName, DiscoveredData.PackageTimestamp, DiscoveredData.Type, DiscoveredData.bBlocked);
 }
 
 SIZE_T FGatheredPathData::GetAllocatedSize() const
@@ -1648,7 +1666,7 @@ void FAssetDataDiscovery::TickInternal(bool bTickAll)
 							Data.IteratedSubDirs.Emplace();
 						}
 						Data.IteratedSubDirs[Data.NumIteratedDirs++].Assign(LocalAbsPath, Data.DirLongPackageName, RelPath,
-							EGatherableFileType::Directory);
+							EGatherableFileType::Directory, false /* bBlocked */);
 					}
 				}
 				else
@@ -1667,8 +1685,9 @@ void FAssetDataDiscovery::TickInternal(bool bTickAll)
 								Data.IteratedFiles.Emplace();
 							}
 							FPathViews::AppendPath(Data.DirLongPackageName, BaseName);
+							const bool bBlocked = FileType == EGatherableFileType::PackageFile && IsPackageBlocked(LocalAbsPath);
 							Data.IteratedFiles[Data.NumIteratedFiles++].Assign(LocalAbsPath, Data.DirLongPackageName, RelPath,
-								InPackageStatData.ModificationTime, FileType);
+								InPackageStatData.ModificationTime, FileType, bBlocked);
 						}
 					}
 				}
@@ -1985,7 +2004,8 @@ void FAssetDataDiscovery::SetPropertiesAndWait(TArrayView<FPathExistence> QueryP
 							LongPackageName << MountDir->GetLongPackageName();
 							FPathViews::AppendPath(LongPackageName, ScanDir->GetMountRelPath());
 							FPathViews::AppendPath(LongPackageName, FileRelPathNoExt);
-							AddDiscoveredFile(FDiscoveredPathData(SearchPath, LongPackageName, RelPathFromParentDir, QueryPath.GetModificationTime(), FileType));
+							const bool bBlocked = FileType == EGatherableFileType::PackageFile && IsPackageBlocked(SearchPath);
+							AddDiscoveredFile(FDiscoveredPathData(SearchPath, LongPackageName, RelPathFromParentDir, QueryPath.GetModificationTime(), FileType, bBlocked));
 							if (FPathViews::IsPathLeaf(RelPath) && !ScanDir->HasScanned())
 							{
 								SetIsIdle(false);
@@ -2452,7 +2472,8 @@ void FAssetDataDiscovery::OnFileCreated(const FString& LocalAbsPath)
 			LongPackageName << MountDir->GetLongPackageName();
 			FPathViews::AppendPath(LongPackageName, ScanDir->GetMountRelPath());
 			FPathViews::AppendPath(LongPackageName, FileRelPathNoExt);
-			AddDiscoveredFile(FDiscoveredPathData(LocalAbsPath, LongPackageName, RelPathFromParentDir, StatData.ModificationTime, FileType));
+			const bool bBlocked = FileType == EGatherableFileType::PackageFile && IsPackageBlocked(LocalAbsPath);
+			AddDiscoveredFile(FDiscoveredPathData(LocalAbsPath, LongPackageName, RelPathFromParentDir, StatData.ModificationTime, FileType, bBlocked));
 			if (FPathViews::IsPathLeaf(FileRelPath))
 			{
 				ScanDir->MarkFileAlreadyScanned(FileRelPath);
@@ -3344,6 +3365,7 @@ void FAssetDataGatherer::TickInternal(bool& bOutIsTickInterrupt, double& TickSta
 	TArray<FPackageDependencyData, FBatchInlineAllocator> LocalDependencyResults;
 	TArray<FString, FBatchInlineAllocator> LocalCookedPackageNamesWithoutAssetDataResults;
 	TArray<FName, FBatchInlineAllocator> LocalVerseResults;
+	TArray<FString, FBatchInlineAllocator> LocalBlockedResults;
 	bool bLoadMonolithicCache = false;
 	bool bLocalIsCacheWriteEnabled = false;
 	double LocalLastCacheWriteTime = 0.0;
@@ -3457,6 +3479,13 @@ void FAssetDataGatherer::TickInternal(bool& bOutIsTickInterrupt, double& TickSta
 			continue;
 		}
 
+		// If this is a blocked package, just directly add its file path to the Blocked results
+		if (AssetFileData.bBlocked)
+		{
+			LocalBlockedResults.Add(AssetFileData.LocalAbsPath);
+			continue;
+		}
+
 		const FName PackageName = *AssetFileData.LongPackageName;
 		const FName Extension = FName(*FPaths::GetExtension(AssetFileData.LocalAbsPath));
 
@@ -3590,6 +3619,7 @@ void FAssetDataGatherer::TickInternal(bool& bOutIsTickInterrupt, double& TickSta
 		DependencyResults.Append(MoveTemp(LocalDependencyResults));
 		CookedPackageNamesWithoutAssetDataResults.Append(MoveTemp(LocalCookedPackageNamesWithoutAssetDataResults));
 		VerseResults.Append(MoveTemp(LocalVerseResults));
+		BlockedResults.Append(MoveTemp(LocalBlockedResults));
 
 		if (bHasCancelation)
 		{
@@ -3754,6 +3784,9 @@ void FAssetDataGatherer::GetAndTrimSearchResults(FResults& InOutResults, FResult
 	MoveAppendRangeToRingBuffer(InOutResults.CookedPackageNamesWithoutAssetData, CookedPackageNamesWithoutAssetDataResults);
 	MoveAppendRangeToRingBuffer(InOutResults.VerseFiles, VerseResults);
 
+	InOutResults.BlockedFiles.Append(MoveTemp(BlockedResults));
+	BlockedResults.Reset();
+
 	OutContext.SearchTimes.Append(MoveTemp(SearchTimes));
 	SearchTimes.Reset();
 
@@ -4342,6 +4375,7 @@ SIZE_T FAssetDataGatherer::GetAllocatedSize() const
 	Result += GetArrayRecursiveAllocatedSize(DependencyResults);
 	Result += GetArrayRecursiveAllocatedSize(CookedPackageNamesWithoutAssetDataResults);
 	Result += VerseResults.GetAllocatedSize();
+	Result += BlockedResults.GetAllocatedSize();
 	Result += SearchTimes.GetAllocatedSize();
 	Result += GetArrayRecursiveAllocatedSize(DiscoveredPaths);
 	Result += GPreloadSettings.GetMonolithicCacheFilename().GetAllocatedSize();
@@ -4371,6 +4405,7 @@ void FAssetDataGatherer::Shrink()
 	DependencyResults.Shrink();
 	CookedPackageNamesWithoutAssetDataResults.Shrink();
 	VerseResults.Shrink();
+	BlockedResults.Shrink();
 	SearchTimes.Shrink();
 	DiscoveredPaths.Shrink();
 }
@@ -4944,4 +4979,4 @@ int32 FFilesToSearch::FTreeNode::NumFiles() const
 	return Num;
 }
 
-}
\ No newline at end of file
+}
diff --git a/Engine/Source/Runtime/AssetRegistry/Private/AssetDataGatherer.h b/Engine/Source/Runtime/AssetRegistry/Private/AssetDataGatherer.h
--- a/Engine/Source/Runtime/AssetRegistry/Private/AssetDataGatherer.h
+++ b/Engine/Source/Runtime/AssetRegistry/Private/AssetDataGatherer.h
@@ -93,11 +93,13 @@ public:
 		TMultiMap<FName, FPackageDependencyData> Dependencies;
 		TRingBuffer<FString> CookedPackageNamesWithoutAssetData;
 		TRingBuffer<FName> VerseFiles;
+		TArray<FString> BlockedFiles;
 
 		SIZE_T GetAllocatedSize() const
 		{
 			return Assets.GetAllocatedSize() + Paths.GetAllocatedSize() + Dependencies.GetAllocatedSize() +
-				CookedPackageNamesWithoutAssetData.GetAllocatedSize() + VerseFiles.GetAllocatedSize();
+				CookedPackageNamesWithoutAssetData.GetAllocatedSize() + VerseFiles.GetAllocatedSize() +
+				BlockedFiles.GetAllocatedSize();
 		}
 		void Shrink()
 		{
@@ -106,6 +108,7 @@ public:
 			Dependencies.Shrink();
 			CookedPackageNamesWithoutAssetData.Trim();
 			VerseFiles.Trim();
+			BlockedFiles.Shrink();
 		}
 	};
 	struct FResultContext
@@ -386,6 +389,8 @@ private:
 	TArray<FString> CookedPackageNamesWithoutAssetDataResults;
 	/** File paths (in UE LongPackagePath notation) of the Verse source code gathered from the searched files. */
 	TArray<FName> VerseResults;
+	/** File paths (in regular filesystem notation) of blocked packages from the searched files. */
+	TArray<FString> BlockedResults;
 
 	/** All the search times since the last call to GetAndTrimSearchResults. */
 	TArray<double> SearchTimes;
diff --git a/Engine/Source/Runtime/AssetRegistry/Private/AssetDataGathererPrivate.h b/Engine/Source/Runtime/AssetRegistry/Private/AssetDataGathererPrivate.h
--- a/Engine/Source/Runtime/AssetRegistry/Private/AssetDataGathererPrivate.h
+++ b/Engine/Source/Runtime/AssetRegistry/Private/AssetDataGathererPrivate.h
@@ -85,16 +85,18 @@ struct FDiscoveredPathData
 	FDateTime PackageTimestamp;
 	/** The type of file that was found */
 	EGatherableFileType Type = EGatherableFileType::Invalid;
+	/** Whether the file should just be reported as blocked and not further data gathered */
+	bool bBlocked = false;
 
 	FDiscoveredPathData() = default;
 	FDiscoveredPathData(FStringView InLocalAbsPath, FStringView InLongPackageName, FStringView InRelPath,
-		const FDateTime& InPackageTimestamp, EGatherableFileType InType);
+		const FDateTime& InPackageTimestamp, EGatherableFileType InType, bool bInBlocked);
 	FDiscoveredPathData(FStringView InLocalAbsPath, FStringView InLongPackageName, FStringView InRelPath,
-		EGatherableFileType InType);
+		EGatherableFileType InType, bool bInBlocked);
 	void Assign(FStringView InLocalAbsPath, FStringView InLongPackageName, FStringView InRelPath,
-		const FDateTime& InPackageTimestamp, EGatherableFileType InType);
+		const FDateTime& InPackageTimestamp, EGatherableFileType InType, bool bInBlocked);
 	void Assign(FStringView InLocalAbsPath, FStringView InLongPackageName, FStringView InRelPath,
-		EGatherableFileType InType);
+		EGatherableFileType InType, bool bInBlocked);
 
 	/** Return the total amount of heap memory used by the gatherer (including not-yet-claimed search results). */
 	SIZE_T GetAllocatedSize() const;
@@ -111,14 +113,16 @@ struct FGatheredPathData
 	FDateTime PackageTimestamp;
 	/** The type of file that was found */
 	EGatherableFileType Type = EGatherableFileType::Invalid;
+	/** Whether the file should just be reported as blocked and not further data gathered */
+	bool bBlocked = false;
 
 	FGatheredPathData() = default;
 	FGatheredPathData(FStringView InLocalAbsPath, FStringView InLongPackageName, const FDateTime& InPackageTimestamp,
-		EGatherableFileType InType);
+		EGatherableFileType InType, bool bInBlocked);
 	explicit FGatheredPathData(const FDiscoveredPathData& DiscoveredData);
 	explicit FGatheredPathData(FDiscoveredPathData&& DiscoveredData);
 	void Assign(FStringView InLocalAbsPath, FStringView InLongPackageName, const FDateTime& InPackageTimestamp,
-		EGatherableFileType InType);
+		EGatherableFileType InType, bool bInBlocked);
 	void Assign(const FDiscoveredPathData& DiscoveredData);
 
 	/**
diff --git a/Engine/Source/Runtime/AssetRegistry/Private/AssetRegistry.cpp b/Engine/Source/Runtime/AssetRegistry/Private/AssetRegistry.cpp
--- a/Engine/Source/Runtime/AssetRegistry/Private/AssetRegistry.cpp
+++ b/Engine/Source/Runtime/AssetRegistry/Private/AssetRegistry.cpp
@@ -4003,6 +4003,13 @@ Impl::EGatherStatus FAssetRegistryImpl::TickGatherer(Impl::FEventContext& EventC
 		VerseFilesGathered(TickStartTime, BackgroundResults.VerseFiles);
 	}
 
+	// Store blocked files to be reported
+	if (BackgroundResults.BlockedFiles.Num())
+	{
+		EventContext.BlockedFiles.Append(MoveTemp(BackgroundResults.BlockedFiles));
+		BackgroundResults.BlockedFiles.Reset();
+	}
+
 	// Load Calculated Dependencies when the gather from disk is complete; the full gather is not complete until after this is done
 	int32 NumGatherFromDiskPending = GetNumGatherFromDiskPending();
 #if WITH_EDITOR
@@ -6893,6 +6900,7 @@ void FEventContext::Clear()
 	PathEvents.Empty();
 	AssetEvents.Empty();
 	RequiredLoads.Empty();
+	BlockedFiles.Empty();
 }
 
 bool FEventContext::IsEmpty() const
@@ -6901,7 +6909,8 @@ bool FEventContext::IsEmpty() const
 		!ProgressUpdateData.IsSet() &&
 		PathEvents.Num() == 0 &&
 		AssetEvents.Num() == 0 &&
-		RequiredLoads.Num() == 0;
+		RequiredLoads.Num() == 0 &&
+		BlockedFiles.Num() == 0;
 }
 
 void FEventContext::Append(FEventContext&& Other)
@@ -6918,6 +6927,7 @@ void FEventContext::Append(FEventContext&& Other)
 	PathEvents.Append(MoveTemp(Other.PathEvents));
 	AssetEvents.Append(MoveTemp(Other.AssetEvents));
 	RequiredLoads.Append(MoveTemp(Other.RequiredLoads));
+	BlockedFiles.Append(MoveTemp(Other.BlockedFiles));
 }
 
 }
@@ -7012,9 +7022,19 @@ void UAssetRegistryImpl::Broadcast(UE::AssetRegistry::Impl::FEventContext& Event
 		}
 		EventContext.RequiredLoads.Empty();
 	}
+	if (EventContext.BlockedFiles.Num())
+	{
+		FilesBlockedEvent.Broadcast(EventContext.BlockedFiles);
+		EventContext.BlockedFiles.Empty();
+	}
 }
 
 
+UAssetRegistryImpl::FFilesBlockedEvent& UAssetRegistryImpl::OnFilesBlocked()
+{
+	return FilesBlockedEvent;
+}
+
 UAssetRegistryImpl::FPathAddedEvent& UAssetRegistryImpl::OnPathAdded()
 {
 	return PathAddedEvent;
diff --git a/Engine/Source/Runtime/AssetRegistry/Private/AssetRegistry.h b/Engine/Source/Runtime/AssetRegistry/Private/AssetRegistry.h
--- a/Engine/Source/Runtime/AssetRegistry/Private/AssetRegistry.h
+++ b/Engine/Source/Runtime/AssetRegistry/Private/AssetRegistry.h
@@ -132,6 +132,9 @@ public:
 	virtual bool ContainsTag(FName TagName) const override;
 	virtual void InitializeSerializationOptions(FAssetRegistrySerializationOptions& Options, const FString& PlatformIniName = FString(), UE::AssetRegistry::ESerializationTarget Target = UE::AssetRegistry::ESerializationTarget::ForGame) const override;
 
+	DECLARE_DERIVED_EVENT( UAssetRegistryImpl, IAssetRegistry::FFilesBlockedEvent, FFilesBlockedEvent);
+	virtual FFilesBlockedEvent& OnFilesBlocked() override;
+
 	DECLARE_DERIVED_EVENT( UAssetRegistryImpl, IAssetRegistry::FPathAddedEvent, FPathAddedEvent);
 	virtual FPathAddedEvent& OnPathAdded() override;
 
@@ -284,6 +287,9 @@ private:
 	bool bAddMetaDataTagsToOnGetExtraObjectTags = true;
 #endif
 
+	/** The delegate to execute when one or more files have been blocked from the registry */
+	FFilesBlockedEvent FilesBlockedEvent;
+
 	/** The delegate to execute when an asset path is added to the registry */
 	FPathAddedEvent PathAddedEvent;
 
diff --git a/Engine/Source/Runtime/AssetRegistry/Private/AssetRegistryImpl.h b/Engine/Source/Runtime/AssetRegistry/Private/AssetRegistryImpl.h
--- a/Engine/Source/Runtime/AssetRegistry/Private/AssetRegistryImpl.h
+++ b/Engine/Source/Runtime/AssetRegistry/Private/AssetRegistryImpl.h
@@ -482,6 +482,7 @@ struct FEventContext
 	TArray<TPair<FString, EEvent>> PathEvents;
 	TArray<TPair<FAssetData, EEvent>> AssetEvents;
 	TArray<FString> RequiredLoads;
+	TArray<FString> BlockedFiles;
 	bool bFileLoadedEventBroadcast = false;
 
 	/** Remove all stored events */
diff --git a/Engine/Source/Runtime/AssetRegistry/Public/AssetRegistry/IAssetRegistry.h b/Engine/Source/Runtime/AssetRegistry/Public/AssetRegistry/IAssetRegistry.h
--- a/Engine/Source/Runtime/AssetRegistry/Public/AssetRegistry/IAssetRegistry.h
+++ b/Engine/Source/Runtime/AssetRegistry/Public/AssetRegistry/IAssetRegistry.h
@@ -590,6 +590,10 @@ public:
 	UFUNCTION(BlueprintCallable, Category = "AssetRegistry")
 	virtual void ScanModifiedAssetFiles(const TArray<FString>& InFilePaths) = 0;
 
+	/** Event for when one or more files have been blocked from the registry */
+	DECLARE_EVENT_OneParam( IAssetRegistry, FFilesBlockedEvent, const TArray<FString>& /*Files*/ );
+	virtual FFilesBlockedEvent& OnFilesBlocked() = 0;
+
 	/** Event for when paths are added to the registry */
 	DECLARE_TS_MULTICAST_DELEGATE_OneParam( FPathAddedEvent, const FString& /*Path*/ );
 	virtual FPathAddedEvent& OnPathAdded() = 0;
diff --git a/Engine/Source/Runtime/Core/Private/Windows/WindowsPlatformFile.cpp b/Engine/Source/Runtime/Core/Private/Windows/WindowsPlatformFile.cpp
--- a/Engine/Source/Runtime/Core/Private/Windows/WindowsPlatformFile.cpp
+++ b/Engine/Source/Runtime/Core/Private/Windows/WindowsPlatformFile.cpp
@@ -49,6 +49,9 @@ static FAutoConsoleCommand TriggerFailedMicrosoftReadCmd(
 	namespace FileConstants
 	{
 		uint32 WIN_INVALID_SET_FILE_POINTER = INVALID_SET_FILE_POINTER;
+
+		/** Special suffix to access Alternate Data Stream "Zone.Identifier" for NTFS files */
+		constexpr FStringView ZoneIdentifierStreamSuffix(TEXTVIEW(":Zone.Identifier"));
 	}
 #include "Windows/HideWindowsPlatformTypes.h"
 
@@ -978,6 +981,12 @@ class CORE_API FWindowsPlatformFile : public IPhysicalPlatformFile
 			NormalizeWindowsPath(*this, /* bIsFilename */ true);
 		}
 
+		explicit FNormalizedFilename(FStringView Filename)
+		{
+			Append(Filename);
+			NormalizeWindowsPath(*this, /* bIsFilename */ true);
+		}
+
 		explicit FNormalizedFilename(FStringView Dir, FStringView Filename)
 		{
 			Append(Dir);
@@ -1180,6 +1189,57 @@ public:
 		return ESymlinkResult::NonSymlink;
 	}
 
+	virtual bool HasMarkOfTheWeb(FStringView Filename, FString* OutSourceURL = nullptr) override
+	{
+		FNormalizedFilename StreamPath(Filename);
+		StreamPath.Append(FileConstants::ZoneIdentifierStreamSuffix);
+
+		const uint32 ZoneId = GetPrivateProfileIntW(TEXT("ZoneTransfer"), TEXT("ZoneId"), -1 /* URLZONE_INVALID */, *StreamPath);
+		if (ZoneId != 2 /* URLZONE_TRUSTED */ && ZoneId != 3 /* URLZONE_INTERNET */ && ZoneId != 4 /* URLZONE_UNTRUSTED */)
+		{
+			return false;
+		}
+
+		if (OutSourceURL != nullptr)
+		{
+			// Can't use FConfigFile to parse ini because it treats double slashes in the url (e.g. HostUrl=http://...) as the beginning of a comment.
+			TCHAR HostUrl[2048]; // @todo: remove hardcoded url length limit by using a simple standalone ini parser
+			const int32 HostUrlLen = static_cast<int32>(GetPrivateProfileStringW(TEXT("ZoneTransfer"), TEXT("HostUrl"), nullptr, HostUrl, static_cast<DWORD>(GetNum(HostUrl)), *StreamPath));
+			if (0 < HostUrlLen && HostUrlLen <= GetNum(HostUrl))
+			{
+				*OutSourceURL = FString(HostUrlLen, HostUrl);
+			}
+		}
+
+		return true;
+	}
+
+	virtual bool SetMarkOfTheWeb(FStringView Filename, bool bNewStatus, const FString* InSourceURL = nullptr) override
+	{
+		FNormalizedFilename StreamPath(Filename);
+		StreamPath.Append(FileConstants::ZoneIdentifierStreamSuffix);
+
+		if (!bNewStatus)
+		{
+			return !!DeleteFileW(*StreamPath) || GetLastError() == ERROR_FILE_NOT_FOUND;
+		}
+
+		if (!WritePrivateProfileStringW(TEXT("ZoneTransfer"), TEXT("ZoneId"), TEXT("3") /* URLZONE_INTERNET */, *StreamPath))
+		{
+			return false;
+		}
+
+		if (InSourceURL != nullptr)
+		{
+			if (!WritePrivateProfileStringW(TEXT("ZoneTransfer"), TEXT("HostUrl"), **InSourceURL, *StreamPath))
+			{
+				return false;
+			}
+		}
+
+		return true;
+	}
+
 #define USE_WINDOWS_ASYNC_IMPL 0
 #if USE_WINDOWS_ASYNC_IMPL
 	virtual IAsyncReadFileHandle* OpenAsyncRead(const TCHAR* Filename) override
diff --git a/Engine/Source/Runtime/Core/Public/GenericPlatform/GenericPlatformFile.h b/Engine/Source/Runtime/Core/Public/GenericPlatform/GenericPlatformFile.h
--- a/Engine/Source/Runtime/Core/Public/GenericPlatform/GenericPlatformFile.h
+++ b/Engine/Source/Runtime/Core/Public/GenericPlatform/GenericPlatformFile.h
@@ -8,6 +8,7 @@
 #pragma once
 
 #include "Containers/Array.h"
+#include "Containers/StringView.h"
 #include "Containers/UnrealString.h"
 #include "CoreTypes.h"
 #include "Misc/AssertionMacros.h"
@@ -331,6 +332,29 @@ public:
 	/** Return true if the file is a symbolic link */
 	virtual ESymlinkResult IsSymlink(const TCHAR* Filename) { return ESymlinkResult::Unimplemented; }
 
+	/**
+	 * Determine if the file has been downloaded from a web browser, based on platform-specific metadata.
+	 *
+	 * @param Filename The file to check if it has a mark of the web.
+	 * @param OutSourceURL An optional pointer to return a source url if available.
+	 *
+	 * @return true if the file has a mark of the web, false otherwise.
+	 * @remark Only works for files on disk and currently only implemented for Windows.
+	 */
+	virtual bool HasMarkOfTheWeb(FStringView Filename, FString* OutSourceURL = nullptr) { return false; }
+
+	/**
+	 * Attempt to change the platform-specific metadata that indicates if the file has been downloaded from a web browser.
+	 *
+	 * @param Filename The file to change it's mark of the web status.
+	 * @param bNewStatus New mark of the web status for the file.
+	 * @param InSourceURL An optional pointer to a source url that will be applied if new status true.
+	 *
+	 * @return true if the file's mark of the web status was successful changed, false otherwise.
+	 * @remark Only works for files on disk and currently only implemented for Windows.
+	 */
+	virtual bool SetMarkOfTheWeb(FStringView Filename, bool bNewStatus, const FString* InSourceURL = nullptr) { return false; }
+
 	/** Attempt to open a file for reading.
 	 *
 	 * @param Filename file to be opened
