AssetRegistry: Add the tags from "GetAssetRegistryTagsExtended" that are stored in on-disk AssetDatas to the tags returned from GetAssetRegistryTags calls when GetAssetData queries return results from in-memory UObjects.
#rnx

diff --git a/Engine/Source/Runtime/AssetRegistry/Private/AssetRegistry.cpp b/Engine/Source/Runtime/AssetRegistry/Private/AssetRegistry.cpp
--- a/Engine/Source/Runtime/AssetRegistry/Private/AssetRegistry.cpp
+++ b/Engine/Source/Runtime/AssetRegistry/Private/AssetRegistry.cpp
@@ -1862,6 +1862,7 @@ bool UAssetRegistryImpl::GetAssets(const FARFilter& InFilter, TArray<FAssetData>
 	{
 		bool bStopIterationUnused;
 		EnumerateMemoryAssets(CompiledFilter, PackagesToSkip, bStopIterationUnused,
+			InterfaceLock, GuardedData.GetState(),
 			[&OutAssetData](FAssetData&& AssetData)
 			{
 				OutAssetData.Add(MoveTemp(AssetData));
@@ -1904,6 +1905,7 @@ bool UAssetRegistryImpl::EnumerateAssets(const FARCompiledFilter& InFilter, TFun
 	{
 		bool bStopIteration;
 		EnumerateMemoryAssets(InFilter, PackagesToSkip, bStopIteration,
+			InterfaceLock, GuardedData.GetState(),
 			[&Callback](FAssetData&& AssetData)
 			{
 				return Callback(AssetData);
@@ -1939,6 +1941,30 @@ namespace UE::AssetRegistry
 namespace Utils
 {
 
+TOptional<FAssetDataTagMap> AddNonOverlappingTags(FAssetData& ExistingAssetData, const FAssetData& NewAssetData)
+{
+	TOptional<FAssetDataTagMap> ModifiedTags;
+	NewAssetData.TagsAndValues.ForEach([&ExistingAssetData, &ModifiedTags](const TPair<FName, FAssetTagValueRef>& TagPair)
+		{
+			if (ModifiedTags)
+			{
+				if (!ModifiedTags->Contains(TagPair.Key))
+				{
+					ModifiedTags->Add(TagPair.Key, TagPair.Value.GetStorageString());
+				}
+			}
+			else
+			{
+				if (!ExistingAssetData.TagsAndValues.Contains(TagPair.Key))
+				{
+					ModifiedTags.Emplace(ExistingAssetData.TagsAndValues.CopyMap());
+					ModifiedTags->Add(TagPair.Key, TagPair.Value.GetStorageString());
+				}
+			}
+		});
+	return ModifiedTags;
+}
+
 void EnumerateMemoryAssetsHelper(const FARCompiledFilter& InFilter, TSet<FName>& OutPackageNamesWithAssets,
 	bool& bOutStopIteration, TFunctionRef<bool(const UObject* Object, FAssetData&& PartialAssetData)> Callback,
 	bool bSkipARFilteredAssets)
@@ -2084,13 +2110,29 @@ void EnumerateMemoryAssetsHelper(const FARCompiledFilter& InFilter, TSet<FName>&
 }
 
 void EnumerateMemoryAssets(const FARCompiledFilter& InFilter, TSet<FName>& OutPackageNamesWithAssets,
-	bool& bOutStopIteration, TFunctionRef<bool(FAssetData&&)> Callback, bool bSkipARFilteredAssets)
+	bool& bOutStopIteration, FRWLock& InterfaceLock, const FAssetRegistryState& GuardedDataState,
+	TFunctionRef<bool(FAssetData&&)> Callback, bool bSkipARFilteredAssets)
 {
 	check(!InFilter.IsEmpty() && Utils::IsFilterValid(InFilter));
 	EnumerateMemoryAssetsHelper(InFilter, OutPackageNamesWithAssets, bOutStopIteration,
-		[&InFilter, &Callback](const UObject* Object, FAssetData&& PartialAssetData)
+		[&InFilter, &Callback, &InterfaceLock, &GuardedDataState](const UObject* Object, FAssetData&& PartialAssetData)
 		{
 			Object->GetAssetRegistryTags(PartialAssetData);
+			{
+				// GetAssetRegistryTags call does not add extended tags that may exist in the on-disk Asset.
+				// Our contract for on-disk versus in-memory tags is that in-memory tags override on-disk tags, but we
+				// keep any on-disk tags that do not exist in the in-memory tags because they may be extended tags.
+				FReadScopeLock InterfaceScopeLock(InterfaceLock);
+				const FAssetData* OnDiskAssetData = GuardedDataState.GetAssetByObjectPath(FSoftObjectPath(Object));
+				if (OnDiskAssetData)
+				{
+					TOptional<FAssetDataTagMap> ModifiedTags = Utils::AddNonOverlappingTags(PartialAssetData, *OnDiskAssetData);
+					if (ModifiedTags)
+					{
+						PartialAssetData.TagsAndValues = FAssetDataTagMapSharedView(MoveTemp(*ModifiedTags));
+					}
+				}
+			}
 			// After adding tags, PartialAssetData is now a full AssetData
 
 			// Tags and values
@@ -5286,25 +5328,7 @@ void FAssetRegistryImpl::UpdateAssetData(Impl::FEventContext& EventContext, FAss
 void FAssetRegistryImpl::AddNonOverlappingTags(Impl::FEventContext& EventContext, FAssetData& ExistingAssetData,
 	const FAssetData& NewAssetData)
 {
-	TOptional<FAssetDataTagMap> ModifiedTags;
-	NewAssetData.TagsAndValues.ForEach([&ExistingAssetData, &ModifiedTags](const TPair<FName, FAssetTagValueRef>& TagPair)
-		{
-			if (ModifiedTags)
-			{
-				if (!ModifiedTags->Contains(TagPair.Key))
-				{
-					ModifiedTags->Add(TagPair.Key, TagPair.Value.GetStorageString());
-				}
-			}
-			else
-			{
-				if (!ExistingAssetData.TagsAndValues.Contains(TagPair.Key))
-				{
-					ModifiedTags.Emplace(ExistingAssetData.TagsAndValues.CopyMap());
-					ModifiedTags->Add(TagPair.Key, TagPair.Value.GetStorageString());
-				}
-			}
-		});
+	TOptional<FAssetDataTagMap> ModifiedTags = Utils::AddNonOverlappingTags(ExistingAssetData, NewAssetData);
 	if (ModifiedTags)
 	{
 		State.SetTagsOnExistingAsset(&ExistingAssetData, MoveTemp(*ModifiedTags));
diff --git a/Engine/Source/Runtime/AssetRegistry/Private/AssetRegistryImpl.h b/Engine/Source/Runtime/AssetRegistry/Private/AssetRegistryImpl.h
--- a/Engine/Source/Runtime/AssetRegistry/Private/AssetRegistryImpl.h
+++ b/Engine/Source/Runtime/AssetRegistry/Private/AssetRegistryImpl.h
@@ -613,7 +613,8 @@ void EnumerateMemoryAssetsHelper(const FARCompiledFilter& InFilter, TSet<FName>&
  * Fills in OutPackageNamesWithAssets with names of all packages tested.
 */
 void EnumerateMemoryAssets(const FARCompiledFilter& InFilter, TSet<FName>& OutPackageNamesWithAssets,
-	bool& bOutStopIteration, TFunctionRef<bool(FAssetData&&)> Callback, bool bSkipARFilteredAssets);
+	bool& bOutStopIteration, FRWLock& InterfaceLock, const FAssetRegistryState& GuardedDataState,
+	TFunctionRef<bool(FAssetData&&)> Callback, bool bSkipARFilteredAssets);
 
 }
 
