AssetRegistry MetaData Tags: The addition of AssetDataTags from disk onto the AssetDataTags calculated by GetAssetRegistryTags should only occur for cooked assets. Otherwise it makes metadata tags impossible to remove even when an uncooked asset is resaved.
#rnx

diff --git a/Engine/Source/Runtime/AssetRegistry/Private/AssetRegistry.cpp b/Engine/Source/Runtime/AssetRegistry/Private/AssetRegistry.cpp
--- a/Engine/Source/Runtime/AssetRegistry/Private/AssetRegistry.cpp
+++ b/Engine/Source/Runtime/AssetRegistry/Private/AssetRegistry.cpp
@@ -6399,18 +6399,24 @@ void UAssetRegistryImpl::OnGetExtraObjectTags(const UObject* Object, TArray<UObj
 {
 	if (bAddMetaDataTagsToOnGetExtraObjectTags)
 	{
-		TSet<FName>& MetaDataTags = UObject::GetMetaDataTagsForAssetRegistry();
 		// It is critical that bIncludeOnlyOnDiskAssets=true otherwise this will cause an infinite loop
 		const FAssetData AssetData = GetAssetByObjectPath(FSoftObjectPath(Object), /*bIncludeOnlyOnDiskAssets=*/true);
-		for (const FName MetaDataTag : MetaDataTags)
+		// Adding metadata tags from disk is only necessary for cooked assets; uncooked assets still have the metadata and add them elsewhere
+		// in UObject::GetAssetRegistryTags. Adding the tags from disk into uncooked assets would make the tags impossible to remove when
+		// the uncooked assets are resaved.
+		if ((AssetData.PackageFlags & PKG_Cooked) != 0)
 		{
-			auto OutTagsContainsTagPredicate = [MetaDataTag](const UObject::FAssetRegistryTag& Tag) { return Tag.Name == MetaDataTag; };
-			if (!OutTags.ContainsByPredicate(OutTagsContainsTagPredicate))
+			TSet<FName>& MetaDataTags = UObject::GetMetaDataTagsForAssetRegistry();
+			for (const FName MetaDataTag : MetaDataTags)
 			{
-				FAssetTagValueRef TagValue = AssetData.TagsAndValues.FindTag(MetaDataTag);
-				if (TagValue.IsSet())
+				auto OutTagsContainsTagPredicate = [MetaDataTag](const UObject::FAssetRegistryTag& Tag) { return Tag.Name == MetaDataTag; };
+				if (!OutTags.ContainsByPredicate(OutTagsContainsTagPredicate))
 				{
-					OutTags.Add(UObject::FAssetRegistryTag(MetaDataTag, TagValue.AsString(), UObject::FAssetRegistryTag::TT_Alphabetical));
+					FAssetTagValueRef TagValue = AssetData.TagsAndValues.FindTag(MetaDataTag);
+					if (TagValue.IsSet())
+					{
+						OutTags.Add(UObject::FAssetRegistryTag(MetaDataTag, TagValue.AsString(), UObject::FAssetRegistryTag::TT_Alphabetical));
+					}
 				}
 			}
 		}
