Fixed Verse compile error after reopening the same project

#rnx

diff --git a/Engine/Source/Runtime/AssetRegistry/Private/AssetRegistry.cpp b/Engine/Source/Runtime/AssetRegistry/Private/AssetRegistry.cpp
--- a/Engine/Source/Runtime/AssetRegistry/Private/AssetRegistry.cpp
+++ b/Engine/Source/Runtime/AssetRegistry/Private/AssetRegistry.cpp
@@ -6041,7 +6041,7 @@ void FAssetRegistryImpl::OnContentPathDismounted(Impl::FEventContext& EventConte
 		GlobalGatherer->RemoveMountPoint(FileSystemPath);
 	}
 
-	// Remove all cached assets found at this location
+	// Remove all cached assets and Verse files found at this location
 	{
 		FName AssetPathNoTrailingSlashFName(*AssetPathNoTrailingSlash);
 		TArray<FAssetData*> AllAssetDataToRemove;
@@ -6051,11 +6051,23 @@ void FAssetRegistryImpl::OnContentPathDismounted(Impl::FEventContext& EventConte
 		PathList.Add(AssetPathNoTrailingSlashFName);
 		for (FName PathName : PathList)
 		{
+			// Gather assets
 			TArray<FAssetData*>* AssetsInPath = State.CachedAssetsByPath.Find(PathName);
 			if (AssetsInPath)
 			{
 				AllAssetDataToRemove.Append(*AssetsInPath);
 			}
+
+			// Forget Verse files
+			const TArray<FName>* VerseFilesInPath = CachedVerseFilesByPath.Find(PathName);
+			if (VerseFilesInPath)
+			{
+				for (FName FilePath : *VerseFilesInPath)
+				{
+					CachedVerseFiles.Remove(FilePath);
+				}
+				CachedVerseFilesByPath.Remove(PathName);
+			}
 		}
 
 		for (FAssetData* AssetData : AllAssetDataToRemove)
