AssetRegistry method to test the existence of verse files under a specific folder

diff --git a/Engine/Source/Runtime/AssetRegistry/Private/AssetRegistry.cpp b/Engine/Source/Runtime/AssetRegistry/Private/AssetRegistry.cpp
--- a/Engine/Source/Runtime/AssetRegistry/Private/AssetRegistry.cpp
+++ b/Engine/Source/Runtime/AssetRegistry/Private/AssetRegistry.cpp
@@ -3322,35 +3322,48 @@ void PrioritizeAssetInstall(const FAssetData& AssetData)
 
 }
 
+bool UAssetRegistryImpl::HasVerseFiles(FName PackagePath, bool bRecursive /*= false*/) const
+{
+	FReadScopeLock InterfaceScopeLock(InterfaceLock);
+	return GuardedData.GetVerseFilesByPath(PackagePath, /*OutFilePaths=*/nullptr, bRecursive);
+}
+
 bool UAssetRegistryImpl::GetVerseFilesByPath(FName PackagePath, TArray<FName>& OutFilePaths, bool bRecursive /*= false*/) const
 {
 	FReadScopeLock InterfaceScopeLock(InterfaceLock);
-	return GuardedData.GetVerseFilesByPath(PackagePath, OutFilePaths, bRecursive);
+	return GuardedData.GetVerseFilesByPath(PackagePath, &OutFilePaths, bRecursive);
 }
 
 namespace UE::AssetRegistry
 {
 
-bool FAssetRegistryImpl::GetVerseFilesByPath(FName PackagePath, TArray<FName>& OutFilePaths, bool bRecursive /*= false*/) const
+bool FAssetRegistryImpl::GetVerseFilesByPath(FName PackagePath, TArray<FName>* OutFilePaths, bool bRecursive) const
 {
-	bool bFoundAnything = false;
 	TSet<FName> PathList;
 	PathList.Reserve(32);
+	PathList.Add(PackagePath);
 	if (bRecursive)
 	{
 		CachedPathTree.GetSubPaths(PackagePath, PathList, true);
 	}
-	PathList.Add(PackagePath);
+
+	bool bFoundAnything = false;
 	for (const FName& PathName : PathList)
 	{
 		const TArray<FName>* FilePaths = CachedVerseFilesByPath.Find(PathName);
 		if (FilePaths)
 		{
-			OutFilePaths.Append(*FilePaths);
 			bFoundAnything = true;
+			if (OutFilePaths)
+			{
+				OutFilePaths->Append(*FilePaths);
+			}
+			else
+			{
+				break;
+			}
 		}
 	}
-
 	return bFoundAnything;
 }
 
diff --git a/Engine/Source/Runtime/AssetRegistry/Private/AssetRegistry.h b/Engine/Source/Runtime/AssetRegistry/Private/AssetRegistry.h
--- a/Engine/Source/Runtime/AssetRegistry/Private/AssetRegistry.h
+++ b/Engine/Source/Runtime/AssetRegistry/Private/AssetRegistry.h
@@ -89,6 +89,7 @@ public:
 	virtual float GetAssetAvailabilityProgress(const FAssetData& AssetData, EAssetAvailabilityProgressReportingType::Type ReportType) const override;
 	virtual bool GetAssetAvailabilityProgressTypeSupported(EAssetAvailabilityProgressReportingType::Type ReportType) const override;
 	virtual void PrioritizeAssetInstall(const FAssetData& AssetData) const override;
+	virtual bool HasVerseFiles(FName PackagePath, bool bRecursive = false) const override;
 	virtual bool GetVerseFilesByPath(FName PackagePath, TArray<FName>& OutFilePaths, bool bRecursive = false) const override;
 	virtual bool AddPath(const FString& PathToAdd) override;
 	virtual bool RemovePath(const FString& PathToRemove) override;
diff --git a/Engine/Source/Runtime/AssetRegistry/Private/AssetRegistryImpl.h b/Engine/Source/Runtime/AssetRegistry/Private/AssetRegistryImpl.h
--- a/Engine/Source/Runtime/AssetRegistry/Private/AssetRegistryImpl.h
+++ b/Engine/Source/Runtime/AssetRegistry/Private/AssetRegistryImpl.h
@@ -135,7 +135,7 @@ public:
 	bool AddPath(Impl::FEventContext& EventContext, const FString& PathToAdd);
 	void SearchAllAssets(Impl::FEventContext& EventContext, Impl::FClassInheritanceContext& InheritanceContext,
 		bool bSynchronousSearch);
-	bool GetVerseFilesByPath(FName PackagePath, TArray<FName>& OutFilePaths, bool bRecursive = false) const;
+	bool GetVerseFilesByPath(FName PackagePath, TArray<FName>* OutFilePaths, bool bRecursive) const;
 	void ScanPathsSynchronous(Impl::FScanPathContext& Context);
 	void PrioritizeSearchPath(const FString& PathToPrioritize);
 	void ScanModifiedAssetFiles(Impl::FEventContext& EventContext, Impl::FClassInheritanceContext& InheritanceContext,
diff --git a/Engine/Source/Runtime/AssetRegistry/Public/AssetRegistry/IAssetRegistry.h b/Engine/Source/Runtime/AssetRegistry/Public/AssetRegistry/IAssetRegistry.h
--- a/Engine/Source/Runtime/AssetRegistry/Public/AssetRegistry/IAssetRegistry.h
+++ b/Engine/Source/Runtime/AssetRegistry/Public/AssetRegistry/IAssetRegistry.h
@@ -502,11 +502,20 @@ public:
 	virtual void PrioritizeAssetInstall(const FAssetData& AssetData) const = 0;
 
 	/**
-	 * Gets paths for all Verse files in the supplied folder path
+	 * Returns whether the supplied folder contains any verse file
+	 * 
+	 * @param PackagePath the path to query asset data in (e.g. /Game/MyFolder)
+	 * @param bRecursive if true, the supplied path will be searched recursively
+	 * @return True if any verse file is found under the supplied folder
+	 */
+	virtual bool HasVerseFiles(FName PackagePath, bool bRecursive = false) const = 0;
+
+	/**
+	 * Gets paths for all Verse files in the supplied folder
 	 *
 	 * @param PackagePath the path to query asset data in (e.g. /Game/MyFolder)
 	 * @param OutFilePaths the list of Verse files in this path, as pseudo UE LongPackagePaths with extension (e.g. /Game/MyFolder/MyVerseFile.verse)
-	 * @param bRecursive if true, all supplied paths will be searched recursively
+	 * @param bRecursive if true, the supplied path will be searched recursively
 	 */
 	virtual bool GetVerseFilesByPath(FName PackagePath, TArray<FName>& OutFilePaths, bool bRecursive = false) const = 0;
 
