AssetRegistry EnumerateMemoryObjects - do not call the GetAssetRegistryTags callback from within ForEachObjectWithPackage, because GetAssetRegistryTags is arbitrary code and might create UObjects, which is disallowed during ForEachObjectWIthPackage.
#rnx

diff --git a/Engine/Source/Runtime/AssetRegistry/Private/AssetRegistry.cpp b/Engine/Source/Runtime/AssetRegistry/Private/AssetRegistry.cpp
--- a/Engine/Source/Runtime/AssetRegistry/Private/AssetRegistry.cpp
+++ b/Engine/Source/Runtime/AssetRegistry/Private/AssetRegistry.cpp
@@ -2039,15 +2039,44 @@ void EnumerateMemoryAssetsHelper(const FARCompiledFilter& InFilter, TSet<FName>&
 	};
 
 	// Iterate over all in-memory assets to find the ones that pass the filter components
-	if (InFilter.ClassPaths.Num() > 0)
+	if (InFilter.ClassPaths.Num() > 0 || InFilter.PackageNames.Num() > 0)
 	{
-		TArray<UObject*> InMemoryObjects;
-		for (FTopLevelAssetPath ClassName : InFilter.ClassPaths)
+		TArray<UObject*, TInlineAllocator<10>> InMemoryObjects;
+		if (InFilter.ClassPaths.Num())
 		{
-			UClass* Class = FindObject<UClass>(ClassName);
-			if (Class != nullptr)
+			for (FTopLevelAssetPath ClassName : InFilter.ClassPaths)
 			{
-				GetObjectsOfClass(Class, InMemoryObjects, false, RF_NoFlags);
+				UClass* Class = FindObject<UClass>(ClassName);
+				if (Class != nullptr)
+				{
+					ForEachObjectOfClass(Class, [&InMemoryObjects](UObject* Object)
+						{
+							InMemoryObjects.Add(Object);
+						}, false /* bIncludeDerivedClasses */, RF_NoFlags);
+				}
+			}
+		}
+		else
+		{
+			for (FName PackageName : InFilter.PackageNames)
+			{
+				UPackage* Package = FindObjectFast<UPackage>(nullptr, PackageName);
+				if (Package != nullptr)
+				{
+					// Store objects in an intermediate rather than calling FilterInMemoryObjectLambda on them directly
+					// because the callback is arbitrary code and might create UObjects, which is disallowed in
+					// ForEachObjectWithPackage
+					ForEachObjectWithPackage(Package, [&InMemoryObjects](UObject* Object)
+						{
+							// Avoid adding an element to InMemoryObjects for every UObject
+							// There could be many UObjects (thousands) but only a single Asset
+							if (Object->IsAsset())
+							{
+								InMemoryObjects.Add(Object);
+							}
+							return true;
+						});
+				}
 			}
 		}
 
@@ -2062,27 +2091,6 @@ void EnumerateMemoryAssetsHelper(const FARCompiledFilter& InFilter, TSet<FName>&
 			}
 		}
 	}
-	else if (InFilter.PackageNames.Num() > 0)
-	{
-		for (FName PackageName : InFilter.PackageNames)
-		{
-			UPackage* Package = FindObjectFast<UPackage>(nullptr, PackageName);
-			if (Package)
-			{
-				bool bContinue = true;
-				ForEachObjectWithPackage(Package, [&FilterInMemoryObjectLambda, &bContinue](UObject* Object)
-					{
-						FilterInMemoryObjectLambda(Object, bContinue);
-						return bContinue;
-					});
-				if (!bContinue)
-				{
-					bOutStopIteration = true;
-					return;
-				}
-			}
-		}
-	}
 	else
 	{
 		for (FThreadSafeObjectIterator ObjIt; ObjIt; ++ObjIt)
