AssetRegistry: Optimize FAssetDataGatherer::GetAndTrimSearchResults by using FFilesToSearch::AvailableFilesNum instead of recalculating it with a treesearch.
18.4s -> 0.1s
#rnx

diff --git a/Engine/Source/Runtime/AssetRegistry/Private/AssetDataGatherer.cpp b/Engine/Source/Runtime/AssetRegistry/Private/AssetDataGatherer.cpp
--- a/Engine/Source/Runtime/AssetRegistry/Private/AssetDataGatherer.cpp
+++ b/Engine/Source/Runtime/AssetRegistry/Private/AssetDataGatherer.cpp
@@ -4748,7 +4748,7 @@ void FFilesToSearch::Shrink()
 
 int32 FFilesToSearch::Num() const
 {
-	return BlockingFiles.Num() + Root.NumFiles() + LaterRetryFiles.Num();
+	return AvailableFilesNum + LaterRetryFiles.Num();
 }
 
 int32 FFilesToSearch::GetNumAvailable() const
