MIsc typo fixes.

diff --git a/Engine/Source/Runtime/AssetRegistry/Private/AssetRegistry.cpp b/Engine/Source/Runtime/AssetRegistry/Private/AssetRegistry.cpp
--- a/Engine/Source/Runtime/AssetRegistry/Private/AssetRegistry.cpp
+++ b/Engine/Source/Runtime/AssetRegistry/Private/AssetRegistry.cpp
@@ -3148,7 +3148,7 @@ PRAGMA_DISABLE_DEPRECATION_WARNINGS
 				}
 				else
 				{
-					UE_LOG(LogAssetRegistry, Error, TEXT("Failed to resolve class path for short clas name \"%s\" when compiling asset registry filter"), *ClassName.ToString());
+					UE_LOG(LogAssetRegistry, Error, TEXT("Failed to resolve class path for short class name \"%s\" when compiling asset registry filter"), *ClassName.ToString());
 				}
 			}
 		}
