AssetRegistry: Add log to indicate when the AssetDataGatherer's monolithic cache is read (and how long it took) to diagnose whether failure to write the file is caused by it being open for read at the same time.
#rnx

diff --git a/Engine/Source/Runtime/AssetRegistry/Private/AssetDataGatherer.cpp b/Engine/Source/Runtime/AssetRegistry/Private/AssetDataGatherer.cpp
--- a/Engine/Source/Runtime/AssetRegistry/Private/AssetDataGatherer.cpp
+++ b/Engine/Source/Runtime/AssetRegistry/Private/AssetDataGatherer.cpp
@@ -3423,6 +3423,7 @@ void FAssetDataGatherer::TickInternal(bool& bOutIsTickInterrupt, double& TickSta
 	// Load the async cache if not yet loaded
 	if (bLoadMonolithicCache)
 	{
+		double CacheLoadStartTime = FPlatformTime::Seconds();
 		FCachePayload Payload;
 #if UE_EDITOR
 		if (GPreloadSettings.IsPreloadMonolithicCache())
@@ -3435,6 +3436,8 @@ void FAssetDataGatherer::TickInternal(bool& bOutIsTickInterrupt, double& TickSta
 			Payload = UE::AssetDataGather::Private::LoadCacheFile(GPreloadSettings.GetMonolithicCacheFilename());
 		}
 		ConsumeCacheFile(MoveTemp(Payload));
+		UE_LOG(LogAssetRegistry, Display, TEXT("AssetDataGatherer spent %.3fs loading %s."),
+			FPlatformTime::Seconds() - CacheLoadStartTime, *GPreloadSettings.GetMonolithicCacheFilename());
 
 		FGathererScopeLock ResultsScopeLock(&ResultsLock);
 		bHasLoadedMonolithicCache = true;
