AssetRegistry: AssetDataGatherer: Avoid a warning and a sleep when unable to write the monolithic cache file because we just finished reading it and the OS has not yet marked the file as deletable again.
#rnx

diff --git a/Engine/Source/Runtime/AssetRegistry/Private/AssetDataGatherer.cpp b/Engine/Source/Runtime/AssetRegistry/Private/AssetDataGatherer.cpp
--- a/Engine/Source/Runtime/AssetRegistry/Private/AssetDataGatherer.cpp
+++ b/Engine/Source/Runtime/AssetRegistry/Private/AssetDataGatherer.cpp
@@ -3441,6 +3441,12 @@ void FAssetDataGatherer::TickInternal(bool& bOutIsTickInterrupt, double& TickSta
 
 		FGathererScopeLock ResultsScopeLock(&ResultsLock);
 		bHasLoadedMonolithicCache = true;
+
+		// After we load the monolithic cache, restart the write timer for it. We don't need to save to if we just
+		// finished loading it (which we do before gathering anything) and we want to avoid failure to save due to
+		// writing a file that we just closed a readhandle for.
+		LastCacheWriteTime = FPlatformTime::Seconds();
+		LocalLastCacheWriteTime = LastCacheWriteTime;
 	}
 
 	struct FReadContext
@@ -4177,6 +4183,23 @@ void FAssetDataGatherer::SaveCacheFileInternal(const FString& CacheFilename, con
 #endif
 		// Close file handle before moving temp file to target 
 		FileAr.Reset();
+
+		// If we recently saved or loaded the file then pause for 0.5 seconds before trying to save on
+		// top of it, to avoid failure to be able to delete the file we just saved/loaded.
+		if (bIsAsyncCacheSave)
+		{
+			double LocalLastCacheWriteTime;
+			{
+				FGathererScopeLock ResultsScopeLock(&ResultsLock);
+				LocalLastCacheWriteTime = LastCacheWriteTime;
+			}
+			double CurrentTime = FPlatformTime::Seconds();
+			constexpr float WaitTimeBeforeReopen = 0.5f;
+			if (CurrentTime < LastCacheWriteTime + WaitTimeBeforeReopen)
+			{
+				FPlatformProcess::Sleep(static_cast<float>(LastCacheWriteTime + WaitTimeBeforeReopen - CurrentTime));
+			}
+		}
 		IFileManager::Get().Move(*CacheFilenameStr, *TempFilename);
 	}
 	else
@@ -4189,6 +4212,7 @@ void FAssetDataGatherer::SaveCacheFileInternal(const FString& CacheFilename, con
 		FScopedPause ScopedPause(*this);
 		FGathererScopeLock TickScopeLock(&TickLock);
 		bIsSavingAsyncCache = false;
+		FGathererScopeLock ResultsScopeLock(&ResultsLock);
 		LastCacheWriteTime = FPlatformTime::Seconds();
 	}
 }
