AssetRegistry: Optimize RunAssetsThroughFilter by using RemoveAll instead of doing N separate removeAt which shifts.
#rnx

diff --git a/Engine/Source/Runtime/AssetRegistry/Private/AssetRegistry.cpp b/Engine/Source/Runtime/AssetRegistry/Private/AssetRegistry.cpp
--- a/Engine/Source/Runtime/AssetRegistry/Private/AssetRegistry.cpp
+++ b/Engine/Source/Runtime/AssetRegistry/Private/AssetRegistry.cpp
@@ -3074,16 +3074,11 @@ void RunAssetsThroughFilter(TArray<FAssetData>& AssetDataList, const FARCompiled
 	const int32 OriginalArrayCount = AssetDataList.Num();
 	const bool bPassFilterValue = FilterMode == EFilterMode::Inclusive;
 
-	// Spin the array backwards to minimize the number of elements that are repeatedly moved down
-	for (int32 AssetDataIndex = AssetDataList.Num() - 1; AssetDataIndex >= 0; --AssetDataIndex)
-	{
-		const bool bFilterResult = RunAssetThroughFilter_Unchecked(AssetDataList[AssetDataIndex], CompiledFilter, bPassFilterValue);
-		if (bFilterResult != bPassFilterValue)
+	AssetDataList.RemoveAll([&CompiledFilter, bPassFilterValue](const FAssetData& AssetData)
 		{
-			AssetDataList.RemoveAt(AssetDataIndex, 1, /*bAllowShrinking*/false);
-			continue;
-		}
-	}
+			const bool bFilterResult = RunAssetThroughFilter_Unchecked(AssetData, CompiledFilter, bPassFilterValue);
+			return bFilterResult != bPassFilterValue;
+		});
 	if (OriginalArrayCount > AssetDataList.Num())
 	{
 		AssetDataList.Shrink();
