deprecated `FCoreDelegates::OnPakFileMounted2` because it's thread-safe but is used concurrently. Added a thread-safe version `FCoreDelegates::GetOnPakFileMounted2()` and switched all use cases to the new one

diff --git a/Engine/Plugins/Experimental/GameFeatures/Source/GameFeatures/Private/GameFeaturePluginStateMachine.cpp b/Engine/Plugins/Experimental/GameFeatures/Source/GameFeatures/Private/GameFeaturePluginStateMachine.cpp
--- a/Engine/Plugins/Experimental/GameFeatures/Source/GameFeatures/Private/GameFeaturePluginStateMachine.cpp
+++ b/Engine/Plugins/Experimental/GameFeatures/Source/GameFeatures/Private/GameFeaturePluginStateMachine.cpp
@@ -1776,7 +1776,7 @@ struct FGameFeaturePluginState_Mounting : public FGameFeaturePluginState
 			IInstallBundleManager::InstallBundleCompleteDelegate.AddRaw(this, &FGameFeaturePluginState_Mounting::OnInstallBundleCompleted);
 			if (UE::GameFeatures::ShouldLogMountedFiles)
 			{
-				FCoreDelegates::OnPakFileMounted2.AddRaw(this, &FGameFeaturePluginState_Mounting::OnPakFileMounted);
+				FCoreDelegates::GetOnPakFileMounted2().AddRaw(this, &FGameFeaturePluginState_Mounting::OnPakFileMounted);
 			}
 		}
 	}
@@ -1857,7 +1857,7 @@ struct FGameFeaturePluginState_Mounting : public FGameFeaturePluginState
 	virtual void EndState() override
 	{
 		IInstallBundleManager::InstallBundleCompleteDelegate.RemoveAll(this);
-		FCoreDelegates::OnPakFileMounted2.RemoveAll(this);
+		FCoreDelegates::GetOnPakFileMounted2().RemoveAll(this);
 	}
 };
 
diff --git a/Engine/Source/Runtime/AssetRegistry/Private/AssetRegistry.cpp b/Engine/Source/Runtime/AssetRegistry/Private/AssetRegistry.cpp
--- a/Engine/Source/Runtime/AssetRegistry/Private/AssetRegistry.cpp
+++ b/Engine/Source/Runtime/AssetRegistry/Private/AssetRegistry.cpp
@@ -332,7 +332,7 @@ private:
 			LoadState = EState::NotFound;
 
 			// The PAK with the main registry isn't mounted yet
-			PakMountedDelegate = FCoreDelegates::OnPakFileMounted2.AddLambda([this](const IPakFile& Pak)
+			PakMountedDelegate = FCoreDelegates::GetOnPakFileMounted2().AddLambda([this](const IPakFile& Pak)
 				{
 					FScopeLock Lock(&StateLock);
 					if (LoadState == EState::NotFound && TrySetPath(Pak))
@@ -343,7 +343,7 @@ private:
 						// DelegateHandle also deallocates our lambda captures
 						FDelegateHandle LocalPakMountedDelegate = PakMountedDelegate;
 						PakMountedDelegate.Reset();
-						FCoreDelegates::OnPakFileMounted2.Remove(LocalPakMountedDelegate);
+						FCoreDelegates::GetOnPakFileMounted2().Remove(LocalPakMountedDelegate);
 					}
 				});
 		}
@@ -418,7 +418,7 @@ private:
 		}
 
 		// Cancel any further searching in Paks since we will no longer accept preloads starting after this point
-		FCoreDelegates::OnPakFileMounted2.Remove(PakMountedDelegate);
+		FCoreDelegates::GetOnPakFileMounted2().Remove(PakMountedDelegate);
 		PakMountedDelegate.Reset();
 
 		if (ConsumeAsynchronous && LoadState == EState::Loading)
@@ -478,7 +478,6 @@ private:
 	 */
 	FAssetRegistryState Payload;
 
-	/** Delegate handle for the callback added to OnPakFileMounted2.After threading starts, Read / Write only within the lock. */
 	FDelegateHandle PakMountedDelegate;
 
 	/** Callback from ConsumeOrDefer that is set so TryLoadAsync can trigger the Consume when it completes.Read / Write only within the lock. */
diff --git a/Engine/Source/Runtime/Core/Private/Internationalization/TextLocalizationManager.cpp b/Engine/Source/Runtime/Core/Private/Internationalization/TextLocalizationManager.cpp
--- a/Engine/Source/Runtime/Core/Private/Internationalization/TextLocalizationManager.cpp
+++ b/Engine/Source/Runtime/Core/Private/Internationalization/TextLocalizationManager.cpp
@@ -351,7 +351,7 @@ void BeginPreInitTextLocalization()
 	DECLARE_SCOPE_CYCLE_COUNTER(TEXT("BeginPreInitTextLocalization"), STAT_BeginPreInitTextLocalization, STATGROUP_LoadTime);
 
 	// Bind this delegate before the PAK file loader is created
-	FCoreDelegates::OnPakFileMounted2.AddRaw(&FTextLocalizationManager::Get(), &FTextLocalizationManager::OnPakFileMounted);
+	FCoreDelegates::GetOnPakFileMounted2().AddRaw(&FTextLocalizationManager::Get(), &FTextLocalizationManager::OnPakFileMounted);
 }
 
 void BeginInitTextLocalization()
diff --git a/Engine/Source/Runtime/Core/Private/Misc/CoreDelegates.cpp b/Engine/Source/Runtime/Core/Private/Misc/CoreDelegates.cpp
--- a/Engine/Source/Runtime/Core/Private/Misc/CoreDelegates.cpp
+++ b/Engine/Source/Runtime/Core/Private/Misc/CoreDelegates.cpp
@@ -32,7 +32,16 @@ TDelegate<IPakFile*(const FString&, int32)> FCoreDelegates::MountPak;
 TDelegate<bool(const FString&)> FCoreDelegates::OnUnmountPak;
 TDelegate<void()> FCoreDelegates::OnOptimizeMemoryUsageForMountedPaks;
 
+PRAGMA_DISABLE_DEPRECATION_WARNINGS
 TMulticastDelegate<void(const IPakFile&)> FCoreDelegates::OnPakFileMounted2;
+PRAGMA_ENABLE_DEPRECATION_WARNINGS
+
+TTSMulticastDelegate<void(const IPakFile&)>& FCoreDelegates::GetOnPakFileMounted2()
+{
+	static TTSMulticastDelegate<void(const IPakFile&)> Singleton;
+	return Singleton;
+}
+
 TMulticastDelegate<void(const FString&)> FCoreDelegates::NewFileAddedDelegate;
 TMulticastDelegate<void()> FCoreDelegates::NoPakFilesMountedDelegate;
 TMulticastDelegate<void(const TCHAR*, const TCHAR*)> FCoreDelegates::OnFileOpenedForReadFromPakFile;
diff --git a/Engine/Source/Runtime/Core/Public/Misc/CoreDelegates.h b/Engine/Source/Runtime/Core/Public/Misc/CoreDelegates.h
--- a/Engine/Source/Runtime/Core/Public/Misc/CoreDelegates.h
+++ b/Engine/Source/Runtime/Core/Public/Misc/CoreDelegates.h
@@ -114,7 +114,9 @@ public:
 	static TDelegate<void()> OnOptimizeMemoryUsageForMountedPaks;
 
 	// After a pakfile is mounted this is called
+	UE_DEPRECATED(5.3, "This delegate is not thread-safe, please use GetOnPakFileMounted2().")
 	static TMulticastDelegate<void(const IPakFile&)> OnPakFileMounted2;
+	static TTSMulticastDelegate<void(const IPakFile&)>& GetOnPakFileMounted2();
 
 	// After a file is added this is called
 	static TMulticastDelegate<void(const FString&)> NewFileAddedDelegate;
diff --git a/Engine/Source/Runtime/CoreUObject/Private/UObject/Obj.cpp b/Engine/Source/Runtime/CoreUObject/Private/UObject/Obj.cpp
--- a/Engine/Source/Runtime/CoreUObject/Private/UObject/Obj.cpp
+++ b/Engine/Source/Runtime/CoreUObject/Private/UObject/Obj.cpp
@@ -4809,7 +4809,7 @@ void InitUObject()
 	FModuleManager::Get().IsPackageLoadedCallback().BindStatic(Local::IsPackageLoaded);
 	
 	FCoreDelegates::NewFileAddedDelegate.AddStatic(FLinkerLoad::OnNewFileAdded);
-	FCoreDelegates::OnPakFileMounted2.AddStatic(FLinkerLoad::OnPakFileMounted);
+	FCoreDelegates::GetOnPakFileMounted2().AddStatic(FLinkerLoad::OnPakFileMounted);
 
 	// Object initialization.
 	StaticUObjectInit();
diff --git a/Engine/Source/Runtime/Engine/Private/Streaming/StreamingManagerTexture.cpp b/Engine/Source/Runtime/Engine/Private/Streaming/StreamingManagerTexture.cpp
--- a/Engine/Source/Runtime/Engine/Private/Streaming/StreamingManagerTexture.cpp
+++ b/Engine/Source/Runtime/Engine/Private/Streaming/StreamingManagerTexture.cpp
@@ -174,7 +174,7 @@ FRenderAssetStreamingManager::FRenderAssetStreamingManager()
 
 	FCoreUObjectDelegates::GetPreGarbageCollectDelegate().AddRaw(this, &FRenderAssetStreamingManager::OnPreGarbageCollect);
 
-	FCoreDelegates::OnPakFileMounted2.AddLambda([this](const IPakFile& PakFile)
+	FCoreDelegates::GetOnPakFileMounted2().AddLambda([this](const IPakFile& PakFile)
 	{
 		FScopeLock ScopeLock(&MountedStateDirtyFilesCS);
 		bRecacheAllFiles = true;
diff --git a/Engine/Source/Runtime/PakFile/Private/IPlatformFilePak.cpp b/Engine/Source/Runtime/PakFile/Private/IPlatformFilePak.cpp
--- a/Engine/Source/Runtime/PakFile/Private/IPlatformFilePak.cpp
+++ b/Engine/Source/Runtime/PakFile/Private/IPlatformFilePak.cpp
@@ -7870,7 +7870,12 @@ bool FPakPlatformFile::Mount(const TCHAR* InPakFilename, uint32 PakOrder, const
 			double OnPakFileMounted2Time = 0.0;
 			{
 				FScopedDurationTimer Timer(OnPakFileMounted2Time);
+
+PRAGMA_DISABLE_DEPRECATION_WARNINGS
 				FCoreDelegates::OnPakFileMounted2.Broadcast(*Pak);
+PRAGMA_ENABLE_DEPRECATION_WARNINGS
+
+				FCoreDelegates::GetOnPakFileMounted2().Broadcast(*Pak);
 			}
 
 			UE_LOG(LogPakFile, Display, TEXT("Mounted Pak file '%s', mount point: '%s'"), InPakFilename, *Pak->GetMountPoint());
diff --git a/Engine/Source/Runtime/RenderCore/Private/ShaderCodeLibrary.cpp b/Engine/Source/Runtime/RenderCore/Private/ShaderCodeLibrary.cpp
--- a/Engine/Source/Runtime/RenderCore/Private/ShaderCodeLibrary.cpp
+++ b/Engine/Source/Runtime/RenderCore/Private/ShaderCodeLibrary.cpp
@@ -3283,7 +3283,7 @@ static void FShaderLibraryPakFileMountedCallback(const IPakFile& PakFile)
 void FShaderCodeLibrary::PreInit()
 {
 	// add a callback for opening later chunks
-	UE::ShaderLibrary::Private::OnPakFileMountedDelegateHandle = FCoreDelegates::OnPakFileMounted2.AddStatic(&FShaderLibraryPakFileMountedCallback);
+	UE::ShaderLibrary::Private::OnPakFileMountedDelegateHandle = FCoreDelegates::GetOnPakFileMounted2().AddStatic(&FShaderLibraryPakFileMountedCallback);
 }
 
 void FShaderCodeLibrary::InitForRuntime(EShaderPlatform ShaderPlatform)
@@ -3354,7 +3354,7 @@ void FShaderCodeLibrary::Shutdown()
 {
 	if (UE::ShaderLibrary::Private::OnPakFileMountedDelegateHandle.IsValid())
 	{
-		FCoreDelegates::OnPakFileMounted2.Remove(UE::ShaderLibrary::Private::OnPakFileMountedDelegateHandle);
+		FCoreDelegates::GetOnPakFileMounted2().Remove(UE::ShaderLibrary::Private::OnPakFileMountedDelegateHandle);
 		UE::ShaderLibrary::Private::OnPakFileMountedDelegateHandle.Reset();
 	}
 	if (UE::ShaderLibrary::Private::OnPluginMountedDelegateHandle.IsValid())
