AssetRegistry: Refuse to scan /Temp files; the AssetRegistry should not contain assets in the temp directory, and UUnrealEdEngine::DeleteActors will break when deleting actors from an autosaved level if the AR reports dependencies from the /Temp autosave to the deleted actor.
#rnx

diff --git a/Engine/Source/Runtime/AssetRegistry/Private/AssetRegistry.cpp b/Engine/Source/Runtime/AssetRegistry/Private/AssetRegistry.cpp
--- a/Engine/Source/Runtime/AssetRegistry/Private/AssetRegistry.cpp
+++ b/Engine/Source/Runtime/AssetRegistry/Private/AssetRegistry.cpp
@@ -4544,6 +4544,11 @@ FScanPathContext::FScanPathContext(FEventContext& InEventContext, FClassInherita
 			UE_LOG(LogAssetRegistry, Warning, TEXT("ScanPathsSynchronous: %s is not in a mounted path, will not scan."), *InFile);
 			continue;
 		}
+		if (FPackageName::IsTempPackage(PackageName))
+		{
+			UE_LOG(LogAssetRegistry, Warning, TEXT("ScanPathsSynchronous: %s is in the /Temp path, will not scan."), *InFile);
+			continue;
+		}
 		if (Extension.IsEmpty())
 		{
 			// The empty extension is not a valid Package extension; it might exist, but we will pay the price to check it
@@ -4575,6 +4580,11 @@ FScanPathContext::FScanPathContext(FEventContext& InEventContext, FClassInherita
 			UE_LOG(LogAssetRegistry, Warning, TEXT("ScanPathsSynchronous: %s is not in a mounted path, will not scan."), *InDir);
 			continue;
 		}
+		if (FPackageName::IsTempPackage(PackageName))
+		{
+			UE_LOG(LogAssetRegistry, Warning, TEXT("ScanPathsSynchronous: %s is in the /Temp path, will not scan."), *InDir);
+			continue;
+		}
 		LocalDirs.Add(LocalPath + Extension);
 		PackageDirs.Add(PackageName + Extension);
 	}
