AssetRegistry: Block completion of the initial search until after AllModuleLoadingPhaseComplete so that all plugins loaded during startup will be loaded. At the point of intial search completion, we cache some results (e.g. UpdateRedirectCollector) and we allow systems that are waiting on completion to be unblocked, so we need to delay until we are really finished, including all necessary plugin scanning.
#rnx

diff --git a/Engine/Source/Runtime/AssetRegistry/Private/AssetRegistry.cpp b/Engine/Source/Runtime/AssetRegistry/Private/AssetRegistry.cpp
--- a/Engine/Source/Runtime/AssetRegistry/Private/AssetRegistry.cpp
+++ b/Engine/Source/Runtime/AssetRegistry/Private/AssetRegistry.cpp
@@ -568,7 +568,7 @@ void FAssetRegistryImpl::ConsumeOrDeferPreloadedPremade(UAssetRegistryImpl& UARI
 	{
 		// if we aren't doing any preloading, then we can set the initial search is done right away.
 		// Otherwise, it is set from LoadPremadeAssetRegistry
-		bCanFinishInitialSearch = true;
+		bPreloadingComplete = true;
 		return;
 	}
 
@@ -752,7 +752,7 @@ void FAssetRegistryImpl::LoadPremadeAssetRegistry(Impl::FEventContext& EventCont
 	}
 
 	// let Tick know that it can finalize the initial search
-	bCanFinishInitialSearch = true;
+	bPreloadingComplete = true;
 }
 
 void FAssetRegistryImpl::Initialize(Impl::FInitializeContext& Context)
@@ -936,6 +936,7 @@ void UAssetRegistryImpl::InitializeEvents(UE::AssetRegistry::Impl::FInitializeCo
 #endif // WITH_EDITOR
 
 	FCoreDelegates::OnEnginePreExit.AddUObject(this, &UAssetRegistryImpl::OnEnginePreExit);
+	FCoreDelegates::OnAllModuleLoadingPhasesComplete.AddUObject(this, &UAssetRegistryImpl::OnAllModuleLoadingPhasesComplete);
 
 	// Listen for new content paths being added or removed at runtime.  These are usually plugin-specific asset paths that
 	// will be loaded a bit later on.
@@ -1399,6 +1400,13 @@ void UAssetRegistryImpl::OnEnginePreExit()
 	GuardedData.OnEnginePreExit();
 }
 
+void UAssetRegistryImpl::OnAllModuleLoadingPhasesComplete()
+{
+	LLM_SCOPE(ELLMTag::AssetRegistry);
+	FWriteScopeLock InterfaceScopeLock(InterfaceLock);
+	GuardedData.OnAllModuleLoadingPhasesComplete();
+}
+
 void UAssetRegistryImpl::FinishDestroy()
 {
 	LLM_SCOPE(ELLMTag::AssetRegistry);
@@ -1410,6 +1418,7 @@ void UAssetRegistryImpl::FinishDestroy()
 		FPackageName::OnContentPathDismounted().RemoveAll(this);
 		FCoreDelegates::OnPostEngineInit.RemoveAll(this);
 		FCoreDelegates::OnEnginePreExit.RemoveAll(this);
+		FCoreDelegates::OnAllModuleLoadingPhasesComplete.RemoveAll(this);
 		IPluginManager::Get().OnLoadingPhaseComplete().RemoveAll(this);
 
 #if WITH_EDITOR
@@ -1492,6 +1501,11 @@ void FAssetRegistryImpl::OnEnginePreExit()
 	GlobalGatherer.Reset();
 }
 
+void FAssetRegistryImpl::OnAllModuleLoadingPhasesComplete()
+{
+	bAllModuleLoadingPhasesComplete = true;
+}
+
 void FAssetRegistryImpl::ConstructGatherer()
 {
 	if (GlobalGatherer.IsValid())
@@ -3981,7 +3995,7 @@ Impl::EGatherStatus FAssetRegistryImpl::TickGatherer(Impl::FEventContext& EventC
 		HighestPending = 0;
 		BackgroundResults.Shrink();
 
-		if (!bInitialSearchCompleted && bCanFinishInitialSearch)
+		if (!bInitialSearchCompleted && bPreloadingComplete && bAllModuleLoadingPhasesComplete)
 		{
 #if WITH_EDITOR
 			// update redirectors
diff --git a/Engine/Source/Runtime/AssetRegistry/Private/AssetRegistry.h b/Engine/Source/Runtime/AssetRegistry/Private/AssetRegistry.h
--- a/Engine/Source/Runtime/AssetRegistry/Private/AssetRegistry.h
+++ b/Engine/Source/Runtime/AssetRegistry/Private/AssetRegistry.h
@@ -185,6 +185,7 @@ protected:
 
 private:
 	void OnEnginePreExit();
+	void OnAllModuleLoadingPhasesComplete();
 #if WITH_EDITOR
 	void OnFEngineLoopInitCompleteSearchAllAssets();
 	/** Called when new gatherer is registered. Requires subsequent call to RebuildAssetDependencyGathererMapIfNeeded */
diff --git a/Engine/Source/Runtime/AssetRegistry/Private/AssetRegistryImpl.h b/Engine/Source/Runtime/AssetRegistry/Private/AssetRegistryImpl.h
--- a/Engine/Source/Runtime/AssetRegistry/Private/AssetRegistryImpl.h
+++ b/Engine/Source/Runtime/AssetRegistry/Private/AssetRegistryImpl.h
@@ -118,6 +118,7 @@ public:
 	/** Construct the AssetRegistryImpl, including initial scans if applicable. */
 	void Initialize(Impl::FInitializeContext& Context);
 	void OnEnginePreExit();
+	void OnAllModuleLoadingPhasesComplete();
 
 	// Helpers for functions of the same name from UAssetRegistryImpl
 
@@ -393,8 +394,17 @@ private:
 	bool bInitialSearchStarted;
 	/** Flag to indicate if the initial background search has completed */
 	bool bInitialSearchCompleted;
-	/** Flag to indicate if the initial background search can be finalized */
-	bool bCanFinishInitialSearch = false;
+	/**
+	 * Flag to indicate PreloadingComplete; finishing the background search is blocked until preloading complete
+	 * because preloading can add assets.
+	 */
+	bool bPreloadingComplete = false;
+	/**
+	 * Flag to indicate LaunchEngineLoop's AllModuleLoadingPhases is complete; finishing the background search is
+	 * blocked until preloading complete because plugins can be mounted during startup up until that point, and
+	 * we need to wait for all the plugins that will load before declaring completion.
+	 */
+	bool bAllModuleLoadingPhasesComplete = false;
 	/** Status of the background search, so we can take actions when it changes to or from idle */
 	Impl::EGatherStatus GatherStatus;
 
