AssetRegistry: Fix GetRedirectedObjectPath to return the redirected path when the initial objectpath is a redirector already in memory.
#rnx

diff --git a/Engine/Source/Runtime/AssetRegistry/Private/AssetRegistry.cpp b/Engine/Source/Runtime/AssetRegistry/Private/AssetRegistry.cpp
--- a/Engine/Source/Runtime/AssetRegistry/Private/AssetRegistry.cpp
+++ b/Engine/Source/Runtime/AssetRegistry/Private/AssetRegistry.cpp
@@ -2685,54 +2685,40 @@ FSoftObjectPath FAssetRegistryImpl::GetRedirectedObjectPath(const FSoftObjectPat
 	// For legacy behavior, for the first object pointed to, we look up the object in memory
 	// before checking the on-disk assets
 	UObject* Asset = ObjectPath.ResolveObject();
-	const FAssetData* AssetData = nullptr;
-	if (!Asset)
+	if (Asset)
 	{
-		AssetData = State.GetAssetByObjectPath(ObjectPath);
+		RedirectedPath = FSoftObjectPath(Asset);
+		UObjectRedirector* Redirector = Cast<UObjectRedirector>(Asset);
+		if (!Redirector || !Redirector->DestinationObject)
+		{
+			return RedirectedPath;
+		}
+		// For legacy behavior, for all redirects after the initial request, we only check on-disk assets
+		RedirectedPath = FSoftObjectPath(Redirector->DestinationObject);
 	}
+	const FAssetData* AssetData = State.GetAssetByObjectPath(RedirectedPath);
 
 	TSet<FSoftObjectPath> SeenPaths;
 	SeenPaths.Add(RedirectedPath);
 
-	auto TryGetRedirectedPath = [](UObject* InAsset, const FAssetData* InAssetData, FSoftObjectPath& OutRedirectedPath)
-	{
-		if (InAsset)
-		{
-			UObjectRedirector* Redirector = Cast<UObjectRedirector>(InAsset);
-			if (Redirector && Redirector->DestinationObject)
-			{
-				OutRedirectedPath = FSoftObjectPath(Redirector->DestinationObject);
-				return true;
-			}
-		}
-		else if (InAssetData && InAssetData->IsRedirector())
-		{
-			FString Dest;
-			if (InAssetData->GetTagValue("DestinationObject", Dest))
-			{
-				ConstructorHelpers::StripObjectClass(Dest);
-				OutRedirectedPath = Dest;
-				return true;
-			}
-		}
-		return false;
-	};
-
 	// Need to follow chain of redirectors
-	while (TryGetRedirectedPath(Asset, AssetData, RedirectedPath))
+	while (AssetData && AssetData->IsRedirector())
 	{
-		if (SeenPaths.Contains(RedirectedPath))
+		FString Dest;
+		if (!AssetData->GetTagValue("DestinationObject", Dest))
 		{
-			// Recursive, bail
 			break;
 		}
-		else
+		ConstructorHelpers::StripObjectClass(Dest);
+		RedirectedPath = Dest;
+		bool bAlreadyExists;
+		SeenPaths.Add(RedirectedPath, &bAlreadyExists);
+		if (bAlreadyExists)
 		{
-			SeenPaths.Add(RedirectedPath);
-			// For legacy behavior, for all redirects after the initial request, we only check on-disk assets
-			Asset = nullptr;
-			AssetData = State.GetAssetByObjectPath(RedirectedPath);
+			// Recursive, bail
+			break;
 		}
+		AssetData = State.GetAssetByObjectPath(RedirectedPath);
 	}
 
 	return RedirectedPath;
