AssetRegistry: Move PackageReader into Internal for use by the cooker.
#rnx

diff --git a/Engine/Source/Runtime/AssetRegistry/Internal/AssetRegistry/PackageReader.h b/Engine/Source/Runtime/AssetRegistry/Internal/AssetRegistry/PackageReader.h
new file mode 100644
--- /dev/null
+++ b/Engine/Source/Runtime/AssetRegistry/Internal/AssetRegistry/PackageReader.h
@@ -0,0 +1,113 @@
+// Copyright Epic Games, Inc. All Rights Reserved.
+
+#pragma once
+
+#include "Misc/EnumClassFlags.h"
+#include "Serialization/ArchiveUObject.h"
+#include "UObject/ObjectResource.h"
+#include "UObject/PackageFileSummary.h"
+
+struct FAssetData;
+class FAssetPackageData;
+class FLinkerTables;
+class FPackageDependencyData;
+
+class FPackageReader : public FArchiveUObject
+{
+public:
+	FPackageReader();
+	~FPackageReader();
+
+	enum class EOpenPackageResult : uint8
+	{
+		/** The package summary loaded successfully */
+		Success,
+		/** The package reader was not given a valid archive to load from */
+		NoLoader,
+		/** The package tag could not be found, the package is probably corrupted */
+		MalformedTag,
+		/** The package is too old to be loaded */
+		VersionTooOld,
+		/** The package is from a newer version of the engine */
+		VersionTooNew,
+		/** The package contains an unknown custom version */
+		CustomVersionMissing,
+		/** The package contains a custom version that failed it's validator */
+		CustomVersionInvalid,
+		/** Package was unversioned but the process cannot load unversioned packages */
+		Unversioned,
+	};
+
+	/** Creates a loader for the filename */
+	bool OpenPackageFile(FStringView PackageFilename, EOpenPackageResult* OutErrorCode = nullptr);
+	bool OpenPackageFile(FStringView LongPackageName, FStringView PackageFilename, EOpenPackageResult* OutErrorCode = nullptr);
+	bool OpenPackageFile(FArchive* Loader, EOpenPackageResult* OutErrorCode = nullptr);
+	bool OpenPackageFile(EOpenPackageResult* OutErrorCode = nullptr);
+
+	/** Returns the LongPackageName from constructor if provided, otherwise calculates it from FPackageName::TryConvertFilenameToLongPackageName. */
+	bool TryGetLongPackageName(FString& OutLongPackageName) const;
+	/** Returns the LongPackageName as in TryGetPackageName; asserts if not found. */
+	FString GetLongPackageName() const;
+
+	/** Reads information from the asset registry data table and converts it to FAssetData */
+	bool ReadAssetRegistryData(TArray<FAssetData*>& AssetDataList, bool& bOutIsCookedWithoutAssetData);
+
+	/** Options for what to read in functions that read multiple things at once. */
+	enum class EReadOptions
+	{
+		None = 0,
+		PackageData = 1 << 0,
+		Dependencies = 1 << 1,
+		Default = PackageData | Dependencies,
+	};
+	/** Reads information used by the dependency graph */
+	bool ReadDependencyData(FPackageDependencyData& OutDependencyData, EReadOptions Options);
+
+	// Farchive implementation to redirect requests to the Loader
+	virtual void Serialize( void* V, int64 Length ) override;
+	virtual bool Precache( int64 PrecacheOffset, int64 PrecacheSize ) override;
+	virtual void Seek( int64 InPos ) override;
+	virtual int64 Tell() override;
+	virtual int64 TotalSize() override;
+	virtual FArchive& operator<<( FName& Name ) override;
+	virtual FString GetArchiveName() const override
+	{
+		return PackageFilename;
+	}
+
+private:
+	/** Attempts to get the class name of an object from the thumbnail cache for packages older than VER_UE4_ASSET_REGISTRY_TAGS */
+	bool ReadAssetDataFromThumbnailCache(TArray<FAssetData*>& AssetDataList);
+	/** Creates asset data reconstructing all the required data from cooked package info */
+	bool ReadAssetRegistryDataFromCookedPackage(TArray<FAssetData*>& AssetDataList, bool& bOutIsCookedWithoutAssetData);
+	bool StartSerializeSection(int64 Offset);
+
+	/** Serializers for different package maps */
+	bool SerializeNameMap();
+	bool SerializeImportMap();
+	bool SerializeExportMap();
+	bool SerializeImportedClasses(const TArray<FObjectImport>& InImportMap, TArray<FName>& OutClassNames);
+	bool SerializeSoftPackageReferenceList(TArray<FName>& OutSoftPackageReferenceList);
+	bool SerializeSearchableNamesMap(FLinkerTables& OutSearchableNames);
+	bool SerializeAssetRegistryDependencyData(TBitArray<>& OutImportUsedInGame, TBitArray<>& OutSoftPackageUsedInGame,
+		const TArray<FObjectImport>& InImportMap, const TArray<FName>& SoftPackageReferenceList);
+	bool SerializePackageTrailer(FAssetPackageData& PackageData);
+
+	void ApplyRelocationToImportMapAndSoftPackageReferenceList(FStringView LoadedPackageName, TArray<FName>& OutSoftPackageReferenceList);
+
+	/** Returns flags the asset package was saved with */
+	uint32 GetPackageFlags() const;
+
+	FString LongPackageName;
+	FString PackageFilename;
+	/* Loader is the interface used to read the bytes from the package's repository. All interpretation of the bytes is done by serializing into *this, which is also an FArchive. */
+	FArchive* Loader;
+	FPackageFileSummary PackageFileSummary;
+	TArray<FName> NameMap;
+	TArray<FObjectImport> ImportMap;
+	TArray<FObjectExport> ExportMap;
+	int64 PackageFileSize;
+	int64 AssetRegistryDependencyDataOffset;
+	bool bLoaderOwner;
+};
+ENUM_CLASS_FLAGS(FPackageReader::EReadOptions);
diff --git a/Engine/Source/Runtime/AssetRegistry/Private/AssetDataGatherer.cpp b/Engine/Source/Runtime/AssetRegistry/Private/AssetDataGatherer.cpp
--- a/Engine/Source/Runtime/AssetRegistry/Private/AssetDataGatherer.cpp
+++ b/Engine/Source/Runtime/AssetRegistry/Private/AssetDataGatherer.cpp
@@ -30,7 +30,6 @@
 #include "Misc/ScopeLock.h"
 #include "Misc/ScopeExit.h"
 #include "Misc/TrackedActivity.h"
-#include "PackageReader.h"
 #include "Serialization/Archive.h"
 #include "Tasks/Task.h"
 #include "Templates/UniquePtr.h"
diff --git a/Engine/Source/Runtime/AssetRegistry/Private/AssetDataGatherer.h b/Engine/Source/Runtime/AssetRegistry/Private/AssetDataGatherer.h
--- a/Engine/Source/Runtime/AssetRegistry/Private/AssetDataGatherer.h
+++ b/Engine/Source/Runtime/AssetRegistry/Private/AssetDataGatherer.h
@@ -18,7 +18,7 @@
 #include "Misc/DateTime.h"
 #include "Misc/Optional.h"
 #include "PackageDependencyData.h"
-#include "PackageReader.h"
+#include "AssetRegistry/PackageReader.h"
 #include "Templates/UniquePtr.h"
 #include "UObject/NameTypes.h"
 
diff --git a/Engine/Source/Runtime/AssetRegistry/Private/AssetRegistry.cpp b/Engine/Source/Runtime/AssetRegistry/Private/AssetRegistry.cpp
--- a/Engine/Source/Runtime/AssetRegistry/Private/AssetRegistry.cpp
+++ b/Engine/Source/Runtime/AssetRegistry/Private/AssetRegistry.cpp
@@ -29,7 +29,7 @@
 #include "Misc/ScopeExit.h"
 #include "Misc/ScopeRWLock.h"
 #include "Misc/TrackedActivity.h"
-#include "PackageReader.h"
+#include "AssetRegistry/PackageReader.h"
 #include "Serialization/ArrayReader.h"
 #include "Serialization/CompactBinarySerialization.h"
 #include "Serialization/CompactBinaryWriter.h"
diff --git a/Engine/Source/Runtime/AssetRegistry/Private/AssetRegistryState.cpp b/Engine/Source/Runtime/AssetRegistry/Private/AssetRegistryState.cpp
--- a/Engine/Source/Runtime/AssetRegistry/Private/AssetRegistryState.cpp
+++ b/Engine/Source/Runtime/AssetRegistry/Private/AssetRegistryState.cpp
@@ -18,7 +18,7 @@
 #include "Misc/PathViews.h"
 #include "Misc/StringBuilder.h"
 #include "NameTableArchive.h"
-#include "PackageReader.h"
+#include "AssetRegistry/PackageReader.h"
 #include "Serialization/ArrayReader.h"
 #include "Serialization/LargeMemoryReader.h"
 #include "Serialization/MemoryWriter.h"
diff --git a/Engine/Source/Runtime/AssetRegistry/Private/PackageReader.cpp b/Engine/Source/Runtime/AssetRegistry/Private/PackageReader.cpp
--- a/Engine/Source/Runtime/AssetRegistry/Private/PackageReader.cpp
+++ b/Engine/Source/Runtime/AssetRegistry/Private/PackageReader.cpp
@@ -1,6 +1,6 @@
 // Copyright Epic Games, Inc. All Rights Reserved.
 
-#include "PackageReader.h"
+#include "AssetRegistry/PackageReader.h"
 
 #include "AssetRegistry.h"
 #include "AssetRegistryPrivate.h"
diff --git a/Engine/Source/Runtime/AssetRegistry/Private/PackageReader.h b/Engine/Source/Runtime/AssetRegistry/Private/PackageReader.h
deleted file mode 100644
