PackageReader - add more accessors and parse functions for use in the cooker and editor domain.
#rnx

diff --git a/Engine/Source/Editor/UnrealEd/Private/EditorDomain/EditorDomainUtils.cpp b/Engine/Source/Editor/UnrealEd/Private/EditorDomain/EditorDomainUtils.cpp
--- a/Engine/Source/Editor/UnrealEd/Private/EditorDomain/EditorDomainUtils.cpp
+++ b/Engine/Source/Editor/UnrealEd/Private/EditorDomain/EditorDomainUtils.cpp
@@ -8,6 +8,7 @@
 #include "Algo/Unique.h"
 #include "AssetRegistry/AssetData.h"
 #include "AssetRegistry/IAssetRegistry.h"
+#include "AssetRegistry/PackageReader.h"
 #include "DerivedDataBuildDefinition.h"
 #include "DerivedDataCache.h"
 #include "DerivedDataCacheKey.h"
@@ -1266,295 +1267,6 @@ void RequestEditorDomainPackage(const FPackagePath& PackagePath,
 	Cache.Get({{{PackagePath.GetDebugName()}, GetEditorDomainPackageKey(EditorDomainHash), CachePolicy}}, Owner, MoveTemp(Callback));
 }
 
-/** TODO: Delete this duplicate of AssetRegistry/PackageReader once we have approved making that class Public */
-class FPackageReader : public FArchiveUObject
-{
-public:
-	FPackageReader()
-		: Loader(nullptr)
-		, PackageFileSize(0)
-	{
-		SetIsLoading(true);
-		SetIsPersistent(true);
-	}
-
-	bool OpenPackageFile(FArchive* InLoader)
-	{
-		Loader = InLoader;
-		return OpenPackageFile();
-	}
-	bool OpenPackageFile()
-	{
-		// Read package file summary from the file
-		*this << PackageFileSummary;
-		PackageFileSize = TotalSize();
-		return true;
-	}
-	bool ReadImportedClasses(TArray<FName>& OutClassNames)
-	{
-		if (!SerializeNameMap())
-		{
-			return false;
-		}
-
-		TArray<FObjectImport> ImportMap;
-		if (!SerializeImportMap(ImportMap))
-		{
-			return false;
-		}
-		if (!SerializeImportedClasses(ImportMap, OutClassNames))
-		{
-			return false;
-		}
-		return true;
-	}
-	const FPackageFileSummary& GetPackageFileSummary() const { return PackageFileSummary; }
-
-
-	// Farchive implementation to redirect requests to the Loader
-	virtual void Serialize(void* V, int64 Length) override
-	{
-		check(Loader);
-		Loader->Serialize(V, Length);
-		if (Loader->IsError())
-		{
-			SetError();
-		}
-	}
-
-	virtual bool Precache(int64 PrecacheOffset, int64 PrecacheSize) override
-	{
-		check(Loader);
-		return Loader->Precache(PrecacheOffset, PrecacheSize);
-	}
-
-	virtual void Seek(int64 InPos) override
-	{
-		check(Loader);
-		Loader->Seek(InPos);
-		if (Loader->IsError())
-		{
-			SetError();
-		}
-	}
-
-	virtual int64 Tell() override
-	{
-		check(Loader);
-		return Loader->Tell();
-	}
-
-	virtual int64 TotalSize() override
-	{
-		check(Loader);
-		return Loader->TotalSize();
-	}
-
-	virtual FArchive& operator<<(FName& Name) override
-	{
-		int32 NameIndex;
-		FArchive& Ar = *this;
-		Ar << NameIndex;
-
-		if (!NameMap.IsValidIndex(NameIndex))
-		{
-			SetError();
-			return *this;
-		}
-
-		// if the name wasn't loaded (because it wasn't valid in this context)
-		if (NameMap[NameIndex] == NAME_None)
-		{
-			int32 TempNumber;
-			Ar << TempNumber;
-			Name = NAME_None;
-		}
-		else
-		{
-			int32 Number;
-			Ar << Number;
-			// simply create the name from the NameMap's name and the serialized instance number
-			Name = FName(NameMap[NameIndex], Number);
-		}
-
-		return *this;
-	}
-
-	virtual FString GetArchiveName() const override
-	{
-		return TEXT("EditorDomainPackageReader");
-	}
-private:
-	// Serializers for different package maps
-	bool SerializeNameMap()
-	{
-		if (NameMap.Num() == 0 && PackageFileSummary.NameCount > 0)
-		{
-			if (!StartSerializeSection(PackageFileSummary.NameOffset))
-			{
-				return false;
-			}
-
-			const int MinSizePerNameEntry = 1;
-			if (PackageFileSize < Tell() + PackageFileSummary.NameCount * MinSizePerNameEntry)
-			{
-				return false;
-			}
-
-			for (int32 NameMapIdx = 0; NameMapIdx < PackageFileSummary.NameCount; ++NameMapIdx)
-			{
-				// Read the name entry from the file.
-				FNameEntrySerialized NameEntry(ENAME_LinkerConstructor);
-				*this << NameEntry;
-				if (IsError())
-				{
-					return false;
-				}
-				NameMap.Add(FName(NameEntry));
-			}
-		}
-
-		return true;
-	}
-	bool SerializeImportMap(TArray<FObjectImport>& OutImportMap)
-	{
-		if (PackageFileSummary.ImportCount > 0)
-		{
-			if (!StartSerializeSection(PackageFileSummary.ImportOffset))
-			{
-				return false;
-			}
-
-			const int MinSizePerImport = 1;
-			if (PackageFileSize < Tell() + PackageFileSummary.ImportCount * MinSizePerImport)
-			{
-				return false;
-			}
-			for (int32 ImportMapIdx = 0; ImportMapIdx < PackageFileSummary.ImportCount; ++ImportMapIdx)
-			{
-				FObjectImport& Import = OutImportMap.Emplace_GetRef();
-				*this << Import;
-				if (IsError())
-				{
-					return false;
-				}
-			}
-		}
-
-		return true;
-	}
-	bool SerializeImportedClasses(const TArray<FObjectImport>& ImportMap, TArray<FName>& OutClassNames)
-	{
-		OutClassNames.Reset();
-
-		TSet<int32> ClassImportIndices;
-		if (PackageFileSummary.ExportCount > 0)
-		{
-			if (!StartSerializeSection(PackageFileSummary.ExportOffset))
-			{
-				return false;
-			}
-
-			const int MinSizePerExport = 1;
-			if (PackageFileSize < Tell() + PackageFileSummary.ExportCount * MinSizePerExport)
-			{
-				return false;
-			}
-			FObjectExport ExportBuffer;
-			for (int32 ExportMapIdx = 0; ExportMapIdx < PackageFileSummary.ExportCount; ++ExportMapIdx)
-			{
-				*this << ExportBuffer;
-				if (IsError())
-				{
-					return false;
-				}
-				if (ExportBuffer.ClassIndex.IsImport())
-				{
-					ClassImportIndices.Add(ExportBuffer.ClassIndex.ToImport());
-				}
-			}
-		}
-
-		TArray<FName, TInlineAllocator<5>>  ParentChain;
-		FNameBuilder ClassObjectPath;
-		for (int32 ClassImportIndex : ClassImportIndices)
-		{
-			ParentChain.Reset();
-			ClassObjectPath.Reset();
-			if (!ImportMap.IsValidIndex(ClassImportIndex))
-			{
-				return false;
-			}
-			bool bParentChainComplete = false;
-			int32 CurrentParentIndex = ClassImportIndex;
-			for (;;)
-			{
-				const FObjectImport& ObjectImport = ImportMap[CurrentParentIndex];
-				ParentChain.Add(ObjectImport.ObjectName);
-				if (ObjectImport.OuterIndex.IsImport())
-				{
-					CurrentParentIndex = ObjectImport.OuterIndex.ToImport();
-					if (!ImportMap.IsValidIndex(CurrentParentIndex))
-					{
-						return false;
-					}
-				}
-				else if (ObjectImport.OuterIndex.IsNull())
-				{
-					bParentChainComplete = true;
-					break;
-				}
-				else
-				{
-					check(ObjectImport.OuterIndex.IsExport());
-					// Ignore classes in an external package but with an object in this package as one of their outers;
-					// We do not need to handle that case yet for Import Classes, and we would have to make this
-					// loop more complex (searching in both ExportMap and ImportMap) to do so
-					break;
-				}
-			}
-
-			if (bParentChainComplete)
-			{
-				int32 NumTokens = ParentChain.Num();
-				check(NumTokens >= 1);
-				const TCHAR Delimiters[] = { '.', SUBOBJECT_DELIMITER_CHAR, '.' };
-				int32 DelimiterIndex = 0;
-				ParentChain[NumTokens - 1].AppendString(ClassObjectPath);
-				for (int32 TokenIndex = NumTokens - 2; TokenIndex >= 0; --TokenIndex)
-				{
-					ClassObjectPath << Delimiters[DelimiterIndex];
-					DelimiterIndex = FMath::Min(DelimiterIndex + 1, static_cast<int32>(UE_ARRAY_COUNT(Delimiters)) - 1);
-					ParentChain[TokenIndex].AppendString(ClassObjectPath);
-				}
-				OutClassNames.Emplace(ClassObjectPath);
-			}
-		}
-
-		OutClassNames.Sort(FNameLexicalLess());
-		return true;
-	}
-	bool StartSerializeSection(int64 Offset)
-	{
-		check(Loader);
-		if (Offset <= 0 || Offset > PackageFileSize)
-		{
-			return false;
-		}
-		ClearError();
-		Loader->ClearError();
-		Seek(Offset);
-		return !IsError();
-	}
-
-	FArchive* Loader;
-	FPackageFileSummary PackageFileSummary;
-	TArray<FName> NameMap;
-	int64 PackageFileSize;
-	int64 AssetRegistryDependencyDataOffset;
-	bool bLoaderOwner;
-};
-
 /** Stores data from SavePackage in accessible fields */
 class FEditorDomainPackageWriter final : public TPackageWriterToSharedBuffer<IPackageWriter>
 {
diff --git a/Engine/Source/Runtime/AssetRegistry/Internal/AssetRegistry/PackageReader.h b/Engine/Source/Runtime/AssetRegistry/Internal/AssetRegistry/PackageReader.h
--- a/Engine/Source/Runtime/AssetRegistry/Internal/AssetRegistry/PackageReader.h
+++ b/Engine/Source/Runtime/AssetRegistry/Internal/AssetRegistry/PackageReader.h
@@ -6,6 +6,7 @@
 #include "Serialization/ArchiveUObject.h"
 #include "UObject/ObjectResource.h"
 #include "UObject/PackageFileSummary.h"
+#include "UObject/SoftObjectPath.h"
 
 struct FAssetData;
 class FAssetPackageData;
@@ -15,8 +16,8 @@ class FPackageDependencyData;
 class FPackageReader : public FArchiveUObject
 {
 public:
-	FPackageReader();
-	~FPackageReader();
+	ASSETREGISTRY_API FPackageReader();
+	ASSETREGISTRY_API ~FPackageReader();
 
 	enum class EOpenPackageResult : uint8
 	{
@@ -38,19 +39,54 @@ public:
 		Unversioned,
 	};
 
+	/** Information provided by PackageReader about a parsed Import or Export. */
+	struct FObjectData
+	{
+		FSoftObjectPath ClassPath;
+		bool bUsedInGame = true;
+	};
+
 	/** Creates a loader for the filename */
-	bool OpenPackageFile(FStringView PackageFilename, EOpenPackageResult* OutErrorCode = nullptr);
-	bool OpenPackageFile(FStringView LongPackageName, FStringView PackageFilename, EOpenPackageResult* OutErrorCode = nullptr);
-	bool OpenPackageFile(FArchive* Loader, EOpenPackageResult* OutErrorCode = nullptr);
-	bool OpenPackageFile(EOpenPackageResult* OutErrorCode = nullptr);
+	ASSETREGISTRY_API bool OpenPackageFile(FStringView PackageFilename, EOpenPackageResult* OutErrorCode = nullptr);
+	ASSETREGISTRY_API bool OpenPackageFile(FStringView LongPackageName, FStringView PackageFilename,
+		EOpenPackageResult* OutErrorCode = nullptr);
+	ASSETREGISTRY_API bool OpenPackageFile(FArchive* Loader, EOpenPackageResult* OutErrorCode = nullptr);
+	ASSETREGISTRY_API bool OpenPackageFile(EOpenPackageResult* OutErrorCode = nullptr);
 
-	/** Returns the LongPackageName from constructor if provided, otherwise calculates it from FPackageName::TryConvertFilenameToLongPackageName. */
-	bool TryGetLongPackageName(FString& OutLongPackageName) const;
+	/**
+	 * Returns the LongPackageName from constructor if provided, otherwise calculates it from
+	 * FPackageName::TryConvertFilenameToLongPackageName.
+	 */
+	ASSETREGISTRY_API bool TryGetLongPackageName(FString& OutLongPackageName) const;
 	/** Returns the LongPackageName as in TryGetPackageName; asserts if not found. */
-	FString GetLongPackageName() const;
+	ASSETREGISTRY_API FString GetLongPackageName() const;
+
+	/** Returns the Summary that was ready by OpenPackageFile. */
+	ASSETREGISTRY_API const FPackageFileSummary& GetPackageFileSummary() const;
+	/** Reads Imports if not already read and returns a copy. */
+	ASSETREGISTRY_API bool GetImports(TArray<FObjectImport>& OutImportMap);
+	/** Reads Exports if not already read and returns a copy. */
+	ASSETREGISTRY_API bool GetExports(TArray<FObjectExport>& OutExportMap);
+	/** Reads the list of SoftPackageDependencies if not already read and returns a copy. */
+	ASSETREGISTRY_API bool GetSoftPackageReferenceList(TArray<FName>& OutSoftPackageReferenceList);
+	/** Reads the EditorOnly flags for Imports and Exports and returns them. not cached, reparsed each call. */
+	ASSETREGISTRY_API bool ReadEditorOnlyFlags(TBitArray<>& OutImportUsedInGame, TBitArray<>& OutSoftPackageUsedInGame);
+
+	/** Reads information from the AR section and converts it to FAssetData. Not cached, reparsed each call. */
+	ASSETREGISTRY_API bool ReadAssetRegistryData(TArray<FAssetData*>& AssetDataList, bool& bOutIsCookedWithoutAssetData);
+
+	/** Reads Imports if not already read and parses the names of imported classes out of the Imports. */
+	ASSETREGISTRY_API bool ReadImportedClasses(TArray<FName>& OutClassNames);
 
-	/** Reads information from the asset registry data table and converts it to FAssetData */
-	bool ReadAssetRegistryData(TArray<FAssetData*>& AssetDataList, bool& bOutIsCookedWithoutAssetData);
+	/**
+	 * Reads LinkerTables if not already read, reads the Editor-Only flags (not cached, reparsed each call), 
+	 * And constructs a Map from ObjectPath -> ObjectData for each import, export, and SoftPackageReference.
+	 * @param OutExports[SoftObjectPath].Value == FObjectData
+	 * @param OutImports[SoftObjectPath].Value == FObjectData
+	 * @param OutSoftPackageReferences[PackageName].Value == bUsedInGame
+	 */
+	ASSETREGISTRY_API bool ReadLinkerObjects(TMap<FSoftObjectPath, FObjectData>& OutExports,
+		TMap<FSoftObjectPath, FObjectData>& OutImports, TMap<FName, bool>& OutSoftPackageReferences);
 
 	/** Options for what to read in functions that read multiple things at once. */
 	enum class EReadOptions
@@ -60,23 +96,23 @@ public:
 		Dependencies = 1 << 1,
 		Default = PackageData | Dependencies,
 	};
-	/** Reads information used by the dependency graph */
-	bool ReadDependencyData(FPackageDependencyData& OutDependencyData, EReadOptions Options);
+	/** Reads information used by the dependency graph. Not cached, reparsed each call. */
+	ASSETREGISTRY_API bool ReadDependencyData(FPackageDependencyData& OutDependencyData, EReadOptions Options);
 
 	// Farchive implementation to redirect requests to the Loader
-	virtual void Serialize( void* V, int64 Length ) override;
-	virtual bool Precache( int64 PrecacheOffset, int64 PrecacheSize ) override;
-	virtual void Seek( int64 InPos ) override;
-	virtual int64 Tell() override;
-	virtual int64 TotalSize() override;
-	virtual FArchive& operator<<( FName& Name ) override;
+	ASSETREGISTRY_API virtual void Serialize( void* V, int64 Length ) override;
+	ASSETREGISTRY_API virtual bool Precache( int64 PrecacheOffset, int64 PrecacheSize ) override;
+	ASSETREGISTRY_API virtual void Seek( int64 InPos ) override;
+	ASSETREGISTRY_API virtual int64 Tell() override;
+	ASSETREGISTRY_API virtual int64 TotalSize() override;
+	ASSETREGISTRY_API virtual FArchive& operator<<( FName& Name ) override;
 	virtual FString GetArchiveName() const override
 	{
 		return PackageFilename;
 	}
 
 private:
-	/** Attempts to get the class name of an object from the thumbnail cache for packages older than VER_UE4_ASSET_REGISTRY_TAGS */
+	/** Attempts to get object class name from the thumbnail cache for packages older than VER_UE4_ASSET_REGISTRY_TAGS */
 	bool ReadAssetDataFromThumbnailCache(TArray<FAssetData*>& AssetDataList);
 	/** Creates asset data reconstructing all the required data from cooked package info */
 	bool ReadAssetRegistryDataFromCookedPackage(TArray<FAssetData*>& AssetDataList, bool& bOutIsCookedWithoutAssetData);
@@ -87,25 +123,33 @@ private:
 	bool SerializeImportMap();
 	bool SerializeExportMap();
 	bool SerializeImportedClasses(const TArray<FObjectImport>& InImportMap, TArray<FName>& OutClassNames);
-	bool SerializeSoftPackageReferenceList(TArray<FName>& OutSoftPackageReferenceList);
+	bool SerializeSoftPackageReferenceList();
+	bool SerializeEditorOnlyFlags(TBitArray<>& OutImportUsedInGame, TBitArray<>& OutSoftPackageUsedInGame);
 	bool SerializeSearchableNamesMap(FLinkerTables& OutSearchableNames);
 	bool SerializeAssetRegistryDependencyData(TBitArray<>& OutImportUsedInGame, TBitArray<>& OutSoftPackageUsedInGame,
-		const TArray<FObjectImport>& InImportMap, const TArray<FName>& SoftPackageReferenceList);
+		const TArray<FObjectImport>& InImportMap, const TArray<FName>& InSoftPackageReferenceList);
 	bool SerializePackageTrailer(FAssetPackageData& PackageData);
 
-	void ApplyRelocationToImportMapAndSoftPackageReferenceList(FStringView LoadedPackageName, TArray<FName>& OutSoftPackageReferenceList);
+	void ApplyRelocationToImportMapAndSoftPackageReferenceList(FStringView LoadedPackageName,
+		TArray<FName>& OutSoftPackageReferenceList);
+	static void ConvertLinkerTableToPaths(FName PackageName, TArray<FObjectExport>& ExportMap,
+		TArray<FObjectImport>& ImportMap, TArray<FSoftObjectPath>& OutExports, TArray<FSoftObjectPath>& OutImports);
 
 	/** Returns flags the asset package was saved with */
 	uint32 GetPackageFlags() const;
 
 	FString LongPackageName;
 	FString PackageFilename;
-	/* Loader is the interface used to read the bytes from the package's repository. All interpretation of the bytes is done by serializing into *this, which is also an FArchive. */
+	/*
+	 * Loader is the interface used to read the bytes from the package's repository. All interpretation of the bytes is
+	 * done by serializing into *this, which is also an FArchive.
+	 */
 	FArchive* Loader;
 	FPackageFileSummary PackageFileSummary;
 	TArray<FName> NameMap;
 	TArray<FObjectImport> ImportMap;
 	TArray<FObjectExport> ExportMap;
+	TArray<FName> SoftPackageReferenceList;
 	int64 PackageFileSize;
 	int64 AssetRegistryDependencyDataOffset;
 	bool bLoaderOwner;
diff --git a/Engine/Source/Runtime/AssetRegistry/Private/PackageReader.cpp b/Engine/Source/Runtime/AssetRegistry/Private/PackageReader.cpp
--- a/Engine/Source/Runtime/AssetRegistry/Private/PackageReader.cpp
+++ b/Engine/Source/Runtime/AssetRegistry/Private/PackageReader.cpp
@@ -11,6 +11,7 @@
 #include "Misc/Guid.h"
 #include "Misc/PackageName.h"
 #include "Misc/PackagePath.h"
+#include "Misc/ScopeExit.h"
 #include "UObject/Class.h"
 #include "UObject/Linker.h"
 #include "UObject/PackageRelocation.h"
@@ -212,23 +213,77 @@ bool FPackageReader::StartSerializeSection(int64 Offset)
 			CorruptPackageWarningArguments).ToString()); \
 	} while (false)
 
-bool FPackageReader::ReadAssetRegistryData(TArray<FAssetData*>& AssetDataList, bool& bOutIsCookedWithoutAssetData)
+
+const FPackageFileSummary& FPackageReader::GetPackageFileSummary() const
 {
-	bOutIsCookedWithoutAssetData = false;
-	if (!!(GetPackageFlags() & PKG_FilterEditorOnly))
+	return PackageFileSummary;
+}
+
+bool FPackageReader::GetImports(TArray<FObjectImport>& OutImportMap)
+{
+	if (!SerializeNameMap() || !SerializeImportMap())
 	{
-		return ReadAssetRegistryDataFromCookedPackage(AssetDataList, bOutIsCookedWithoutAssetData);
+		return false;
 	}
+	OutImportMap = ImportMap;
+	return true;
+}
 
-	if (!SerializeNameMap())
+bool FPackageReader::GetExports(TArray<FObjectExport>& OutExportMap)
+{
+	if (!SerializeNameMap() || !SerializeExportMap())
 	{
 		return false;
 	}
-	if (!SerializeImportMap())
+	OutExportMap = ExportMap;
+	return true;
+}
+
+bool FPackageReader::GetSoftPackageReferenceList(TArray<FName>& OutSoftPackageReferenceList)
+{
+	if (!SerializeNameMap() || !SerializeSoftPackageReferenceList())
 	{
 		return false;
 	}
-	if (!SerializeExportMap())
+	OutSoftPackageReferenceList = SoftPackageReferenceList;
+	return true;
+}
+
+bool FPackageReader::ReadEditorOnlyFlags(TBitArray<>& OutImportUsedInGame, TBitArray<>& OutSoftPackageUsedInGame)
+{
+	if (!SerializeNameMap() || !SerializeImportMap() || !SerializeSoftPackageReferenceList())
+	{
+		return false;
+	}
+	if (!SerializeEditorOnlyFlags(OutImportUsedInGame, OutSoftPackageUsedInGame))
+	{
+		return false;
+	}
+	return true;
+}
+
+bool FPackageReader::ReadImportedClasses(TArray<FName>& OutClassNames)
+{
+	if (!SerializeNameMap() || !SerializeImportMap())
+	{
+		return false;
+	}
+	if (!SerializeImportedClasses(ImportMap, OutClassNames))
+	{
+		return false;
+	}
+	return true;
+}
+
+bool FPackageReader::ReadAssetRegistryData(TArray<FAssetData*>& AssetDataList, bool& bOutIsCookedWithoutAssetData)
+{
+	bOutIsCookedWithoutAssetData = false;
+	if (!!(GetPackageFlags() & PKG_FilterEditorOnly))
+	{
+		return ReadAssetRegistryDataFromCookedPackage(AssetDataList, bOutIsCookedWithoutAssetData);
+	}
+
+	if (!SerializeNameMap() || !SerializeImportMap() || !SerializeExportMap())
 	{
 		return false;
 	}
@@ -278,14 +333,83 @@ bool FPackageReader::ReadAssetRegistryData(TArray<FAssetData*>& AssetDataList, b
 	return true;
 }
 
+bool FPackageReader::ReadLinkerObjects(TMap<FSoftObjectPath, FObjectData>& OutExports,
+	TMap<FSoftObjectPath, FObjectData>& OutImports, TMap<FName, bool>& OutSoftPackageReferences)
+{
+	if (!SerializeNameMap() || !SerializeImportMap() || !SerializeExportMap() || !SerializeSoftPackageReferenceList())
+	{
+		return false;
+	}
+	TBitArray<> ImportUsedInGame;
+	TBitArray<> SoftPackageUsedInGame;
+	if (!SerializeEditorOnlyFlags(ImportUsedInGame, SoftPackageUsedInGame))
+	{
+		return false;
+	}
+
+	TArray<FSoftObjectPath> ExportsPaths;
+	TArray<FSoftObjectPath> ImportsPaths;
+	int32 NumExports = ExportMap.Num();
+	int32 NumImports = ImportMap.Num();
+	int32 NumSoftPackageReferences = SoftPackageReferenceList.Num();
+	ConvertLinkerTableToPaths(FName(*GetLongPackageName()), ExportMap, ImportMap, ExportsPaths, ImportsPaths);
+	auto GetClassPath = [&ExportsPaths, &ImportsPaths](FPackageIndex ClassIndex)
+	{
+		if (ClassIndex.IsNull())
+		{
+			return FSoftObjectPath(TEXT("/Script/CoreUObject.Class"));
+		}
+		else if (ClassIndex.IsExport())
+		{
+			int32 ExportIndex = ClassIndex.ToExport();
+			return ExportIndex < ExportsPaths.Num() ? ExportsPaths[ExportIndex] :
+				FSoftObjectPath(TEXT("/Script/Unknown.Unknown"));
+		}
+		else
+		{
+			int32 ImportIndex = ClassIndex.ToImport();
+			return ImportIndex < ImportsPaths.Num() ? ImportsPaths[ImportIndex] :
+				FSoftObjectPath(TEXT("/Script/Unknown.Unknown"));
+		}
+	};
+	OutExports.Reserve(NumExports);
+	for (int32 Index = 0; Index < NumExports; ++Index)
+	{
+		FObjectData Data;
+		Data.ClassPath = GetClassPath(ExportMap[Index].ClassIndex);
+		Data.bUsedInGame = true;
+		OutExports.Add(ExportsPaths[Index], MoveTemp(Data));
+	}
+	OutImports.Reserve(NumImports);
+	for (int32 Index = 0; Index < NumImports; ++Index)
+	{
+		FObjectData Data;
+		const FObjectImport& Import = ImportMap[Index];
+		Data.ClassPath = FSoftObjectPath(Import.ClassPackage, Import.ClassName, FString());
+		Data.bUsedInGame = ImportUsedInGame[Index];
+		OutImports.Add(ImportsPaths[Index], MoveTemp(Data));
+	}
+	OutSoftPackageReferences.Reserve(NumSoftPackageReferences);
+	for (int32 Index = 0; Index < NumSoftPackageReferences; ++Index)
+	{
+		OutSoftPackageReferences.Add(SoftPackageReferenceList[Index], SoftPackageUsedInGame[Index]);
+	}
+
+	auto SoftObjectPathLexicalLess = [](const FSoftObjectPath& A, const FSoftObjectPath& B) { return A.LexicalLess(B); };
+	OutExports.KeySort(SoftObjectPathLexicalLess);
+	OutImports.KeySort(SoftObjectPathLexicalLess);
+	OutSoftPackageReferences.KeySort(FNameLexicalLess());
+	return true;
+}
+
 bool FPackageReader::SerializeAssetRegistryDependencyData(TBitArray<>& OutImportUsedInGame, TBitArray<>& OutSoftPackageUsedInGame,
-	const TArray<FObjectImport>& InImportMap, const TArray<FName>& SoftPackageReferenceList)
+	const TArray<FObjectImport>& InImportMap, const TArray<FName>& InSoftPackageReferenceList)
 {
 	if (AssetRegistryDependencyDataOffset == INDEX_NONE)
 	{
 		// For old package versions that did not write out the dependency flags, set default values of the flags
 		OutImportUsedInGame.Init(true, InImportMap.Num());
-		OutSoftPackageUsedInGame.Init(true, SoftPackageReferenceList.Num());
+		OutSoftPackageUsedInGame.Init(true, InSoftPackageReferenceList.Num());
 		return true;
 	}
 
@@ -296,7 +420,7 @@ bool FPackageReader::SerializeAssetRegistryDependencyData(TBitArray<>& OutImport
 
 	if (!UE::AssetRegistry::ReadPackageDataDependencies(*this, OutImportUsedInGame, OutSoftPackageUsedInGame) ||
 		OutImportUsedInGame.Num() != InImportMap.Num() ||
-		OutSoftPackageUsedInGame.Num() != SoftPackageReferenceList.Num())
+		OutSoftPackageUsedInGame.Num() != InSoftPackageReferenceList.Num())
 	{
 		UE_PACKAGEREADER_CORRUPTPACKAGE_WARNING("SerializeAssetRegistryDependencyData", PackageFilename);
 		return false;
@@ -416,15 +540,7 @@ bool FPackageReader::ReadAssetRegistryDataFromCookedPackage(TArray<FAssetData*>&
 	{
 		const FString PackagePath = FPackageName::GetLongPackagePath(PackageName);
 
-		if (!SerializeNameMap())
-		{
-			return false;
-		}
-		if (!SerializeImportMap())
-		{
-			return false;
-		}
-		if (!SerializeExportMap())
+		if (!SerializeNameMap() || !SerializeImportMap() || !SerializeExportMap())
 		{
 			return false;
 		}
@@ -478,12 +594,7 @@ bool FPackageReader::ReadDependencyData(FPackageDependencyData& OutDependencyDat
 		return true;
 	}
 
-	if (!SerializeNameMap())
-	{
-		return false;
-	}
-
-	if (!SerializeImportMap())
+	if (!SerializeNameMap() || !SerializeImportMap())
 	{
 		return false;
 	}
@@ -515,8 +626,7 @@ bool FPackageReader::ReadDependencyData(FPackageDependencyData& OutDependencyDat
 	if (EnumHasAnyFlags(Options, EReadOptions::Dependencies))
 	{
 		OutDependencyData.bHasDependencyData = true;
-		TArray<FName> SoftPackageReferenceList;
-		if (!SerializeSoftPackageReferenceList(SoftPackageReferenceList))
+		if (!SerializeSoftPackageReferenceList())
 		{
 			return false;
 		}
@@ -773,8 +883,13 @@ bool FPackageReader::SerializeExportMap()
 	return true;
 }
 
-bool FPackageReader::SerializeSoftPackageReferenceList(TArray<FName>& OutSoftPackageReferenceList)
+bool FPackageReader::SerializeSoftPackageReferenceList()
 {
+	if (SoftPackageReferenceList.Num() > 0)
+	{
+		return true;
+	}
+
 	if (UEVer() >= VER_UE4_ADD_STRING_ASSET_REFERENCES_MAP && PackageFileSummary.SoftPackageReferencesOffset > 0 && PackageFileSummary.SoftPackageReferencesCount > 0)
 	{
 		if (!StartSerializeSection(PackageFileSummary.SoftPackageReferencesOffset))
@@ -789,6 +904,8 @@ bool FPackageReader::SerializeSoftPackageReferenceList(TArray<FName>& OutSoftPac
 			UE_PACKAGEREADER_CORRUPTPACKAGE_WARNING("SerializeSoftPackageReferenceListInvalidReferencesCount", PackageFilename);
 			return false;
 		}
+
+		SoftPackageReferenceList.Reserve(PackageFileSummary.SoftPackageReferencesCount);
 		if (UEVer() < VER_UE4_ADDED_SOFT_OBJECT_PATH)
 		{
 			for (int32 ReferenceIdx = 0; ReferenceIdx < PackageFileSummary.SoftPackageReferencesCount; ++ReferenceIdx)
@@ -798,6 +915,7 @@ bool FPackageReader::SerializeSoftPackageReferenceList(TArray<FName>& OutSoftPac
 				if (IsError())
 				{
 					UE_PACKAGEREADER_CORRUPTPACKAGE_WARNING("SerializeSoftPackageReferenceListInvalidReferencePreSoftObjectPath", PackageFilename);
+					SoftPackageReferenceList.Reset();
 					return false;
 				}
 
@@ -810,7 +928,7 @@ bool FPackageReader::SerializeSoftPackageReferenceList(TArray<FName>& OutSoftPac
 					}
 				}
 
-				OutSoftPackageReferenceList.Add(FName(*PackageName));
+				SoftPackageReferenceList.Add(FName(*PackageName));
 			}
 		}
 		else
@@ -822,10 +940,11 @@ bool FPackageReader::SerializeSoftPackageReferenceList(TArray<FName>& OutSoftPac
 				if (IsError())
 				{
 					UE_PACKAGEREADER_CORRUPTPACKAGE_WARNING("SerializeSoftPackageReferenceListInvalidReference", PackageFilename);
+					SoftPackageReferenceList.Reset();
 					return false;
 				}
 
-				OutSoftPackageReferenceList.Add(PackageName);
+				SoftPackageReferenceList.Add(PackageName);
 			}
 		}
 	}
@@ -833,6 +952,47 @@ bool FPackageReader::SerializeSoftPackageReferenceList(TArray<FName>& OutSoftPac
 	return true;
 }
 
+bool FPackageReader::SerializeEditorOnlyFlags(TBitArray<>& OutImportUsedInGame, TBitArray<>& OutSoftPackageUsedInGame)
+{
+	if (AssetRegistryDependencyDataOffset == INDEX_NONE &&
+		!(PackageFileSummary.GetPackageFlags() & PKG_FilterEditorOnly) &&
+		StartSerializeSection(PackageFileSummary.AssetRegistryDataOffset))
+	{
+		UE::AssetRegistry::FDeserializePackageData DeserializePackageData;
+		UE::AssetRegistry::EReadPackageDataMainErrorCode Error;
+		if (!DeserializePackageData.DoSerialize(*this, PackageFileSummary, Error))
+		{
+			UE_PACKAGEREADER_CORRUPTPACKAGE_WARNING("EReadPackageDataMainErrorCode::InvalidTagCount", PackageFilename);
+			return false;
+		}
+		AssetRegistryDependencyDataOffset = DeserializePackageData.DependencyDataOffset;
+	}
+
+	if (AssetRegistryDependencyDataOffset == INDEX_NONE)
+	{
+		// For cooked packages or old package versions that did not write out the dependency flags,
+		// set default values of the flags
+		OutImportUsedInGame.Init(true, ImportMap.Num());
+		OutSoftPackageUsedInGame.Init(true, SoftPackageReferenceList.Num());
+		return true;
+	}
+
+	if (!StartSerializeSection(AssetRegistryDependencyDataOffset))
+	{
+		return false;
+	}
+
+	if (!UE::AssetRegistry::ReadPackageDataDependencies(*this, OutImportUsedInGame, OutSoftPackageUsedInGame) ||
+		OutImportUsedInGame.Num() != ImportMap.Num() ||
+		OutSoftPackageUsedInGame.Num() != SoftPackageReferenceList.Num())
+	{
+		UE_PACKAGEREADER_CORRUPTPACKAGE_WARNING("SerializeAssetRegistryDependencyData", PackageFilename);
+		return false;
+	}
+
+	return true;
+}
+
 bool FPackageReader::SerializeSearchableNamesMap(FLinkerTables& OutSearchableNames)
 {
 	if (UEVer() >= VER_UE4_ADDED_SEARCHABLE_NAMES && PackageFileSummary.SearchableNamesOffset > 0)
@@ -928,6 +1088,111 @@ FArchive& FPackageReader::operator<<( FName& Name )
 	return *this;
 }
 
+void FPackageReader::ConvertLinkerTableToPaths(FName InPackageName, TArray<FObjectExport>& InExportMap, TArray<FObjectImport>& InImportMap,
+	TArray<FSoftObjectPath>& OutExports, TArray<FSoftObjectPath>& OutImports)
+{
+	TMap<FPackageIndex, FSoftObjectPath> ObjectForIndex;
+	TFunction<FSoftObjectPath& (FPackageIndex Index)> GetSoftObjectPath;
+	FSoftObjectPath EmptySoftObjectPath;
+	int32 Counter = 0;
+	int32 NumExports = InExportMap.Num();
+	int32 NumImports = InImportMap.Num();
+	int32 NumObjects = NumExports + NumImports;
+
+	ObjectForIndex.Reserve(NumObjects);
+	GetSoftObjectPath = [&InImportMap, &InExportMap, &ObjectForIndex, &GetSoftObjectPath, &EmptySoftObjectPath,
+		&Counter, NumObjects, InPackageName]
+	(FPackageIndex Index) -> FSoftObjectPath&
+	{
+		if (++Counter > (NumObjects + 1) * 2)
+		{
+			// Recursive overflow should be impossible because every call fills in a new element of the table
+			check(false);
+		}
+		ON_SCOPE_EXIT
+		{
+			--Counter;
+		};
+		if (Index.IsNull())
+		{
+			return EmptySoftObjectPath;
+		}
+		FSoftObjectPath* Existing = ObjectForIndex.Find(Index);
+		if (Existing)
+		{
+			return *Existing;
+		}
+
+		FPackageIndex ParentIndex;
+		FName ObjectName;
+		if (Index.IsExport())
+		{
+			int32 ExportIndex = Index.ToExport();
+			if (ExportIndex < InExportMap.Num())
+			{
+				FObjectExport& Export = InExportMap[ExportIndex];
+				ParentIndex = Export.OuterIndex;
+				ObjectName = Export.ObjectName;
+			}
+		}
+		else
+		{
+			int32 ImportIndex = Index.ToImport();
+			if (ImportIndex < InImportMap.Num())
+			{
+				FObjectImport& Import = InImportMap[ImportIndex];
+				ParentIndex = Import.OuterIndex;
+				ObjectName = Import.ObjectName;
+			}
+		}
+		FSoftObjectPath Result;
+		if (!ObjectName.IsNone())
+		{
+			FSoftObjectPath& ParentPath = GetSoftObjectPath(ParentIndex);
+			if (ParentPath.IsNull())
+			{
+				if (Index.IsExport())
+				{
+					Result = FSoftObjectPath(InPackageName, ObjectName, FString());
+				}
+				else
+				{
+					Result = FSoftObjectPath(ObjectName, NAME_None, FString());
+				}
+			}
+			else if (ParentPath.GetAssetFName().IsNone())
+			{
+				Result = FSoftObjectPath(ParentPath.GetLongPackageFName(), ObjectName, FString());
+			}
+			else if (ParentPath.GetSubPathString().IsEmpty())
+			{
+				Result = FSoftObjectPath(ParentPath.GetLongPackageFName(), ParentPath.GetAssetFName(),
+					ObjectName.ToString());
+			}
+			else
+			{
+				Result = FSoftObjectPath(ParentPath.GetLongPackageFName(), ParentPath.GetAssetFName(),
+					ParentPath.GetSubPathString() + TEXT(".") + ObjectName.ToString());
+			}
+		}
+
+		// Note we have to look up the Element again in ObjectForIndex. We can not cache a FindOrAdd Result for Index
+		// because we have potentially modified ObjectForIndex by calling GetSoftObjectPath(ParentIndex).
+		return ObjectForIndex.Add(Index, Result);
+	};
+	OutExports.Reset(NumExports);
+	for (int32 ExportIndex = 0; ExportIndex < NumExports; ++ExportIndex)
+	{
+		OutExports.Add(GetSoftObjectPath(FPackageIndex::FromExport(ExportIndex)));
+	}
+	OutImports.Reset(NumImports);
+	for (int32 ImportIndex = 0; ImportIndex < NumImports; ++ImportIndex)
+	{
+		OutImports.Add(GetSoftObjectPath(FPackageIndex::FromImport(ImportIndex)));
+	}
+}
+
+
 namespace UE::AssetRegistry
 {
 	class FNameMapAwareArchive : public FArchiveProxy
