Ensure AssetRegistry doesn't treat initial discovery as complete unless all modules have finished loading.  Add logging to indicate size and time of monolithic asset registry cache reads/writes.

diff --git a/Engine/Source/Runtime/AssetRegistry/Private/AssetDataGatherer.cpp b/Engine/Source/Runtime/AssetRegistry/Private/AssetDataGatherer.cpp
--- a/Engine/Source/Runtime/AssetRegistry/Private/AssetDataGatherer.cpp
+++ b/Engine/Source/Runtime/AssetRegistry/Private/AssetDataGatherer.cpp
@@ -88,7 +88,7 @@ struct FCachePayload
 
 void SerializeCacheSave(FAssetRegistryWriter& Ar, const TArray<TPair<FName, FDiskCachedAssetData*>>& AssetsToSave);
 FCachePayload SerializeCacheLoad(FAssetRegistryReader& Ar);
-FCachePayload LoadCacheFile(FStringView CacheFilename);
+FCachePayload LoadCacheFile(FStringView CacheFilename, bool bIsMonolithicCache);
 
 
 /** InOutResult = Value, but without shrinking the string to fit. */
@@ -3118,7 +3118,7 @@ private:
 	void LoadAsync()
 	{
 		LLM_SCOPE(ELLMTag::AssetRegistry);
-		Payload = LoadCacheFile(GPreloadSettings.GetMonolithicCacheFilename());
+		Payload = LoadCacheFile(GPreloadSettings.GetMonolithicCacheFilename(), true /* bIsMonolithicCache */);
 	}
 
 	TFuture<void> PreloadReady;
@@ -3146,6 +3146,7 @@ FAssetDataGatherer::FAssetDataGatherer(const TArray<FString>& InLongPackageNames
 	, bIsIdle(false)
 	, bFirstTickAfterIdle(true)
 	, bFinishedInitialDiscovery(false)
+	, bAllModuleLoadingPhasesComplete(false)
 	, WaitBatchCount(0)
 	, NumCachedAssetFiles(0)
 	, NumUncachedAssetFiles(0)
@@ -3181,10 +3182,12 @@ FAssetDataGatherer::FAssetDataGatherer(const TArray<FString>& InLongPackageNames
 
 	Discovery = MakeUnique<UE::AssetDataGather::Private::FAssetDataDiscovery>(InLongPackageNamesDenyList, InMountRelativePathsDenyList, bInIsSynchronous);
 	FilesToSearch = MakeUnique<UE::AssetDataGather::Private::FFilesToSearch>();
+	FCoreDelegates::OnAllModuleLoadingPhasesComplete.AddRaw(this, &FAssetDataGatherer::OnAllModuleLoadingPhasesComplete);
 }
 
 FAssetDataGatherer::~FAssetDataGatherer()
 {
+	FCoreDelegates::OnAllModuleLoadingPhasesComplete.RemoveAll(this);
 	EnsureCompletion();
 	NewCachedAssetDataMap.Empty();
 	DiskCachedAssetDataMap.Empty();
@@ -3296,7 +3299,7 @@ void FAssetDataGatherer::InnerTickLoop(bool bInIsSynchronousTick, bool bContribu
 	}
 	if (bShouldSaveMonolithicCache)
 	{
-		SaveCacheFileInternal(GPreloadSettings.GetMonolithicCacheFilename(), AssetsToSave, true /* bIsAsyncCache */);
+		SaveCacheFileInternal(GPreloadSettings.GetMonolithicCacheFilename(), AssetsToSave, true /* bIsAsyncCache */, true /* bIsMonolithicCache */);
 		FGathererScopeLock TickScopeLock(&TickLock);
 		check(CacheInUseCount > 0);
 		CacheInUseCount--;
@@ -3432,7 +3435,7 @@ void FAssetDataGatherer::TickInternal(bool& bOutIsTickInterrupt, double& TickSta
 		else
 #endif
 		{
-			Payload = UE::AssetDataGather::Private::LoadCacheFile(GPreloadSettings.GetMonolithicCacheFilename());
+			Payload = UE::AssetDataGather::Private::LoadCacheFile(GPreloadSettings.GetMonolithicCacheFilename(), true /* bIsMonolithicCache */);
 		}
 		ConsumeCacheFile(MoveTemp(Payload));
 		UE_LOG(LogAssetRegistry, Display, TEXT("AssetDataGatherer spent %.3fs loading %s."),
@@ -3967,7 +3970,7 @@ void FAssetDataGatherer::WaitOnPathsInternal(TArrayView<UE::AssetDataGather::Pri
 			GetAssetsToSave(SaveCacheLongPackageNameDirs, AssetsToSave);
 			CacheInUseCount++;
 		}
-		SaveCacheFileInternal(SaveCacheFilename, AssetsToSave, false /* bIsAsyncCacheSave */);
+		SaveCacheFileInternal(SaveCacheFilename, AssetsToSave, false /* bIsAsyncCacheSave */, false /* bIsMonolithicCache */);
 		{
 			FGathererScopeLock TickScopeLock(&TickLock);
 			check(CacheInUseCount > 0);
@@ -4067,7 +4070,7 @@ void FAssetDataGatherer::LoadCacheFile(FStringView CacheFilename)
 		return;
 	}
 
-	FCachePayload Payload = UE::AssetDataGather::Private::LoadCacheFile(CacheFilename);
+	FCachePayload Payload = UE::AssetDataGather::Private::LoadCacheFile(CacheFilename, false /* bIsMonolithicCache */);
 	FScopedPause ScopedPause(*this);
 	CHECK_IS_NOT_LOCKED_CURRENT_THREAD(ResultsLock);
 	FGathererScopeLock TickScopeLock(&TickLock);
@@ -4151,7 +4154,7 @@ void FAssetDataGatherer::GetAssetsToSave(TArrayView<const FString> SaveCacheLong
 	}
 }
 
-void FAssetDataGatherer::SaveCacheFileInternal(const FString& CacheFilename, const TArray<TPair<FName,FDiskCachedAssetData*>>& AssetsToSave, bool bIsAsyncCacheSave)
+void FAssetDataGatherer::SaveCacheFileInternal(const FString& CacheFilename, const TArray<TPair<FName,FDiskCachedAssetData*>>& AssetsToSave, bool bIsAsyncCacheSave, bool bIsMonolithicCache)
 {
 	if (CacheFilename.IsEmpty() || !bCacheWriteEnabled)
 	{
@@ -4180,6 +4183,7 @@ void FAssetDataGatherer::SaveCacheFileInternal(const FString& CacheFilename, con
 #else		
 		checkf(false, TEXT("Cannot save asset registry cache in this configuration"));
 #endif
+		int64 CacheSize = FileAr->TotalSize();
 		// Close file handle before moving temp file to target 
 		FileAr.Reset();
 
@@ -4200,6 +4204,11 @@ void FAssetDataGatherer::SaveCacheFileInternal(const FString& CacheFilename, con
 			}
 		}
 		IFileManager::Get().Move(*CacheFilenameStr, *TempFilename);
+
+		if (bIsMonolithicCache)
+		{
+			UE_LOG(LogAssetRegistry, Display, TEXT("Asset registry cache written as %.1f MiB to %s"), CacheSize/1024.f/1024.f, *CacheFilenameStr);
+		}
 	}
 	else
 	{
@@ -4296,7 +4305,7 @@ FCachePayload SerializeCacheLoad(FAssetRegistryReader& Ar)
 	return Result;
 }
 
-FCachePayload LoadCacheFile(FStringView InCacheFilename)
+FCachePayload LoadCacheFile(FStringView InCacheFilename, bool bIsMonolithicCache)
 {
 	TRACE_CPUPROFILER_EVENT_SCOPE(LoadCacheFile);
 	FString CacheFilename(InCacheFilename);
@@ -4327,6 +4336,7 @@ FCachePayload LoadCacheFile(FStringView InCacheFilename)
 		{
 			FChecksumViewReader ChecksummingReader(MoveTemp(FileReader), CacheFilename);
 			Payload = DoLoad(ChecksummingReader);
+			UE_CLOG(bIsMonolithicCache && Payload.bSucceeded, LogAssetRegistry, Display, TEXT("Asset registry cache read as %.1f MiB from %s"), FileReader.GetTotalSize()/1024.f/1024.f, *CacheFilename);
 			UE_CLOG(!Payload.bSucceeded, LogAssetRegistry, Warning, TEXT("There was an error loading the asset registry cache using memory mapping"));
 		}
 
@@ -4344,6 +4354,7 @@ FCachePayload LoadCacheFile(FStringView InCacheFilename)
 			{
 				FChecksumArchiveReader ChecksummingReader(*FileAr);
 				Payload = DoLoad(ChecksummingReader);
+				UE_CLOG(bIsMonolithicCache && Payload.bSucceeded, LogAssetRegistry, Display, TEXT("Asset registry cache read as %.1f MiB from %s"), FileAr->TotalSize()/1024.f/1024.f, *CacheFilename);
 				UE_CLOG(!Payload.bSucceeded, LogAssetRegistry, Warning, TEXT("There was an error loading the asset registry cache"));
 			}
 		}
@@ -4609,7 +4620,7 @@ void FAssetDataGatherer::SetIsIdle(bool bInIsIdle, double& TickStartTime)
 			CurrentSearchTime += FPlatformTime::Seconds() - TickStartTime;
 			TickStartTime = -1.;
 		}
-		if (!bFinishedInitialDiscovery)
+		if (!bFinishedInitialDiscovery && bAllModuleLoadingPhasesComplete)
 		{
 			bFinishedInitialDiscovery = true;
 
@@ -4645,6 +4656,14 @@ FStringView FAssetDataGatherer::NormalizeLongPackageName(FStringView LongPackage
 	return LongPackageName;
 }
 
+void FAssetDataGatherer::OnAllModuleLoadingPhasesComplete()
+{
+	CHECK_IS_NOT_LOCKED_CURRENT_THREAD(ResultsLock);
+	FGathererScopeLock ResultsScopeLock(&ResultsLock);
+	bAllModuleLoadingPhasesComplete = true;
+	SetIsIdle(false);
+}
+
 namespace UE::AssetDataGather::Private
 {
 
diff --git a/Engine/Source/Runtime/AssetRegistry/Private/AssetDataGatherer.h b/Engine/Source/Runtime/AssetRegistry/Private/AssetDataGatherer.h
--- a/Engine/Source/Runtime/AssetRegistry/Private/AssetDataGatherer.h
+++ b/Engine/Source/Runtime/AssetRegistry/Private/AssetDataGatherer.h
@@ -290,7 +290,7 @@ private:
 	 * data from this Gatherer and so can be run outside any critical section.
 	 */
 	void SaveCacheFileInternal(const FString& CacheFilename,
-		const TArray<TPair<FName,FDiskCachedAssetData*>>& AssetsToSave, bool bIsAsyncCacheSave);
+		const TArray<TPair<FName,FDiskCachedAssetData*>>& AssetsToSave, bool bIsAsyncCacheSave, bool bIsMonolithicCache);
 	/**
 	 * Get the list of FDiskCachedAssetData* that have been loaded in the gatherer, for saving into a cachefile.
 	 * Filters the list of assets by child paths of the elements in SaveCacheLongPackageNameDirs, if it is non-empty.
@@ -324,6 +324,7 @@ private:
 	/** Convert the LongPackageName into our normalized version. */
 	static FStringView NormalizeLongPackageName(FStringView LongPackageName);
 
+	void OnAllModuleLoadingPhasesComplete();
 private:
 
 	/**
@@ -426,6 +427,12 @@ private:
 	bool bFirstTickAfterIdle;
 	/** True if we have finished discovering our first wave of files, to report metrics for that most-important wave. */
 	bool bFinishedInitialDiscovery;
+	/**
+	 * Flag to indicate LaunchEngineLoop's AllModuleLoadingPhases is complete; finishing the initial discovery is
+	 * blocked until preloading complete because plugins can be mounted during startup up until that point, and
+	 * we need to wait for all the plugins that will load before declaring completion.
+	 */
+	bool bAllModuleLoadingPhasesComplete = false;
 
 
 	// Variable section for variables that are read/writable only within TickLock.
