Fix type conversion error.

#rnx

diff --git a/Engine/Source/Runtime/AssetRegistry/Private/AssetDataGatherer.cpp b/Engine/Source/Runtime/AssetRegistry/Private/AssetDataGatherer.cpp
--- a/Engine/Source/Runtime/AssetRegistry/Private/AssetDataGatherer.cpp
+++ b/Engine/Source/Runtime/AssetRegistry/Private/AssetDataGatherer.cpp
@@ -23,6 +23,7 @@
 #include "Misc/AsciiSet.h"
 #include "Misc/Char.h"
 #include "Misc/CommandLine.h"
+#include "Misc/CoreDelegates.h"
 #include "Misc/DelayedAutoRegister.h"
 #include "Misc/Parse.h"
 #include "Misc/Paths.h"
@@ -4207,7 +4208,7 @@ void FAssetDataGatherer::SaveCacheFileInternal(const FString& CacheFilename, con
 
 		if (bIsMonolithicCache)
 		{
-			UE_LOG(LogAssetRegistry, Display, TEXT("Asset registry cache written as %.1f MiB to %s"), CacheSize/1024.f/1024.f, *CacheFilenameStr);
+			UE_LOG(LogAssetRegistry, Display, TEXT("Asset registry cache written as %.1f MiB to %s"), static_cast<float>(CacheSize)/1024.f/1024.f, *CacheFilenameStr);
 		}
 	}
 	else
@@ -4336,7 +4337,7 @@ FCachePayload LoadCacheFile(FStringView InCacheFilename, bool bIsMonolithicCache
 		{
 			FChecksumViewReader ChecksummingReader(MoveTemp(FileReader), CacheFilename);
 			Payload = DoLoad(ChecksummingReader);
-			UE_CLOG(bIsMonolithicCache && Payload.bSucceeded, LogAssetRegistry, Display, TEXT("Asset registry cache read as %.1f MiB from %s"), FileReader.GetTotalSize()/1024.f/1024.f, *CacheFilename);
+			UE_CLOG(bIsMonolithicCache && Payload.bSucceeded, LogAssetRegistry, Display, TEXT("Asset registry cache read as %.1f MiB from %s"), static_cast<float>(FileReader.GetTotalSize())/1024.f/1024.f, *CacheFilename);
 			UE_CLOG(!Payload.bSucceeded, LogAssetRegistry, Warning, TEXT("There was an error loading the asset registry cache using memory mapping"));
 		}
 
@@ -4354,7 +4355,7 @@ FCachePayload LoadCacheFile(FStringView InCacheFilename, bool bIsMonolithicCache
 			{
 				FChecksumArchiveReader ChecksummingReader(*FileAr);
 				Payload = DoLoad(ChecksummingReader);
-				UE_CLOG(bIsMonolithicCache && Payload.bSucceeded, LogAssetRegistry, Display, TEXT("Asset registry cache read as %.1f MiB from %s"), FileAr->TotalSize()/1024.f/1024.f, *CacheFilename);
+				UE_CLOG(bIsMonolithicCache && Payload.bSucceeded, LogAssetRegistry, Display, TEXT("Asset registry cache read as %.1f MiB from %s"), static_cast<float>(FileAr->TotalSize())/1024.f/1024.f, *CacheFilename);
 				UE_CLOG(!Payload.bSucceeded, LogAssetRegistry, Warning, TEXT("There was an error loading the asset registry cache"));
 			}
 		}
