Added bSkipARFilteredAssets to IAssetRegistry::GetAssetByObjectPath

diff --git a/Engine/Source/Runtime/AssetRegistry/Private/AssetRegistry.cpp b/Engine/Source/Runtime/AssetRegistry/Private/AssetRegistry.cpp
--- a/Engine/Source/Runtime/AssetRegistry/Private/AssetRegistry.cpp
+++ b/Engine/Source/Runtime/AssetRegistry/Private/AssetRegistry.cpp
@@ -629,9 +629,9 @@ const TCHAR* GetDevelopmentAssetRegistryFilename()
 	return TEXT("DevelopmentAssetRegistry.bin");
 }
 
-FAssetData IAssetRegistry::K2_GetAssetByObjectPath(const FSoftObjectPath& ObjectPath, bool bIncludeOnlyOnDiskAssets) const
+FAssetData IAssetRegistry::K2_GetAssetByObjectPath(const FSoftObjectPath& ObjectPath, bool bIncludeOnlyOnDiskAssets, bool bSkipARFilteredAssets) const
 {
-	return GetAssetByObjectPath(ObjectPath, bIncludeOnlyOnDiskAssets);
+	return GetAssetByObjectPath(ObjectPath, bIncludeOnlyOnDiskAssets, bSkipARFilteredAssets);
 }
 
 IAssetRegistry::FLoadPackageRegistryData::FLoadPackageRegistryData(bool bInGetDependencies)
@@ -2185,7 +2185,7 @@ void FAssetRegistryImpl::EnumerateDiskAssets(const FARCompiledFilter& InFilter,
 
 }
 
-FAssetData UAssetRegistryImpl::GetAssetByObjectPath(const FSoftObjectPath& ObjectPath, bool bIncludeOnlyOnDiskAssets) const
+FAssetData UAssetRegistryImpl::GetAssetByObjectPath(const FSoftObjectPath& ObjectPath, bool bIncludeOnlyOnDiskAssets, bool bSkipARFilteredAssets) const
 {
 	if (!bIncludeOnlyOnDiskAssets)
 	{
@@ -2195,7 +2195,7 @@ FAssetData UAssetRegistryImpl::GetAssetByObjectPath(const FSoftObjectPath& Objec
 
 		if (Asset)
 		{
-			if (!UE::AssetRegistry::FFiltering::ShouldSkipAsset(Asset))
+			if (!bSkipARFilteredAssets || !UE::AssetRegistry::FFiltering::ShouldSkipAsset(Asset))
 			{
 				return FAssetData(Asset, false /* bAllowBlueprintClass */);
 			}
@@ -2210,14 +2210,14 @@ FAssetData UAssetRegistryImpl::GetAssetByObjectPath(const FSoftObjectPath& Objec
 		FReadScopeLock InterfaceScopeLock(InterfaceLock);
 		const FAssetRegistryState& State = GuardedData.GetState();
 		const FAssetData* FoundData = State.GetAssetByObjectPath(ObjectPath);
-		return (FoundData && !GuardedData.ShouldSkipAsset(FoundData->AssetClassPath, FoundData->PackageFlags)) ? *FoundData : FAssetData();
+		return (FoundData && (!bSkipARFilteredAssets || !GuardedData.ShouldSkipAsset(FoundData->AssetClassPath, FoundData->PackageFlags))) ? *FoundData : FAssetData();
 	}
 }
 
 FAssetData UAssetRegistryImpl::GetAssetByObjectPath(const FName ObjectPath, bool bIncludeOnlyOnDiskAssets) const
 {
 PRAGMA_DISABLE_DEPRECATION_WARNINGS;
-	return GetAssetByObjectPath(FSoftObjectPath(ObjectPath.ToString()));
+	return GetAssetByObjectPath(FSoftObjectPath(ObjectPath.ToString()), bIncludeOnlyOnDiskAssets);
 PRAGMA_ENABLE_DEPRECATION_WARNINGS;
 }
 
diff --git a/Engine/Source/Runtime/AssetRegistry/Private/AssetRegistry.h b/Engine/Source/Runtime/AssetRegistry/Private/AssetRegistry.h
--- a/Engine/Source/Runtime/AssetRegistry/Private/AssetRegistry.h
+++ b/Engine/Source/Runtime/AssetRegistry/Private/AssetRegistry.h
@@ -44,7 +44,7 @@ public:
 	virtual bool EnumerateAssets(const FARCompiledFilter& Filter, TFunctionRef<bool(const FAssetData&)> Callback, bool bSkipARFilteredAssets = true) const override;
 	UE_DEPRECATED(5.1, "Asset path FNames have been deprecated, use FSoftObjectPath instead.")
 	virtual FAssetData GetAssetByObjectPath( const FName ObjectPath, bool bIncludeOnlyOnDiskAssets = false ) const override;
-	virtual FAssetData GetAssetByObjectPath(const FSoftObjectPath& ObjectPath, bool bIncludeOnlyOnDiskAssets = false) const override;
+	virtual FAssetData GetAssetByObjectPath(const FSoftObjectPath& ObjectPath, bool bIncludeOnlyOnDiskAssets = false, bool bSkipARFilteredAssets = true) const override;
 	virtual UE::AssetRegistry::EExists TryGetAssetByObjectPath(const FSoftObjectPath& ObjectPath, FAssetData& OutAssetData) const override;
 	virtual UE::AssetRegistry::EExists TryGetAssetPackageData(const FName PackageName, FAssetPackageData& OutAssetPackageData) const override;
 	virtual bool GetAllAssets(TArray<FAssetData>& OutAssetData, bool bIncludeOnlyOnDiskAssets = false) const override;
diff --git a/Engine/Source/Runtime/AssetRegistry/Public/AssetRegistry/IAssetRegistry.h b/Engine/Source/Runtime/AssetRegistry/Public/AssetRegistry/IAssetRegistry.h
--- a/Engine/Source/Runtime/AssetRegistry/Public/AssetRegistry/IAssetRegistry.h
+++ b/Engine/Source/Runtime/AssetRegistry/Public/AssetRegistry/IAssetRegistry.h
@@ -258,19 +258,21 @@ public:
 	 *
 	 * @param ObjectPath the path of the object to be looked up
 	 * @param bIncludeOnlyOnDiskAssets if true, in-memory objects will be ignored. The call will be faster.
+	 * @param bSkipARFilteredAssets If true, skips Objects that return true for IsAsset but are not assets in the current platform.
 	 * @return the assets data;Will be invalid if object could not be found
 	 */
 	UFUNCTION(BlueprintCallable, BlueprintPure=false, Category = "AssetRegistry", DisplayName="Get Asset By Object Path")
-	ASSETREGISTRY_API virtual FAssetData K2_GetAssetByObjectPath(const FSoftObjectPath& ObjectPath, bool bIncludeOnlyOnDiskAssets = false) const;
+	ASSETREGISTRY_API virtual FAssetData K2_GetAssetByObjectPath(const FSoftObjectPath& ObjectPath, bool bIncludeOnlyOnDiskAssets = false, bool bSkipARFilteredAssets = true) const;
 
 	/**
 	 * Gets the asset data for the specified object path
 	 *
 	 * @param ObjectPath the path of the object to be looked up
 	 * @param bIncludeOnlyOnDiskAssets if true, in-memory objects will be ignored. The call will be faster.
+	 * @param bSkipARFilteredAssets If true, skips Objects that return true for IsAsset but are not assets in the current platform.
 	 * @return the assets data;Will be invalid if object could not be found
 	 */
-	virtual FAssetData GetAssetByObjectPath(const FSoftObjectPath& ObjectPath, bool bIncludeOnlyOnDiskAssets = false) const = 0;
+	virtual FAssetData GetAssetByObjectPath(const FSoftObjectPath& ObjectPath, bool bIncludeOnlyOnDiskAssets = false, bool bSkipARFilteredAssets = true) const = 0;
 
 	/**
 	 * Tries to get the asset data for the specified object path
