Add support for ignoring skip classes in GetMostImportantAsset in order to facilitate consistency across all executable types (e.g. Programs vs Commandlets).

diff --git a/Engine/Source/Developer/IoStoreUtilities/Private/IoStoreUtilities.cpp b/Engine/Source/Developer/IoStoreUtilities/Private/IoStoreUtilities.cpp
--- a/Engine/Source/Developer/IoStoreUtilities/Private/IoStoreUtilities.cpp
+++ b/Engine/Source/Developer/IoStoreUtilities/Private/IoStoreUtilities.cpp
@@ -3132,7 +3132,7 @@ static void AddChunkInfoToAssetRegistry(TMap<FPackageId, TArray<FIoStoreTocChunk
 			continue;
 		}
 
-		const FAssetData* AssetData = UE::AssetRegistry::GetMostImportantAsset(AssetRegistry.GetAssetsByPackageName(AssetPackage.Key), false);
+		const FAssetData* AssetData = UE::AssetRegistry::GetMostImportantAsset(AssetRegistry.GetAssetsByPackageName(AssetPackage.Key), UE::AssetRegistry::EGetMostImportantAssetFlags::IgnoreSkipClasses);
 		if (AssetData == nullptr)
 		{
 			// e.g. /Script packages.
diff --git a/Engine/Source/Editor/UnrealEd/Private/Commandlets/DiffAssetBulkDataCommandlet.cpp b/Engine/Source/Editor/UnrealEd/Private/Commandlets/DiffAssetBulkDataCommandlet.cpp
--- a/Engine/Source/Editor/UnrealEd/Private/Commandlets/DiffAssetBulkDataCommandlet.cpp
+++ b/Engine/Source/Editor/UnrealEd/Private/Commandlets/DiffAssetBulkDataCommandlet.cpp
@@ -319,7 +319,7 @@ int32 UDiffAssetBulkDataCommandlet::Main(const FString& FullCommandLine)
 		{
 			// Nothing has anything to use for diff blaming for this package.
 			// Try to find a representative asset class from the assets in the package.
-			FAssetData const* RepresentativeAsset = UE::AssetRegistry::GetMostImportantAsset(CurrentAssetDatas, true);
+			FAssetData const* RepresentativeAsset = UE::AssetRegistry::GetMostImportantAsset(CurrentAssetDatas, UE::AssetRegistry::EGetMostImportantAssetFlags::RequireOneTopLevelAsset);
 			if (RepresentativeAsset == nullptr)
 			{
 				NoTagPackagesByAssumedClass.FindOrAdd(FTopLevelAssetPath()).Add(ChangedPackageName);
@@ -376,7 +376,7 @@ int32 UDiffAssetBulkDataCommandlet::Main(const FString& FullCommandLine)
 			}
 			else
 			{
-				FAssetData const* RepresentativeAsset = UE::AssetRegistry::GetMostImportantAsset(CurrentAssetDatas, true);
+				FAssetData const* RepresentativeAsset = UE::AssetRegistry::GetMostImportantAsset(CurrentAssetDatas, UE::AssetRegistry::EGetMostImportantAssetFlags::RequireOneTopLevelAsset);
 				if (RepresentativeAsset == nullptr)
 				{
 					PackagesWithUnassignableDiffsByAssumedClass.FindOrAdd(FTopLevelAssetPath()).Add(ChangedPackageName);
diff --git a/Engine/Source/Editor/UnrealEd/Private/CookOnTheFlyServer.cpp b/Engine/Source/Editor/UnrealEd/Private/CookOnTheFlyServer.cpp
--- a/Engine/Source/Editor/UnrealEd/Private/CookOnTheFlyServer.cpp
+++ b/Engine/Source/Editor/UnrealEd/Private/CookOnTheFlyServer.cpp
@@ -11624,8 +11624,7 @@ void UCookOnTheFlyServer::ReportHiddenDependency(FName Referencer, FName Depende
 		{
 			AssetDataPtrs.Add(&AssetData);
 		}
-		const FAssetData* PrimaryAssetData = UE::AssetRegistry::GetMostImportantAsset(AssetDataPtrs,
-			false /* bRequireOneTopLevelAsset */);
+		const FAssetData* PrimaryAssetData = UE::AssetRegistry::GetMostImportantAsset(AssetDataPtrs);
 		if (PrimaryAssetData && PrimaryAssetData->AssetClassPath.IsValid())
 		{
 			PrimaryAssetClassPath = FString::Printf(TEXT(", PrimaryAssetClass=%s"),
diff --git a/Engine/Source/Runtime/AssetRegistry/Private/AssetRegistry.cpp b/Engine/Source/Runtime/AssetRegistry/Private/AssetRegistry.cpp
--- a/Engine/Source/Runtime/AssetRegistry/Private/AssetRegistry.cpp
+++ b/Engine/Source/Runtime/AssetRegistry/Private/AssetRegistry.cpp
@@ -7071,7 +7071,7 @@ UAssetRegistryImpl::FFileLoadProgressUpdatedEvent& UAssetRegistryImpl::OnFileLoa
 
 namespace UE::AssetRegistry
 {
-const FAssetData* GetMostImportantAsset(TConstArrayView<const FAssetData*> PackageAssetDatas, bool bRequireOneTopLevelAsset)
+const FAssetData* GetMostImportantAsset(TConstArrayView<const FAssetData*> PackageAssetDatas, EGetMostImportantAssetFlags InFlags)
 {
 	if (PackageAssetDatas.Num() == 1) // common case
 	{
@@ -7109,10 +7109,14 @@ const FAssetData* GetMostImportantAsset(TConstArrayView<const FAssetData*> Packa
 			return Asset;
 		}
 		// This is after IsUAsset because Blueprints can be the UAsset but also be considered skipable.
-		if (FFiltering::ShouldSkipAsset(Asset->AssetClassPath, Asset->PackageFlags))
+		if (!EnumHasAnyFlags(InFlags, EGetMostImportantAssetFlags::IgnoreSkipClasses))
 		{
-			continue;
+			if (FFiltering::ShouldSkipAsset(Asset->AssetClassPath, Asset->PackageFlags))
+			{
+				continue;
+			}
 		}
+
 		if (Asset->IsTopLevelAsset())
 		{
 			TopLevelAssetCount++;
@@ -7129,7 +7133,7 @@ const FAssetData* GetMostImportantAsset(TConstArrayView<const FAssetData*> Packa
 		}
 	}
 
-	if (bRequireOneTopLevelAsset)
+	if (EnumHasAnyFlags(InFlags, EGetMostImportantAssetFlags::RequireOneTopLevelAsset))
 	{
 		if (TopLevelAssetCount == 1)
 		{
@@ -7175,7 +7179,7 @@ void GetAssetForPackages(TConstArrayView<FName> PackageNames, TMap<FName, FAsset
 	{
 		if (CurrentPackageName != AssetData.PackageName)
 		{
-			OutPackageToAssetData.FindOrAdd(CurrentPackageName) = *GetMostImportantAsset(PackageAssetDatas, false);
+			OutPackageToAssetData.FindOrAdd(CurrentPackageName) = *GetMostImportantAsset(PackageAssetDatas);
 			PackageAssetDatas.Reset();
 			CurrentPackageName = AssetData.PackageName;
 		}
@@ -7183,7 +7187,7 @@ void GetAssetForPackages(TConstArrayView<FName> PackageNames, TMap<FName, FAsset
 		PackageAssetDatas.Push(&AssetData);
 	}
 
-	OutPackageToAssetData.FindOrAdd(CurrentPackageName) = *GetMostImportantAsset(PackageAssetDatas, false);
+	OutPackageToAssetData.FindOrAdd(CurrentPackageName) = *GetMostImportantAsset(PackageAssetDatas);
 
 }
 
diff --git a/Engine/Source/Runtime/AssetRegistry/Public/AssetRegistry/IAssetRegistry.h b/Engine/Source/Runtime/AssetRegistry/Public/AssetRegistry/IAssetRegistry.h
--- a/Engine/Source/Runtime/AssetRegistry/Public/AssetRegistry/IAssetRegistry.h
+++ b/Engine/Source/Runtime/AssetRegistry/Public/AssetRegistry/IAssetRegistry.h
@@ -852,6 +852,19 @@ namespace AssetRegistry
 	 */
 	ASSETREGISTRY_API void GetAssetForPackages(TConstArrayView<FName> PackageNames, TMap<FName, FAssetData>& OutPackageToAssetData);
 
+	enum class EGetMostImportantAssetFlags
+	{
+		None = 0,
+		
+		/* Returns nullptr if there are multiple top level assets. */
+		RequireOneTopLevelAsset = 0x1,
+
+		/* Don't skip AR filtered classes (i.e. BP and BPGC). This is cruicial if you are
+		* running in environments where the skip classes can't be initialized (i.e. Programs) as you
+		* could get different results in such cases. */
+		IgnoreSkipClasses = 0x2
+	};
+
 	/**
 	 * Given a list of asset datas for a specific package, find an asset considered "most important" or "representative".
 	 * This is distinct from a Primary asset, and is used for user facing representation of a package or other cases
@@ -861,13 +874,14 @@ namespace AssetRegistry
 	 *	Tries to find the "UAsset" via the FAssetData::IsUAsset() function. (i.e. asset name matches package name)
 	 *	If none exist, tries to find a "Top Level Asset" using FAssetData::IsToplevelAsset(). (i.e. outer == package)
 	 *		If only one exists, use that.
-	 *		Otherwise, if bRequireOneTopLevelAsset is false, gather the set of possibles and return the first sorted on asset class then name.
+	 *		Otherwise, if EGetMostImportantAssetFlags::RequireOneTopLevelAsset isn't set, gather the set of possibles and return the first sorted on asset class then name.
 	 *			If no top level assets, all package assets
 	 *			If multiple top level assets, all top level assets
 	 * 
 	 * A good source for PackageAssetDatas is FAssetRegistryState::GetAssetsByPackageName.
+	 * 
 	 */
-	ASSETREGISTRY_API const FAssetData* GetMostImportantAsset(TConstArrayView<const FAssetData*> PackageAssetDatas, bool bRequireOneTopLevelAsset);
+	ASSETREGISTRY_API const FAssetData* GetMostImportantAsset(TConstArrayView<const FAssetData*> PackageAssetDatas, EGetMostImportantAssetFlags InFlags = EGetMostImportantAssetFlags::None);
 
 	/*
 	* Returns true if the asset registry should start searching all assets on startup
