AssetRegistry: Fix deadlock due to reentering lock in UAssetRegistryImpl::AssetRenamed.

The problem occurs in this pseudo-callstack:
AssetRenamed -> FWriteScopeLock -> UPackage::IsEmptyPackage -> LoadPackage -> ... -> AssetRegistry functions.
#rnx

diff --git a/Engine/Source/Runtime/AssetRegistry/Private/AssetRegistry.cpp b/Engine/Source/Runtime/AssetRegistry/Private/AssetRegistry.cpp
--- a/Engine/Source/Runtime/AssetRegistry/Private/AssetRegistry.cpp
+++ b/Engine/Source/Runtime/AssetRegistry/Private/AssetRegistry.cpp
@@ -3630,6 +3630,9 @@ void UAssetRegistryImpl::AssetRenamed(const UObject* RenamedAsset, const FString
 			OldPackage = FindPackage(nullptr, *OldPackageName);
 		}
 
+		// Call IsEmptyPackage outside of the lock; it can call LoadPackage internally.
+		bool bOldPackageIsEmpty = OldPackage && UPackage::IsEmptyPackage(OldPackage);
+
 		bool bShouldSkipAsset;
 		UE::AssetRegistry::Impl::FEventContext EventContext;
 		{
@@ -3637,7 +3640,7 @@ void UAssetRegistryImpl::AssetRenamed(const UObject* RenamedAsset, const FString
 			FWriteScopeLock InterfaceScopeLock(InterfaceLock);
 			GuardedData.RemoveEmptyPackage(NewPackage->GetFName());
 
-			if (OldPackage && UPackage::IsEmptyPackage(OldPackage))
+			if (bOldPackageIsEmpty)
 			{
 				GuardedData.AddEmptyPackage(OldPackage->GetFName());
 			}
