AssetRegistry: Fix deadlock in WaitForCompletion when running with async enabled but without async started. This can occur when running a commandlet that calls ScanPathsSynchronous but does not call SearchAllAssets; ScanPathsSynchronous does not start async.
#rnx

diff --git a/Engine/Source/Runtime/AssetRegistry/Private/AssetDataGatherer.cpp b/Engine/Source/Runtime/AssetRegistry/Private/AssetDataGatherer.cpp
--- a/Engine/Source/Runtime/AssetRegistry/Private/AssetDataGatherer.cpp
+++ b/Engine/Source/Runtime/AssetRegistry/Private/AssetDataGatherer.cpp
@@ -1257,11 +1257,12 @@ bool FMountDir::IsChildMountPath(FStringView MountRelPath) const
 }
 
 
-FAssetDataDiscovery::FAssetDataDiscovery(const TArray<FString>& InLongPackageNamesDenyList, const TArray<FString>& InMountRelativePathsDenyList, bool bInIsSynchronous)
+FAssetDataDiscovery::FAssetDataDiscovery(const TArray<FString>& InLongPackageNamesDenyList,
+	const TArray<FString>& InMountRelativePathsDenyList, bool bInAsyncEnabled)
 	: LongPackageNamesDenyList(InLongPackageNamesDenyList)
 	, MountRelativePathsDenyList(InMountRelativePathsDenyList)
 	, Thread(nullptr)
-	, bIsSynchronous(bInIsSynchronous)
+	, bAsyncEnabled(bInAsyncEnabled)
 	, bIsIdle(false)
 	, IsStopped(0)
 	, IsPaused(0)
@@ -1269,9 +1270,9 @@ FAssetDataDiscovery::FAssetDataDiscovery(const TArray<FString>& InLongPackageNam
 {
 	DirLongPackageNamesToNotReport.Add(TEXT("/Game/Collections"));
 
-	if (!bIsSynchronous && !FPlatformProcess::SupportsMultithreading())
+	if (bAsyncEnabled && !FPlatformProcess::SupportsMultithreading())
 	{
-		bIsSynchronous = true;
+		bAsyncEnabled = false;
 		UE_LOG(LogAssetRegistry, Warning, TEXT("Requested asyncronous asset data discovery, but threading support is disabled. Performing a synchronous discovery instead!"));
 	}
 
@@ -1297,13 +1298,18 @@ FAssetDataDiscovery::~FAssetDataDiscovery()
 
 void FAssetDataDiscovery::StartAsync()
 {
-	if (!bIsSynchronous && !Thread)
+	if (bAsyncEnabled && !Thread)
 	{
 		Thread = FRunnableThread::Create(this, TEXT("FAssetDataDiscovery"), 0, TPri_BelowNormal);
 		checkf(Thread, TEXT("Failed to create asset data discovery thread"));
 	}
 }
 
+bool FAssetDataDiscovery::IsSynchronous() const
+{
+	return Thread == nullptr;
+}
+
 bool FAssetDataDiscovery::Init()
 {
 	return true;
@@ -1356,7 +1362,7 @@ uint32 FAssetDataDiscovery::Run()
 FAssetDataDiscovery::FScopedPause::FScopedPause(const FAssetDataDiscovery& InOwner)
 	:Owner(InOwner)
 {
-	if (!Owner.bIsSynchronous)
+	if (!Owner.IsSynchronous())
 	{
 		Owner.IsPaused++;
 	}
@@ -1371,7 +1377,7 @@ FAssetDataDiscovery::FScopedPause::FScopedPause(const FAssetDataDiscovery& InOwn
 FAssetDataDiscovery::FScopedPause::~FScopedPause()
 {
 	Owner.TickOwner.ReleaseOwnershipChecked(Owner.TreeLock);
-	if (!Owner.bIsSynchronous)
+	if (!Owner.IsSynchronous())
 	{
 		check(Owner.IsPaused > 0);
 		Owner.IsPaused--;
@@ -1731,7 +1737,7 @@ void FAssetDataDiscovery::SetIsIdle(bool bInIsIdle, double& TickStartTime)
 	}
 	FGathererScopeLock ResultsScopeLock(&ResultsLock);
 	bIsIdle = bInIsIdle;
-	if (!bIsSynchronous)
+	if (!IsSynchronous())
 	{
 		if (bIsIdle)
 		{
@@ -3132,9 +3138,10 @@ FPreloader GPreloader;
 
 } // namespace UE::AssetDataGather::Private
 
-FAssetDataGatherer::FAssetDataGatherer(const TArray<FString>& InLongPackageNamesDenyList, const TArray<FString>& InMountRelativePathsDenyList, bool bInIsSynchronous)
+FAssetDataGatherer::FAssetDataGatherer(const TArray<FString>& InLongPackageNamesDenyList,
+	const TArray<FString>& InMountRelativePathsDenyList, bool bInAsyncEnabled)
 	: Thread(nullptr)
-	, bIsSynchronous(bInIsSynchronous)
+	, bAsyncEnabled(bInAsyncEnabled)
 	, IsStopped(0)
 	, IsPaused(0)
 	, bInitialPluginsLoaded(false)
@@ -3172,17 +3179,19 @@ FAssetDataGatherer::FAssetDataGatherer(const TArray<FString>& InLongPackageNames
 	bool bCommandlineSynchronous;
 	if (FParse::Bool(FCommandLine::Get(), TEXT("AssetGatherSync="), bCommandlineSynchronous))
 	{
-		bIsSynchronous = bCommandlineSynchronous;
+		bAsyncEnabled = !bCommandlineSynchronous;
 	}
 #endif
-	if (!bIsSynchronous && !FPlatformProcess::SupportsMultithreading())
+	if (bAsyncEnabled && !FPlatformProcess::SupportsMultithreading())
 	{
-		bIsSynchronous = true;
+		bAsyncEnabled = false;
 		UE_LOG(LogAssetRegistry, Warning, TEXT("Requested asynchronous asset data gather, but threading support is disabled. Performing a synchronous gather instead!"));
 	}
-	bIsSynchronousTick = bIsSynchronous;
+	// Tick is Synchronous until StartAsync is called, even if bAsyncEnabled
+	bSynchronousTick = true; 
 
-	Discovery = MakeUnique<UE::AssetDataGather::Private::FAssetDataDiscovery>(InLongPackageNamesDenyList, InMountRelativePathsDenyList, bInIsSynchronous);
+	Discovery = MakeUnique<UE::AssetDataGather::Private::FAssetDataDiscovery>(InLongPackageNamesDenyList,
+		InMountRelativePathsDenyList, bAsyncEnabled);
 	FilesToSearch = MakeUnique<UE::AssetDataGather::Private::FFilesToSearch>();
 	FCoreDelegates::OnAllModuleLoadingPhasesComplete.AddRaw(this, &FAssetDataGatherer::OnAllModuleLoadingPhasesComplete);
 }
@@ -3224,11 +3233,12 @@ void FAssetDataGatherer::ActivateMonolithicCache()
 
 void FAssetDataGatherer::StartAsync()
 {
-	if (!bIsSynchronous && !Thread)
+	if (bAsyncEnabled && !Thread)
 	{
 		Thread = FRunnableThread::Create(this, TEXT("FAssetDataGatherer"), 0, TPri_BelowNormal);
 		checkf(Thread, TEXT("Failed to create asset data gatherer thread"));
 		Discovery->StartAsync();
+		bSynchronousTick = false;
 	}
 }
 
@@ -3244,7 +3254,7 @@ uint32 FAssetDataGatherer::Run()
 
 	while (!IsStopped)
 	{
-		InnerTickLoop(false /* bInIsSynchronousTick */, true /* bContributeToCacheSave */);
+		InnerTickLoop(false /* bInSynchronousTick */, true /* bContributeToCacheSave */);
 
 		for (;;)
 		{
@@ -3263,24 +3273,24 @@ uint32 FAssetDataGatherer::Run()
 	return 0;
 }
 
-void FAssetDataGatherer::InnerTickLoop(bool bInIsSynchronousTick, bool bContributeToCacheSave)
+void FAssetDataGatherer::InnerTickLoop(bool bInSynchronousTick, bool bContributeToCacheSave)
 {
 	using namespace UE::AssetDataGather::Private;
 
 	// Synchronous ticks during Wait contribute to saving of the async cache only if there is no dedicated async thread to do it (Thread == nullptr)
 	// The dedicated async thread always contributes
-	bContributeToCacheSave = !bInIsSynchronousTick || (Thread == nullptr && bContributeToCacheSave);
+	bContributeToCacheSave = !bInSynchronousTick || (Thread == nullptr && bContributeToCacheSave);
 
 	bool bShouldSaveMonolithicCache = false;
 	TArray<TPair<FName, FDiskCachedAssetData*>> AssetsToSave;
 	{
 		CHECK_IS_NOT_LOCKED_CURRENT_THREAD(ResultsLock);
 		FGathererScopeLock RunScopeLock(&TickLock);
-		TGuardValue<bool> ScopeSynchronousTick(bIsSynchronousTick, bInIsSynchronousTick);
+		TGuardValue<bool> ScopeSynchronousTick(bSynchronousTick, bInSynchronousTick);
 		TRACE_CPUPROFILER_EVENT_SCOPE(FAssetDataGatherer::Tick);
 		bool bTickInterruptionEvent = false;
 		double TickStartTime = FPlatformTime::Seconds();
-		while (!IsStopped && (bInIsSynchronousTick || !IsPaused) && !bTickInterruptionEvent)
+		while (!IsStopped && (bInSynchronousTick || !IsPaused) && !bTickInterruptionEvent)
 		{
 			TickInternal(bTickInterruptionEvent, TickStartTime);
 		}
@@ -3311,7 +3321,7 @@ void FAssetDataGatherer::InnerTickLoop(bool bInIsSynchronousTick, bool bContribu
 FAssetDataGatherer::FScopedPause::FScopedPause(const FAssetDataGatherer& InOwner)
 	:Owner(InOwner)
 {
-	if (!Owner.bIsSynchronous)
+	if (!Owner.IsSynchronous())
 	{
 		Owner.IsPaused++;
 	}
@@ -3319,7 +3329,7 @@ FAssetDataGatherer::FScopedPause::FScopedPause(const FAssetDataGatherer& InOwner
 
 FAssetDataGatherer::FScopedPause::~FScopedPause()
 {
-	if (!Owner.bIsSynchronous)
+	if (!Owner.IsSynchronous())
 	{
 		check(Owner.IsPaused > 0)
 		Owner.IsPaused--;
@@ -3336,9 +3346,14 @@ void FAssetDataGatherer::Exit()
 {
 }
 
+bool FAssetDataGatherer::IsAsyncEnabled() const
+{
+	return bAsyncEnabled;
+}
+
 bool FAssetDataGatherer::IsSynchronous() const
 {
-	return bIsSynchronous;
+	return Thread == nullptr;
 }
 
 void FAssetDataGatherer::EnsureCompletion()
@@ -3557,7 +3572,7 @@ void FAssetDataGatherer::TickInternal(bool& bOutIsTickInterrupt, double& TickSta
 		[this, &ReadContexts](int32 Index)
 		{
 			FReadContext& ReadContext = ReadContexts[Index];
-			if (!bIsSynchronousTick && IsPaused)
+			if (!bSynchronousTick && IsPaused)
 			{
 				ReadContext.bCanceled = true;
 				return;
@@ -3961,7 +3976,7 @@ void FAssetDataGatherer::WaitOnPathsInternal(TArrayView<UE::AssetDataGather::Pri
 	TRACE_CPUPROFILER_EVENT_SCOPE(FAssetDataGatherer::Tick);
 	for (;;)
 	{
-		InnerTickLoop(true /* bInIsSynchronousTick */, bContributeToCacheSave);
+		InnerTickLoop(true /* bInSynchronousTick */, bContributeToCacheSave);
 		FGathererScopeLock ResultsScopeLock(&ResultsLock); // WaitBatchCount requires the lock
 		if (WaitBatchCount == 0)
 		{
@@ -4003,7 +4018,7 @@ void FAssetDataGatherer::WaitForIdle()
 	// Tick until idle
 	for (;;)
 	{
-		InnerTickLoop(true /* bInIsSynchronousTick */, true /* bContributeToCacheSave */);
+		InnerTickLoop(true /* bInSynchronousTick */, true /* bContributeToCacheSave */);
 		FGathererScopeLock ResultsScopeLock(&ResultsLock); // bIsIdle requires the lock
 		if (bIsIdle)
 		{
diff --git a/Engine/Source/Runtime/AssetRegistry/Private/AssetDataGatherer.h b/Engine/Source/Runtime/AssetRegistry/Private/AssetDataGatherer.h
--- a/Engine/Source/Runtime/AssetRegistry/Private/AssetDataGatherer.h
+++ b/Engine/Source/Runtime/AssetRegistry/Private/AssetDataGatherer.h
@@ -59,7 +59,7 @@ class FAssetDataGatherer : public FRunnable
 {
 public:
 	FAssetDataGatherer(const TArray<FString>& InLongPackageNamesDenyList,
-		const TArray<FString>& InMountRelativePathsDenyList, bool bInIsSynchronous);
+		const TArray<FString>& InMountRelativePathsDenyList, bool bInAsyncEnabled);
 	virtual ~FAssetDataGatherer();
 
 
@@ -80,6 +80,7 @@ public:
 	virtual void Stop() override;
 	virtual void Exit() override;
 
+	bool IsAsyncEnabled() const;
 	bool IsSynchronous() const;
 	/** Signals to end the thread and waits for it to close before returning */
 	void EnsureCompletion();
@@ -232,7 +233,7 @@ private:
 	 * Helper function to run the tick in a loop-within-a-loop to minimize critical section entry, and to move expensive
 	 * operations out of the critical section
 	 */
-	void InnerTickLoop(bool bInIsSynchronousTick, bool bContributeToCacheSave);
+	void InnerTickLoop(bool bInSynchronousTick, bool bContributeToCacheSave);
 	/**
 	 * Tick function to pump scanning and push results into the search results structure. May be called from devoted
 	 * thread or inline from synchronous functions on other threads.
@@ -341,13 +342,17 @@ private:
 
 	// Variable section for variables that are constant during threading.
 
-	/** Thread to run async Ticks on. Constant during threading. */
+	/**
+	 * Thread to run async Ticks on. Constant during threading.
+	 * Activated when StartAsync is called and bAsyncEnabled is true.
+	 * If null, results will only be added when Wait functions are called. Constant during threading.
+	 */
 	FRunnableThread* Thread;
 	/**
-	 * True if this Gatherer is synchronous, false if it has a worker thread. If synchronous, results will only be
-	 * added when Wait functions are called. Constant during threading.
+	 * True if async gathering is enabled, false if e.g. singlethreaded or disabled by commandline.
+ 	 * Even when enabled, gathering is still synchronous until StartAsync is called.
 	 */
-	bool bIsSynchronous;
+	bool bAsyncEnabled;
 	/** True if AssetPackageData should be gathered. Constant during threading. */
 	bool bGatherAssetPackageData;
 	/** True if dependency data should be gathered. Constant during threading. */
@@ -462,10 +467,10 @@ private:
 	 */
 	int32 CacheInUseCount;
 	/**
-	 * True if the current TickInternal is synchronous, which may be because bIsSynchronous or because the game thread has
+	 * True if the current TickInternal is synchronous, which may be because !IsSynchronous or because the game thread has
 	 * taken over the tick for a synchronous function.
 	 */
-	bool bIsSynchronousTick;
+	bool bSynchronousTick;
 	/** True when a thread is saving an async cache and so another save of the cache should not be triggered. */
 	bool bIsSavingAsyncCache;
 	/** Packages can be marked for retry up until bInitialPluginsLoaded is set. After it is set, we retry them once. */
diff --git a/Engine/Source/Runtime/AssetRegistry/Private/AssetDataGathererPrivate.h b/Engine/Source/Runtime/AssetRegistry/Private/AssetDataGathererPrivate.h
--- a/Engine/Source/Runtime/AssetRegistry/Private/AssetDataGathererPrivate.h
+++ b/Engine/Source/Runtime/AssetRegistry/Private/AssetDataGathererPrivate.h
@@ -577,7 +577,7 @@ class FAssetDataDiscovery : public FRunnable
 {
 public:
 	FAssetDataDiscovery(const TArray<FString>& InLongPackageNamesDenyList,
-		const TArray<FString>& InMountRelativePathsDenyList, bool bInIsSynchronous);
+		const TArray<FString>& InMountRelativePathsDenyList, bool bInAsyncEnabled);
 	virtual ~FAssetDataDiscovery();
 
 
@@ -595,6 +595,8 @@ public:
 	/** Signals to end the thread and waits for it to close before returning */
 	void EnsureCompletion();
 
+	bool IsSynchronous() const;
+
 	// Receiving Results and reading properties (possibly while tick is running)
 	/** Gets search results from the file discovery. */
 	void GetAndTrimSearchResults(bool& bOutIsComplete, TArray<FString>& OutDiscoveredPaths,
@@ -723,8 +725,11 @@ private:
 	TSet<FString> DirLongPackageNamesToNotReport;
 	/** Thread to run the discovery FRunnable on. Read-only while threading is possible. */
 	FRunnableThread* Thread;
-	/** True if this gather request is synchronous (i.e, IsRunningCommandlet()). */
-	bool bIsSynchronous;
+	/**
+	 * True if async discovery is enabled, false if e.g. singlethreaded or disabled by commandline.
+	 * Even when enabled, discovery is still synchronous until StartAsync is called.
+	 */
+	bool bAsyncEnabled;
 
 
 	// Variable section for variables that are atomics read/writable from outside critical sections.
diff --git a/Engine/Source/Runtime/AssetRegistry/Private/AssetRegistry.cpp b/Engine/Source/Runtime/AssetRegistry/Private/AssetRegistry.cpp
--- a/Engine/Source/Runtime/AssetRegistry/Private/AssetRegistry.cpp
+++ b/Engine/Source/Runtime/AssetRegistry/Private/AssetRegistry.cpp
@@ -798,7 +798,7 @@ void FAssetRegistryImpl::Initialize(Impl::FInitializeContext& Context)
 	if (ShouldSearchAllAssetsAtStart())
 	{
 		ConstructGatherer();
-		if (!GlobalGatherer->IsSynchronous())
+		if (GlobalGatherer->IsAsyncEnabled())
 		{
 			SearchAllAssetsInitialAsync(Context.Events, Context.InheritanceContext);
 		}
@@ -1521,8 +1521,8 @@ void FAssetRegistryImpl::ConstructGatherer()
 		EngineIni->GetArray(TEXT("AssetRegistry"), TEXT("BlacklistContentSubPathScanFilters"), ContentSubPathsDenyList);
 	}
 
-	bool bIsSynchronous = IsRunningGame();
-	GlobalGatherer = MakeUnique<FAssetDataGatherer>(PathsDenyList, ContentSubPathsDenyList, bIsSynchronous);
+	bool bAsyncGatherEnabled = !IsRunningGame();
+	GlobalGatherer = MakeUnique<FAssetDataGatherer>(PathsDenyList, ContentSubPathsDenyList, bAsyncGatherEnabled);
 
 	// Read script packages if all initial plugins have been loaded, otherwise do nothing; we wait for the callback.
 	ELoadingPhase::Type LoadingPhase = IPluginManager::Get().GetLastCompletedLoadingPhase();
@@ -1593,7 +1593,7 @@ void FAssetRegistryImpl::SearchAllAssets(Impl::FEventContext& EventContext,
 {
 	ConstructGatherer();
 	FAssetDataGatherer& Gatherer = *GlobalGatherer;
-	if (Gatherer.IsSynchronous())
+	if (!Gatherer.IsAsyncEnabled())
 	{
 		UE_CLOG(!bSynchronousSearch, LogAssetRegistry, Warning, TEXT("SearchAllAssets: Gatherer is in synchronous mode; forcing bSynchronousSearch=true."));
 		bSynchronousSearch = true;
