Allow asset registry prune to remove FDependsNode that are missing FAssetData
#rnx

diff --git a/Engine/Source/Runtime/AssetRegistry/Private/AssetRegistryState.cpp b/Engine/Source/Runtime/AssetRegistry/Private/AssetRegistryState.cpp
--- a/Engine/Source/Runtime/AssetRegistry/Private/AssetRegistryState.cpp
+++ b/Engine/Source/Runtime/AssetRegistry/Private/AssetRegistryState.cpp
@@ -475,6 +475,21 @@ void FAssetRegistryState::PruneAssetData(const TSet<FName>& RequiredPackages, co
 
 void FAssetRegistryState::PruneAssetData(const TSet<FName>& RequiredPackages, const TSet<FName>& RemovePackages, const TSet<int32> ChunksToKeep, const FAssetRegistrySerializationOptions& Options)
 {
+	FAssetRegistryPruneOptions PruneOptions;
+	PruneOptions.RequiredPackages = RequiredPackages;
+	PruneOptions.RemovePackages = RemovePackages;
+	PruneOptions.ChunksToKeep = ChunksToKeep;
+	PruneOptions.Options = Options;
+	Prune(PruneOptions);
+}
+
+void FAssetRegistryState::Prune(const FAssetRegistryPruneOptions& PruneOptions)
+{
+	const TSet<FName>& RequiredPackages = PruneOptions.RequiredPackages;
+	const TSet<FName>& RemovePackages = PruneOptions.RemovePackages;
+	const TSet<int32>& ChunksToKeep = PruneOptions.ChunksToKeep;
+	const FAssetRegistrySerializationOptions& Options = PruneOptions.Options;
+
 	const bool bIsFilteredByChunkId = ChunksToKeep.Num() != 0;
 	const bool bIsFilteredByRequiredPackages = RequiredPackages.Num() != 0;
 	const bool bIsFilteredByRemovedPackages = RemovePackages.Num() != 0;
@@ -485,6 +500,8 @@ void FAssetRegistryState::PruneAssetData(const TSet<FName>& RequiredPackages, co
 	TArray<FAssetData*> AllAssetData = CachedAssets.Array();
 	TSet<FDependsNode*> RemoveDependsNodes;
 
+	TSet<FPrimaryAssetId> KnownPrimaryAssetIds;
+
 	// Remove assets and mark-for-removal any dependencynodes for assets removed due to having no tags
 	for (FAssetData* AssetData : AllAssetData)
 	{
@@ -532,6 +549,14 @@ void FAssetRegistryState::PruneAssetData(const TSet<FName>& RequiredPackages, co
 				}
 			}
 		}
+		else if (PruneOptions.bRemoveDependenciesWithoutPackages)
+		{
+			FPrimaryAssetId PrimaryAssetId = AssetData->GetPrimaryAssetId();
+			if (PrimaryAssetId.IsValid())
+			{
+				KnownPrimaryAssetIds.Add(PrimaryAssetId);
+			}
+		}
 	}
 
 	TArray<FDependsNode*> AllDependsNodes;
@@ -558,6 +583,20 @@ void FAssetRegistryState::PruneAssetData(const TSet<FName>& RequiredPackages, co
 		{
 			bRemoveDependsNode = true;
 		}
+		else if (PruneOptions.bRemoveDependenciesWithoutPackages)
+		{
+			const FPrimaryAssetId PrimaryAssetId = Id.GetPrimaryAssetId();
+			if (PrimaryAssetId.IsValid() && Id.IsObject())
+			{
+				if (!KnownPrimaryAssetIds.Contains(PrimaryAssetId))
+				{
+					if (!PruneOptions.RemoveDependenciesWithoutPackagesKeepPrimaryAssetTypes.Contains(PrimaryAssetId.PrimaryAssetType))
+					{
+						bRemoveDependsNode = true;
+					}
+				}
+			}
+		}
 		
 		if (bRemoveDependsNode)
 		{
diff --git a/Engine/Source/Runtime/AssetRegistry/Public/AssetRegistry/AssetRegistryState.h b/Engine/Source/Runtime/AssetRegistry/Public/AssetRegistry/AssetRegistryState.h
--- a/Engine/Source/Runtime/AssetRegistry/Public/AssetRegistry/AssetRegistryState.h
+++ b/Engine/Source/Runtime/AssetRegistry/Public/AssetRegistry/AssetRegistryState.h
@@ -114,6 +114,20 @@ struct FAssetRegistryLoadOptions
 	int32 ParallelWorkers = 0;
 };
 
+struct FAssetRegistryPruneOptions
+{
+	TSet<FName> RequiredPackages;
+	TSet<FName> RemovePackages;
+	TSet<int32> ChunksToKeep;
+	FAssetRegistrySerializationOptions Options;
+
+	/* Remove FDependsNodes that do not point to packages */
+	bool bRemoveDependenciesWithoutPackages = false;
+
+	/* List of types that should not be pruned because they do not have a package */
+	TSet<FPrimaryAssetType> RemoveDependenciesWithoutPackagesKeepPrimaryAssetTypes;
+};
+
 namespace UE::AssetRegistry::Private
 {
 	/* 
@@ -502,6 +516,7 @@ public:
 	 */
 	void PruneAssetData(const TSet<FName>& RequiredPackages, const TSet<FName>& RemovePackages, const TSet<int32> ChunksToKeep, const FAssetRegistrySerializationOptions& Options);
 	void PruneAssetData(const TSet<FName>& RequiredPackages, const TSet<FName>& RemovePackages, const FAssetRegistrySerializationOptions& Options);
+	void Prune(const FAssetRegistryPruneOptions& PruneOptions);
 
 	
 	/**
