Asset Registry:
- Added new function ScanSynchronous that is meant to be a replacement of ScanPathsSynchronous and ScanFilesSynchronous (not deprecated yet).
- This function takes a scan flag enum instead of bools and also introduce a new option to skip in-memory assets registrations during scan.

diff --git a/Engine/Source/Runtime/AssetRegistry/Private/AssetRegistry.cpp b/Engine/Source/Runtime/AssetRegistry/Private/AssetRegistry.cpp
--- a/Engine/Source/Runtime/AssetRegistry/Private/AssetRegistry.cpp
+++ b/Engine/Source/Runtime/AssetRegistry/Private/AssetRegistry.cpp
@@ -3435,17 +3435,42 @@ bool UAssetRegistryImpl::PathExists(const FName PathToTest) const
 void UAssetRegistryImpl::ScanPathsSynchronous(const TArray<FString>& InPaths, bool bForceRescan,
 	bool bIgnoreDenyListScanFilters)
 {
-	ScanPathsSynchronousInternal(InPaths, TArray<FString>(), bForceRescan, bIgnoreDenyListScanFilters);
+	// The contract of this older version of ScanSynchronous always set the WaitForInMemoryObjects flag.
+	UE::AssetRegistry::EScanFlags ScanFlags = UE::AssetRegistry::EScanFlags::WaitForInMemoryObjects;
+
+	if (bForceRescan)
+	{
+		ScanFlags |= UE::AssetRegistry::EScanFlags::ForceRescan;
+	}
+
+	if (bIgnoreDenyListScanFilters)
+	{
+		ScanFlags |= UE::AssetRegistry::EScanFlags::IgnoreDenyListScanFilters;
+	}
+
+	ScanPathsSynchronousInternal(InPaths, TArray<FString>(), ScanFlags);
 }
 
 void UAssetRegistryImpl::ScanFilesSynchronous(const TArray<FString>& InFilePaths, bool bForceRescan)
 {
-	ScanPathsSynchronousInternal(TArray<FString>(), InFilePaths, bForceRescan,
-		false /* bIgnoreDenyListScanFilters */);
+	// The contract of this older version of ScanSynchronous always set the WaitForInMemoryObjects flag.
+	UE::AssetRegistry::EScanFlags ScanFlags = UE::AssetRegistry::EScanFlags::WaitForInMemoryObjects;
+
+	if (bForceRescan)
+	{
+		ScanFlags |= UE::AssetRegistry::EScanFlags::ForceRescan;
+	}
+
+	ScanPathsSynchronousInternal(TArray<FString>(), InFilePaths, ScanFlags);
+}
+
+void UAssetRegistryImpl::ScanSynchronous(const TArray<FString>& InPaths, const TArray<FString>& InFilePaths, UE::AssetRegistry::EScanFlags InScanFlags)
+{
+	ScanPathsSynchronousInternal(InPaths, InFilePaths, InScanFlags);
 }
 
 void UAssetRegistryImpl::ScanPathsSynchronousInternal(const TArray<FString>& InDirs, const TArray<FString>& InFiles,
-	bool bInForceRescan, bool bInIgnoreDenyListScanFilters)
+	UE::AssetRegistry::EScanFlags InScanFlags)
 {
 	UE_SCOPED_IO_ACTIVITY(*WriteToString<256>("Scan Paths"));
 
@@ -3453,11 +3478,15 @@ void UAssetRegistryImpl::ScanPathsSynchronousInternal(const TArray<FString>& InD
 	UE_TRACK_REFERENCING_OPNAME_SCOPED(PackageAccessTrackingOps::NAME_ResetContext);
 	const double SearchStartTime = FPlatformTime::Seconds();
 
+	const bool bForceRescan = !!(InScanFlags & UE::AssetRegistry::EScanFlags::ForceRescan);
+	const bool bIgnoreDenyListScanFilters = !!(InScanFlags & UE::AssetRegistry::EScanFlags::IgnoreDenyListScanFilters);
+	const bool bWaitForInMemoryObjects = !!(InScanFlags & UE::AssetRegistry::EScanFlags::WaitForInMemoryObjects);
+
 	UE::AssetRegistry::Impl::FEventContext EventContext;
 	UE::AssetRegistry::Impl::FClassInheritanceContext InheritanceContext;
 	UE::AssetRegistry::Impl::FClassInheritanceBuffer InheritanceBuffer;
 	UE::AssetRegistry::Impl::FScanPathContext Context(EventContext, InheritanceContext, InDirs, InFiles,
-		bInForceRescan, bInIgnoreDenyListScanFilters, nullptr /* OutFindAssets */);
+		bForceRescan, bIgnoreDenyListScanFilters, nullptr /* OutFindAssets */);
 
 	{
 		LLM_SCOPE(ELLMTag::AssetRegistry);
@@ -3474,7 +3503,10 @@ void UAssetRegistryImpl::ScanPathsSynchronousInternal(const TArray<FString>& InD
 	}
 
 #if WITH_EDITOR
-	ProcessLoadedAssetsToUpdateCache(EventContext, -1., Context.Status);
+	if (bWaitForInMemoryObjects)
+	{
+		ProcessLoadedAssetsToUpdateCache(EventContext, -1., Context.Status);
+	}
 #endif
 	Broadcast(EventContext);
 
diff --git a/Engine/Source/Runtime/AssetRegistry/Private/AssetRegistry.h b/Engine/Source/Runtime/AssetRegistry/Private/AssetRegistry.h
--- a/Engine/Source/Runtime/AssetRegistry/Private/AssetRegistry.h
+++ b/Engine/Source/Runtime/AssetRegistry/Private/AssetRegistry.h
@@ -101,6 +101,7 @@ public:
 	virtual void WaitForCompletion() override;
 	virtual void ClearGathererCache() override;
 	virtual void WaitForPackage(const FString& PackageName) override;
+	virtual void ScanSynchronous(const TArray<FString>& InPaths, const TArray<FString>& InFilePaths, UE::AssetRegistry::EScanFlags InScanFlags = UE::AssetRegistry::EScanFlags::None) override;
 	virtual void ScanPathsSynchronous(const TArray<FString>& InPaths, bool bForceRescan = false, bool bIgnoreDenyListScanFilters = false) override;
 	virtual void ScanFilesSynchronous(const TArray<FString>& InFilePaths, bool bForceRescan = false) override;
 	virtual void PrioritizeSearchPath(const FString& PathToPrioritize) override;
@@ -236,7 +237,7 @@ private:
 
 	/** Shared helper for Scan*Synchronous function */
 	void ScanPathsSynchronousInternal(const TArray<FString>& InDirs, const TArray<FString>& InFiles,
-		bool bInForceRescan, bool bInIgnoreDenyListScanFilters);
+		UE::AssetRegistry::EScanFlags InScanFlags);
 
 #if WITH_EDITOR
 	/** Create FAssetData from any loaded UObject assets and store the updated AssetData in the state */
diff --git a/Engine/Source/Runtime/AssetRegistry/Public/AssetRegistry/IAssetRegistry.h b/Engine/Source/Runtime/AssetRegistry/Public/AssetRegistry/IAssetRegistry.h
--- a/Engine/Source/Runtime/AssetRegistry/Public/AssetRegistry/IAssetRegistry.h
+++ b/Engine/Source/Runtime/AssetRegistry/Public/AssetRegistry/IAssetRegistry.h
@@ -55,6 +55,20 @@ namespace EAssetAvailabilityProgressReportingType
 	};
 }
 
+namespace UE::AssetRegistry
+{
+
+enum class EScanFlags : uint32
+{
+	None = 0,
+	ForceRescan = 1 << 0,				// the paths will be scanned again, even if they were previously scanned
+	IgnoreDenyListScanFilters = 1 << 1,	// ignore deny list scan filters
+	WaitForInMemoryObjects = 1 << 2,	// update the tags of all assets that have loaded into memory before returning from the scan
+};
+ENUM_CLASS_FLAGS(EScanFlags);
+
+} // namespace UE::AssetRegistry
+
 USTRUCT(BlueprintType)
 struct FAssetRegistryDependencyOptions
 {
@@ -539,6 +553,9 @@ public:
 	UFUNCTION(BlueprintCallable, Category = "AssetRegistry")
 	virtual void ScanFilesSynchronous(const TArray<FString>& InFilePaths, bool bForceRescan = false) = 0;
 
+	/** Scan the supplied paths and files recursively right now and populate the asset registry. */
+	virtual void ScanSynchronous(const TArray<FString>& InPaths, const TArray<FString>& InFilePaths, UE::AssetRegistry::EScanFlags InScanFlags = UE::AssetRegistry::EScanFlags::None) = 0;
+
 	/** Look for all assets on disk (can be async or synchronous) */
 	UFUNCTION(BlueprintCallable, Category = "AssetRegistry")
 	virtual void SearchAllAssets(bool bSynchronousSearch) = 0;
