Persist contents of loaded asset registry cache to new cache when saving.
Prune entries which were not discovered for 3 days to persist cache entries when changing e.g. loaded game feature plugins.

diff --git a/Engine/Source/Runtime/AssetRegistry/AssetRegistry.Build.cs b/Engine/Source/Runtime/AssetRegistry/AssetRegistry.Build.cs
--- a/Engine/Source/Runtime/AssetRegistry/AssetRegistry.Build.cs
+++ b/Engine/Source/Runtime/AssetRegistry/AssetRegistry.Build.cs
@@ -10,6 +10,7 @@ public class AssetRegistry : ModuleRules
 			new string[] {
 				"Core",
 				"CoreUObject",
+				"TraceLog",
 				"ApplicationCore",
 				"Projects",
 			}
diff --git a/Engine/Source/Runtime/AssetRegistry/Private/AssetDataGatherer.cpp b/Engine/Source/Runtime/AssetRegistry/Private/AssetDataGatherer.cpp
--- a/Engine/Source/Runtime/AssetRegistry/Private/AssetDataGatherer.cpp
+++ b/Engine/Source/Runtime/AssetRegistry/Private/AssetDataGatherer.cpp
@@ -32,17 +32,23 @@
 #include "Misc/ScopeExit.h"
 #include "Misc/TrackedActivity.h"
 #include "Serialization/Archive.h"
+#include "String/LexFromString.h"
 #include "Tasks/Task.h"
 #include "Templates/UniquePtr.h"
 #include "Templates/UnrealTemplate.h"
+#include "Trace/Trace.h"
+#include "Trace/Trace.inl"
 
 namespace AssetDataGathererConstants
 {
 	constexpr int32 SingleThreadFilesPerBatch = 3;
 	constexpr int32 ExpectedMaxBatchSize = 100;
-	constexpr int32 MinSecondsToElapseBeforeCacheWrite = 60;
-	static constexpr uint32 CacheSerializationMagic = 0x3339D87B; // Versioning and integrity checking
+	constexpr int32 MinSecondsToElapseBeforeCacheWrite = 30;
+	constexpr int32 MinAssetReadsBeforeCacheWrite = 100000;
+	constexpr int32 CacheShardAssetCount = 250000;
+	static constexpr uint32 CacheSerializationMagic = 0xC3A9D89B; // Versioning and integrity checking
 	static constexpr uint64 CurrentVersion = FAssetRegistryVersion::LatestVersion | (uint64(CacheSerializationMagic) << 32);
+	static const FTimespan CachePruneAge = FTimespan(7, 0, 0, 0);
 	/**
 	 * The maximum number of threads used in the ParallelFor in FAssetDataDiscovery::TickInternal.
 	 * Higher values cause more time spent waking up workers. Lower values cause more time spent running the DirectoryEnumeration on each thread.
@@ -89,7 +95,7 @@ struct FCachePayload
 
 void SerializeCacheSave(FAssetRegistryWriter& Ar, const TArray<TPair<FName, FDiskCachedAssetData*>>& AssetsToSave);
 FCachePayload SerializeCacheLoad(FAssetRegistryReader& Ar);
-FCachePayload LoadCacheFile(FStringView CacheFilename, bool bIsMonolithicCache);
+TArray<FCachePayload> LoadCacheFiles(TConstArrayView<FString> CacheFilenames, bool bIsMonolithicCache);
 
 
 /** InOutResult = Value, but without shrinking the string to fit. */
@@ -3015,7 +3021,7 @@ struct FPreloadSettings
 		bCacheWriteEnabled = !bNoAssetRegistryCache && !bNoAssetRegistryCacheWrite && !bMultiprocess;
 		bool bAsyncEnabled = FPlatformProcess::SupportsMultithreading() && FTaskGraphInterface::IsRunning();
 
-		MonolithicCacheFilename = AssetRegistryCacheRootFolder / (bGatherDependsData ? TEXT("CachedAssetRegistry.bin") : TEXT("CachedAssetRegistryNoDeps.bin"));
+		MonolithicCacheBaseFilename = AssetRegistryCacheRootFolder / (bGatherDependsData ? TEXT("CachedAssetRegistry") : TEXT("CachedAssetRegistryNoDeps"));
 #if UE_EDITOR // See note on FPreloader for why we only allow preloading if UE_EDITOR
 		bMonolithicCacheActivatedDuringPreload = bAsyncEnabled && UE::AssetRegistry::ShouldSearchAllAssetsAtStart();
 #else
@@ -3046,16 +3052,36 @@ struct FPreloadSettings
 	{
 		return bForceDependsGathering;
 	}
-	const FString& GetMonolithicCacheFilename() const
+	FString GetLegacyMonolithicCacheFilename() const
 	{
-		return MonolithicCacheFilename;
+		return MonolithicCacheBaseFilename + TEXT(".bin");
+	}
+	const FString& GetMonolithicCacheBaseFilename() const
+	{
+		return MonolithicCacheBaseFilename;
 	}
 	const FString& GetAssetRegistryCacheRootFolder() const
 	{
 		return AssetRegistryCacheRootFolder;
 	}
+	
+	TArray<FString> FindShardedMonolithicCacheFiles() const
+	{
+		TArray<FString> CachePaths;
+		IFileManager::Get().FindFiles(CachePaths, *(GetMonolithicCacheBaseFilename() + TEXT("_*.bin")), /* Files */ true, /* Directories */ false);
+		if (CachePaths.Num())
+		{
+			FString Directory = FPaths::GetPath(GetMonolithicCacheBaseFilename());
+			for (FString& Path : CachePaths)
+			{
+				Path = Directory / Path;
+			}
+		}
+		return CachePaths;
+	}
+
 private:
-	FString MonolithicCacheFilename;
+	FString MonolithicCacheBaseFilename;
 	FString AssetRegistryCacheRootFolder;
 	bool bForceDependsGathering = false;
 	bool bGatherDependsData = false;
@@ -3093,17 +3119,17 @@ public:
 		// PreloadReady.Wait();
 	}
 
-	FCachePayload Consume()
+	TArray<FCachePayload> Consume()
 	{
 		check(bInitialized);
 		if (bConsumed)
 		{
-			return FCachePayload();
+			return {};
 		}
 		bConsumed = true;
 		PreloadReady.Wait();
 		PreloadReady.Reset();
-		return FCachePayload(MoveTemp(Payload));
+		return MoveTemp(Payloads);
 	}
 
 private:
@@ -3114,22 +3140,23 @@ private:
 
 		if (GPreloadSettings.IsPreloadMonolithicCache())
 		{
-			if (IFileManager::Get().FileExists(*GPreloadSettings.GetMonolithicCacheFilename()))
+			TArray<FString> CachePaths = GPreloadSettings.FindShardedMonolithicCacheFiles();
+			if (CachePaths.Num())
 			{
-				PreloadReady = Async(EAsyncExecution::TaskGraph, [this]() { LoadAsync(); });
+				PreloadReady = Async(EAsyncExecution::TaskGraph, [this, CachePaths=MoveTemp(CachePaths)]() { LoadAsync(CachePaths); });
 			}
 		}
 		bInitialized = true;
 	}
 
-	void LoadAsync()
+	void LoadAsync(const TArray<FString>& Paths)
 	{
 		LLM_SCOPE(ELLMTag::AssetRegistry);
-		Payload = LoadCacheFile(GPreloadSettings.GetMonolithicCacheFilename(), true /* bIsMonolithicCache */);
+		Payloads = LoadCacheFiles(Paths, true /* bIsMonolithicCache */);
 	}
 
 	TFuture<void> PreloadReady;
-	FCachePayload Payload;
+	TArray<FCachePayload> Payloads;
 	bool bInitialized = false;
 	bool bConsumed = false;
 };
@@ -3142,6 +3169,7 @@ FAssetDataGatherer::FAssetDataGatherer(const TArray<FString>& InLongPackageNames
 	const TArray<FString>& InMountRelativePathsDenyList, bool bInAsyncEnabled)
 	: Thread(nullptr)
 	, bAsyncEnabled(bInAsyncEnabled)
+	, GatherStartTime(FDateTime::Now())
 	, IsStopped(0)
 	, IsPaused(0)
 	, bInitialPluginsLoaded(false)
@@ -3158,6 +3186,7 @@ FAssetDataGatherer::FAssetDataGatherer(const TArray<FString>& InLongPackageNames
 	, WaitBatchCount(0)
 	, NumCachedAssetFiles(0)
 	, NumUncachedAssetFiles(0)
+	, LastMonolithicCacheSaveUncachedAssetFiles(0)
 	, CacheInUseCount(0)
 	, bIsSavingAsyncCache(false)
 	, bFlushedRetryFiles(false)
@@ -3311,11 +3340,99 @@ void FAssetDataGatherer::InnerTickLoop(bool bInSynchronousTick, bool bContribute
 	}
 	if (bShouldSaveMonolithicCache)
 	{
-		SaveCacheFileInternal(GPreloadSettings.GetMonolithicCacheFilename(), AssetsToSave, true /* bIsAsyncCache */, true /* bIsMonolithicCache */);
-		FGathererScopeLock TickScopeLock(&TickLock);
-		check(CacheInUseCount > 0);
-		CacheInUseCount--;
+		SaveMonolithicCacheFile(AssetsToSave);
+	}
+}
+
+void FAssetDataGatherer::SaveMonolithicCacheFile(const TArray<TPair<FName,FDiskCachedAssetData*>>& AssetsToSave)
+{
+	using namespace UE::AssetDataGather::Private;
+	TRACE_CPUPROFILER_EVENT_SCOPE_STR("Save Monolithic Cache")
+	
+	// If we recently saved or loaded the file then pause for 0.5 seconds before trying to save on
+	// top of it, to avoid failure to be able to delete the file we just saved/loaded.
+	double LocalLastCacheWriteTime;
+	{
+		FGathererScopeLock ResultsScopeLock(&ResultsLock);
+		LocalLastCacheWriteTime = LastCacheWriteTime;
+	}
+	double CurrentTime = FPlatformTime::Seconds();
+	constexpr float WaitTimeBeforeReopen = 0.5f;
+	if (CurrentTime < LastCacheWriteTime + WaitTimeBeforeReopen)
+	{
+		FPlatformProcess::Sleep(static_cast<float>(LastCacheWriteTime + WaitTimeBeforeReopen - CurrentTime));
+	}
+
+	const int32 MaxShards = FPlatformMisc::NumberOfCores();
+	constexpr int32 CacheShardAssetCount = AssetDataGathererConstants::CacheShardAssetCount;
+	const int32 CacheShards = FMath::RoundUpToPowerOfTwo(FMath::Min((AssetsToSave.Num() + CacheShardAssetCount-1) / CacheShardAssetCount, MaxShards));
+	const int32 ShardMask = CacheShards - 1;
+	std::atomic<int64> TotalCacheSize{0};
+
+	TArray<TArray<TPair<FName, FDiskCachedAssetData*>>> DataPerShard;
+	{
+		// Hash package names by string rather than by name id so that cache shard is maintained across runs in case of interruption
+		TRACE_CPUPROFILER_EVENT_SCOPE_STR("Build Shards")
+		DataPerShard.AddDefaulted(CacheShards);
+		for (TArray<TPair<FName, FDiskCachedAssetData*>>& ShardData : DataPerShard)
+		{
+			ShardData.Reserve(AssetsToSave.Num() / CacheShards);	
+		}
+		TStringBuilder<256> Buffer;
+		for (const TPair<FName, FDiskCachedAssetData*>& Entry : AssetsToSave)
+		{
+			Buffer.Reset();
+			Buffer << Entry.Key;
+			// Force to lower case to maintain ordering across runs where fname casing changes because of load order
+			// Hash doesn't need to be persistent across utf-8 transition so no need to re-encode strings
+			for (int32 i=0; i < Buffer.Len(); ++i)
+			{
+				Buffer.GetData()[i] = FChar::ToLower(Buffer.GetData()[i]);
+			}
+			uint64 Hash = CityHash64((const char*)Buffer.GetData(), Buffer.Len() * sizeof(TCHAR));
+			int32 Index = (Hash & ShardMask);
+			DataPerShard[Index].Add(Entry);
+		}
 	}
+
+	ParallelFor(CacheShards, [this, &AssetsToSave, &TotalCacheSize, &DataPerShard, ShardMask](int32 Shard) 
+	{
+		TArray<TPair<FName, FDiskCachedAssetData*>>& ShardData = DataPerShard[Shard];
+		FString Filename = FString::Printf(TEXT("%s_%d.bin"), *GPreloadSettings.GetMonolithicCacheBaseFilename(), Shard);
+		int64 CacheSize = SaveCacheFileInternal(Filename, ShardData);
+		TotalCacheSize += CacheSize;
+	}, EParallelForFlags::BackgroundPriority);
+
+	UE_LOG(LogAssetRegistry, Display, TEXT("Asset registry cache written as %.1f MiB to %s_*.bin"), 
+		static_cast<float>(TotalCacheSize)/1024.f/1024.f, *GPreloadSettings.GetMonolithicCacheBaseFilename());
+	// Delete old monolithic cache file if it exists
+	FString OldMonolithicName = GPreloadSettings.GetLegacyMonolithicCacheFilename();
+	IFileManager::Get().Delete(*OldMonolithicName, /* bRequireExists */ false, /* EvenReadOnly */ true, /* Quiet */ true);
+	
+	// Delete any other shards if number of shards was reduced
+	TArray<FString> CacheFiles = GPreloadSettings.FindShardedMonolithicCacheFiles();
+	for (const FString& CacheFile : CacheFiles)
+	{
+		FStringView BaseName = FPathViews::GetBaseFilename(CacheFile);
+		if (int32 Index; BaseName.FindLastChar('_', Index))
+		{
+			BaseName.MidInline(Index+1);
+			int32 Suffix = 0;
+			LexFromString(Suffix, BaseName);
+			if (Suffix >= CacheShards)
+			{
+				IFileManager::Get().Delete(*CacheFile);
+			}
+		}
+	}
+
+	FScopedPause ScopedPause(*this);
+	FGathererScopeLock TickScopeLock(&TickLock);
+	bIsSavingAsyncCache = false;
+	check(CacheInUseCount > 0);
+	CacheInUseCount--;
+	FGathererScopeLock ResultsScopeLock(&ResultsLock);
+	LastCacheWriteTime = FPlatformTime::Seconds();
 }
 
 FAssetDataGatherer::FScopedPause::FScopedPause(const FAssetDataGatherer& InOwner)
@@ -3429,11 +3546,13 @@ void FAssetDataGatherer::TickInternal(bool& bOutIsTickInterrupt, double& TickSta
 			WaitBatchCount = 0; // Clear WaitBatchCount in case it was set higher than FilesToSearch->GetNumAvailable().
 			bOutIsTickInterrupt = true;
 
-			if (!bFinishedInitialDiscovery)
+			const bool bWasInitialDiscoveryFinished = bFinishedInitialDiscovery;
+			SetIsIdle(true, TickStartTime);
+			if (!bWasInitialDiscoveryFinished && bFinishedInitialDiscovery) 
 			{
+				UE_LOG(LogAssetRegistry, Display, TEXT("Triggering cache save on discovery complete"));
 				bSaveAsyncCacheTriggered = true;
 			}
-			SetIsIdle(true, TickStartTime);
 			return;
 		}
 		if (bReadMonolithicCache && !bHasLoadedMonolithicCache)
@@ -3448,20 +3567,21 @@ void FAssetDataGatherer::TickInternal(bool& bOutIsTickInterrupt, double& TickSta
 	if (bLoadMonolithicCache)
 	{
 		double CacheLoadStartTime = FPlatformTime::Seconds();
-		FCachePayload Payload;
+		TArray<FCachePayload> Payloads;
 #if UE_EDITOR
 		if (GPreloadSettings.IsPreloadMonolithicCache())
 		{
-			Payload = GPreloader.Consume();
+			Payloads = GPreloader.Consume();
 		}
 		else
 #endif
 		{
-			Payload = UE::AssetDataGather::Private::LoadCacheFile(GPreloadSettings.GetMonolithicCacheFilename(), true /* bIsMonolithicCache */);
+			TArray<FString> CachePaths = GPreloadSettings.FindShardedMonolithicCacheFiles();
+			Payloads = UE::AssetDataGather::Private::LoadCacheFiles(CachePaths, true /* bIsMonolithicCache */);
 		}
-		ConsumeCacheFile(MoveTemp(Payload));
-		UE_LOG(LogAssetRegistry, Display, TEXT("AssetDataGatherer spent %.3fs loading %s."),
-			FPlatformTime::Seconds() - CacheLoadStartTime, *GPreloadSettings.GetMonolithicCacheFilename());
+		ConsumeCacheFiles(MoveTemp(Payloads));
+		UE_LOG(LogAssetRegistry, Display, TEXT("AssetDataGatherer spent %.3fs loading caches %s_*.bin."),
+			FPlatformTime::Seconds() - CacheLoadStartTime, *GPreloadSettings.GetMonolithicCacheBaseFilename());
 
 		FGathererScopeLock ResultsScopeLock(&ResultsLock);
 		bHasLoadedMonolithicCache = true;
@@ -3526,7 +3646,7 @@ void FAssetDataGatherer::TickInternal(bool& bOutIsTickInterrupt, double& TickSta
 		if (DiskCachedAssetData)
 		{
 			// Check whether we need to invalidate the cached data
-			const FDateTime& CachedTimestamp = DiskCachedAssetData->Timestamp;
+			const FDateTime& CachedTimestamp = DiskCachedAssetData->ModificationTime;
 			if (AssetFileData.PackageTimestamp != CachedTimestamp)
 			{
 				DiskCachedAssetData = nullptr;
@@ -3617,7 +3737,7 @@ void FAssetDataGatherer::TickInternal(bool& bOutIsTickInterrupt, double& TickSta
 			if (bCachePackage)
 			{
 				// Update the cache
-				FDiskCachedAssetData* NewData = new FDiskCachedAssetData(ReadContext.AssetFileData.PackageTimestamp, ReadContext.Extension);
+				FDiskCachedAssetData* NewData = new FDiskCachedAssetData(ReadContext.AssetFileData.PackageTimestamp, GatherStartTime, ReadContext.Extension);
 				NewData->AssetDataList.Reserve(ReadContext.AssetDataFromFile.Num());
 				for (const FAssetData* BackgroundAssetData : ReadContext.AssetDataFromFile)
 				{
@@ -3670,8 +3790,10 @@ void FAssetDataGatherer::TickInternal(bool& bOutIsTickInterrupt, double& TickSta
 			}
 		}
 
+		const int32 NumAssetsReadSinceLastCacheWrite = NumUncachedAssetFiles - LastMonolithicCacheSaveUncachedAssetFiles;
 		if (bWriteMonolithicCache && !bIsSavingAsyncCache 
-			&& FPlatformTime::Seconds() - LocalLastCacheWriteTime >= AssetDataGathererConstants::MinSecondsToElapseBeforeCacheWrite)
+			&& FPlatformTime::Seconds() - LocalLastCacheWriteTime >= AssetDataGathererConstants::MinSecondsToElapseBeforeCacheWrite
+			&& NumAssetsReadSinceLastCacheWrite >= AssetDataGathererConstants::MinAssetReadsBeforeCacheWrite)
 		{
 			bSaveAsyncCacheTriggered = true;
 			bOutIsTickInterrupt = true;
@@ -3992,7 +4114,7 @@ void FAssetDataGatherer::WaitOnPathsInternal(TArrayView<UE::AssetDataGather::Pri
 			GetAssetsToSave(SaveCacheLongPackageNameDirs, AssetsToSave);
 			CacheInUseCount++;
 		}
-		SaveCacheFileInternal(SaveCacheFilename, AssetsToSave, false /* bIsAsyncCacheSave */, false /* bIsMonolithicCache */);
+		SaveCacheFileInternal(SaveCacheFilename, AssetsToSave);
 		{
 			FGathererScopeLock TickScopeLock(&TickLock);
 			check(CacheInUseCount > 0);
@@ -4083,7 +4205,7 @@ FString FAssetDataGatherer::GetCacheFilename(TConstArrayView<FString> CacheFileP
 	return UE::AssetDataGather::Private::GPreloadSettings.GetAssetRegistryCacheRootFolder() / TEXT("AssetRegistryCache") / FString::Printf(TEXT("%08x%s.bin"), CacheHash, bGatherDependsData ? TEXT("") : TEXT("NoDeps"));
 }
 
-void FAssetDataGatherer::LoadCacheFile(FStringView CacheFilename)
+void FAssetDataGatherer::LoadCacheFiles(TConstArrayView<FString> CacheFilenames)
 {
 	using namespace UE::AssetDataGather::Private;
 	if (!bCacheReadEnabled)
@@ -4091,30 +4213,40 @@ void FAssetDataGatherer::LoadCacheFile(FStringView CacheFilename)
 		return;
 	}
 
-	FCachePayload Payload = UE::AssetDataGather::Private::LoadCacheFile(CacheFilename, false /* bIsMonolithicCache */);
+	TArray<FCachePayload> Payloads = UE::AssetDataGather::Private::LoadCacheFiles(CacheFilenames, false /* bIsMonolithicCache */);
 	FScopedPause ScopedPause(*this);
 	CHECK_IS_NOT_LOCKED_CURRENT_THREAD(ResultsLock);
 	FGathererScopeLock TickScopeLock(&TickLock);
-	ConsumeCacheFile(MoveTemp(Payload));
+	ConsumeCacheFiles(MoveTemp(Payloads));
 }
 
-void FAssetDataGatherer::ConsumeCacheFile(UE::AssetDataGather::Private::FCachePayload&& Payload)
+void FAssetDataGatherer::ConsumeCacheFiles(TArray<UE::AssetDataGather::Private::FCachePayload> Payloads)
 {
 	CHECK_IS_LOCKED_CURRENT_THREAD(TickLock);
 	CHECK_IS_NOT_LOCKED_CURRENT_THREAD(ResultsLock);
-	if (!Payload.bSucceeded || Payload.NumAssets == 0)
+	
+	int32 Count = 0;
+	for (UE::AssetDataGather::Private::FCachePayload& Payload : Payloads)
+	{
+		Count += Payload.NumAssets;
+	}
+
+	if (Count == 0)
 	{
 		return;
 	}
 
-	DiskCachedAssetDataMap.Reserve(DiskCachedAssetDataMap.Num() + Payload.NumAssets);
-	for (int32 AssetIndex = 0; AssetIndex < Payload.NumAssets; ++AssetIndex)
+	DiskCachedAssetDataMap.Reserve(DiskCachedAssetDataMap.Num() + Count);
+	for (UE::AssetDataGather::Private::FCachePayload& Payload : Payloads)
 	{
-		DiskCachedAssetDataMap.Add(*(Payload.PackageNames.Get() + AssetIndex),
-			(Payload.AssetDatas.Get() + AssetIndex)); // -C6385
+		for (int32 AssetIndex = 0; AssetIndex < Payload.NumAssets; ++AssetIndex)
+		{
+			DiskCachedAssetDataMap.Add(*(Payload.PackageNames.Get() + AssetIndex),
+				(Payload.AssetDatas.Get() + AssetIndex)); // -C6385
+		}
+		DiskCachedAssetBlocks.Emplace(Payload.NumAssets, Payload.AssetDatas.Release());
+		Payload.Reset();
 	}
-	DiskCachedAssetBlocks.Emplace(Payload.NumAssets, Payload.AssetDatas.Release());
-	Payload.Reset();
 
 	FGathererScopeLock ResultsScopeLock(&ResultsLock);
 	DependencyResults.Reserve(DiskCachedAssetDataMap.Num());
@@ -4139,8 +4271,9 @@ void FAssetDataGatherer::TryReserveSaveMonolithicCache(bool& bOutShouldSave, TAr
 	}
 	if (bOutShouldSave)
 	{
-		GetAssetsToSave(TArrayView<const FString>(), AssetsToSave);
+		GetMonolithicCacheAssetsToSave(AssetsToSave);
 		bIsSavingAsyncCache = true;
+		LastMonolithicCacheSaveUncachedAssetFiles = NumUncachedAssetFiles;
 	}
 	bSaveAsyncCacheTriggered = false;
 }
@@ -4175,15 +4308,50 @@ void FAssetDataGatherer::GetAssetsToSave(TArrayView<const FString> SaveCacheLong
 	}
 }
 
-void FAssetDataGatherer::SaveCacheFileInternal(const FString& CacheFilename, const TArray<TPair<FName,FDiskCachedAssetData*>>& AssetsToSave, bool bIsAsyncCacheSave, bool bIsMonolithicCache)
+void FAssetDataGatherer::GetMonolithicCacheAssetsToSave(TArray<TPair<FName,FDiskCachedAssetData*>>& OutAssetsToSave)
+{
+	CHECK_IS_LOCKED_CURRENT_THREAD(TickLock);
+
+	OutAssetsToSave.Reset();
+	OutAssetsToSave.Reserve(FMath::Max(NewCachedAssetDataMap.Num(), DiskCachedAssetDataMap.Num()));
+	for (const TPair<FName, FDiskCachedAssetData*>& Pair : NewCachedAssetDataMap)
+	{
+		OutAssetsToSave.Add(Pair);
+	}
+
+	for (const TPair<FName, FDiskCachedAssetData*>& Pair : DiskCachedAssetDataMap)
+	{
+		if(NewCachedAssetDataMap.Contains(Pair.Key))
+		{
+			continue; // Data was replaced when populating NewCachedAssetDataMap
+		}	
+		FTimespan Age = GatherStartTime - Pair.Value->LastGatheredTime;
+		// Conservatively persist cached data until final save to avoid pruning the cache too much if discovery is interrupted 
+		// when revisiting a workspace after some delay 
+		if (!bFinishedInitialDiscovery || Age < AssetDataGathererConstants::CachePruneAge)
+		{
+			OutAssetsToSave.Add(Pair);
+		}
+	}
+}
+
+UE_TRACE_EVENT_BEGIN(Cpu, SaveCacheFile, NoSync)
+	UE_TRACE_EVENT_FIELD(UE::Trace::WideString, Path)
+	UE_TRACE_EVENT_FIELD(int32, Count)
+UE_TRACE_EVENT_END()
+
+int64 FAssetDataGatherer::SaveCacheFileInternal(const FString& CacheFilename, const TArray<TPair<FName,FDiskCachedAssetData*>>& AssetsToSave)
 {
 	if (CacheFilename.IsEmpty() || !bCacheWriteEnabled)
 	{
-		return;
+		return 0;
 	}
 	CHECK_IS_NOT_LOCKED_CURRENT_THREAD(TickLock);
 	CHECK_IS_NOT_LOCKED_CURRENT_THREAD(ResultsLock);
-	TRACE_CPUPROFILER_EVENT_SCOPE(SaveCacheFile);
+	UE_TRACE_LOG_SCOPED_T(Cpu, SaveCacheFile, CpuChannel)
+		<< SaveCacheFile.Path(*CacheFilename)
+		<< SaveCacheFile.Count(AssetsToSave.Num());
+	
 	// Save to a temp file first, then move to the destination to avoid corruption
 	FString CacheFilenameStr(CacheFilename);
 	FString TempFilename = CacheFilenameStr + TEXT(".tmp");
@@ -4207,42 +4375,13 @@ void FAssetDataGatherer::SaveCacheFileInternal(const FString& CacheFilename, con
 		int64 CacheSize = FileAr->TotalSize();
 		// Close file handle before moving temp file to target 
 		FileAr.Reset();
-
-		// If we recently saved or loaded the file then pause for 0.5 seconds before trying to save on
-		// top of it, to avoid failure to be able to delete the file we just saved/loaded.
-		if (bIsAsyncCacheSave)
-		{
-			double LocalLastCacheWriteTime;
-			{
-				FGathererScopeLock ResultsScopeLock(&ResultsLock);
-				LocalLastCacheWriteTime = LastCacheWriteTime;
-			}
-			double CurrentTime = FPlatformTime::Seconds();
-			constexpr float WaitTimeBeforeReopen = 0.5f;
-			if (CurrentTime < LastCacheWriteTime + WaitTimeBeforeReopen)
-			{
-				FPlatformProcess::Sleep(static_cast<float>(LastCacheWriteTime + WaitTimeBeforeReopen - CurrentTime));
-			}
-		}
 		IFileManager::Get().Move(*CacheFilenameStr, *TempFilename);
-
-		if (bIsMonolithicCache)
-		{
-			UE_LOG(LogAssetRegistry, Display, TEXT("Asset registry cache written as %.1f MiB to %s"), static_cast<float>(CacheSize)/1024.f/1024.f, *CacheFilenameStr);
-		}
+		return CacheSize;
 	}
 	else
 	{
 		UE_LOG(LogAssetRegistry, Error, TEXT("Failed to open file for write %s"), *TempFilename);
-	}
-
-	if (bIsAsyncCacheSave)
-	{
-		FScopedPause ScopedPause(*this);
-		FGathererScopeLock TickScopeLock(&TickLock);
-		bIsSavingAsyncCache = false;
-		FGathererScopeLock ResultsScopeLock(&ResultsLock);
-		LastCacheWriteTime = FPlatformTime::Seconds();
+		return 0;
 	}
 }
 
@@ -4326,64 +4465,69 @@ FCachePayload SerializeCacheLoad(FAssetRegistryReader& Ar)
 	return Result;
 }
 
-FCachePayload LoadCacheFile(FStringView InCacheFilename, bool bIsMonolithicCache)
+TArray<FCachePayload> LoadCacheFiles(TConstArrayView<FString> InCacheFilenames, bool bIsMonolithicCache)
 {
-	TRACE_CPUPROFILER_EVENT_SCOPE(LoadCacheFile);
-	FString CacheFilename(InCacheFilename);
-
-	auto DoLoad = [&](FArchive& ChecksummingReader)
-	{
-		int32 WorkerReduction = 2; // Current worker + preload task
-		int32 Parallelism = FMath::Max(FTaskGraphInterface::Get().GetNumWorkerThreads() - WorkerReduction, 0);
-		
-		// The discovery cache is always serialized with a fixed format.
-		// We discard it before this point if it's not the latest version, and it always includes editor-only data.
-		FAssetRegistryHeader Header;
-		Header.Version = FAssetRegistryVersion::LatestVersion;
-		Header.bFilterEditorOnlyData = false;
-		FAssetRegistryReader RegistryReader(ChecksummingReader, Parallelism, Header);
-		return RegistryReader.IsError() ? FCachePayload() : SerializeCacheLoad(RegistryReader);
-	};
-
-	FCachePayload Payload;
-	if (FPlatformProperties::SupportsMemoryMappedFiles())
+	TRACE_CPUPROFILER_EVENT_SCOPE(LoadCacheFiles);
+	TArray<FCachePayload> Results;
+	Results.AddDefaulted(InCacheFilenames.Num());
+	ParallelFor(InCacheFilenames.Num(), [&InCacheFilenames, &Results, bIsMonolithicCache](int32 Index)
 	{
-		FMemoryMappedFile File(*CacheFilename);
-		UE::Tasks::FTask Preload = UE::Tasks::Launch(UE_SOURCE_LOCATION, [&]() { File.Preload(); });
-
-		FMemoryViewReader FileReader(File.View());
-		TOptional<uint64> Version = FileReader.TryLoad<uint64>();
-		if (Version == AssetDataGathererConstants::CurrentVersion)
+		const FString& CacheFilename = InCacheFilenames[Index];
+		auto DoLoad = [&](FArchive& ChecksummingReader)
 		{
-			FChecksumViewReader ChecksummingReader(MoveTemp(FileReader), CacheFilename);
-			Payload = DoLoad(ChecksummingReader);
-			UE_CLOG(bIsMonolithicCache && Payload.bSucceeded, LogAssetRegistry, Display, TEXT("Asset registry cache read as %.1f MiB from %s"), static_cast<float>(FileReader.GetTotalSize())/1024.f/1024.f, *CacheFilename);
-			UE_CLOG(!Payload.bSucceeded, LogAssetRegistry, Warning, TEXT("There was an error loading the asset registry cache using memory mapping"));
-		}
+			int32 WorkerReduction = 2; // Current worker + preload task
+			int32 Parallelism = FMath::Max(FTaskGraphInterface::Get().GetNumWorkerThreads() - WorkerReduction, 0);
+			
+			// The discovery cache is always serialized with a fixed format.
+			// We discard it before this point if it's not the latest version, and it always includes editor-only data.
+			FAssetRegistryHeader Header;
+			Header.Version = FAssetRegistryVersion::LatestVersion;
+			Header.bFilterEditorOnlyData = false;
+			FAssetRegistryReader RegistryReader(ChecksummingReader, Parallelism, Header);
+			return RegistryReader.IsError() ? FCachePayload() : SerializeCacheLoad(RegistryReader);
+		};
 
-		Preload.Wait();
-		UE::Tasks::Launch(UE_SOURCE_LOCATION, [KillAsync = MoveTemp(File)](){});
-	}
-	else
-	{
-		TUniquePtr<FArchive> FileAr(IFileManager::Get().CreateFileReader(*CacheFilename, FILEREAD_Silent));
-		if (FileAr && !FileAr->IsError() && FileAr->TotalSize() > sizeof(uint64))
+		FCachePayload Payload;
+		if (FPlatformProperties::SupportsMemoryMappedFiles())
 		{
-			uint64 Version = 0;
-			*FileAr << Version;
+			FMemoryMappedFile File(*CacheFilename);
+			UE::Tasks::FTask Preload = UE::Tasks::Launch(UE_SOURCE_LOCATION, [&]() { File.Preload(); });
+
+			FMemoryViewReader FileReader(File.View());
+			TOptional<uint64> Version = FileReader.TryLoad<uint64>();
 			if (Version == AssetDataGathererConstants::CurrentVersion)
 			{
-				FChecksumArchiveReader ChecksummingReader(*FileAr);
+				FChecksumViewReader ChecksummingReader(MoveTemp(FileReader), CacheFilename);
 				Payload = DoLoad(ChecksummingReader);
-				UE_CLOG(bIsMonolithicCache && Payload.bSucceeded, LogAssetRegistry, Display, TEXT("Asset registry cache read as %.1f MiB from %s"), static_cast<float>(FileAr->TotalSize())/1024.f/1024.f, *CacheFilename);
-				UE_CLOG(!Payload.bSucceeded, LogAssetRegistry, Warning, TEXT("There was an error loading the asset registry cache"));
+				UE_CLOG(bIsMonolithicCache && Payload.bSucceeded, LogAssetRegistry, Display, TEXT("Asset registry cache read as %.1f MiB from %s"), static_cast<float>(FileReader.GetTotalSize())/1024.f/1024.f, *CacheFilename);
+				UE_CLOG(!Payload.bSucceeded, LogAssetRegistry, Warning, TEXT("There was an error loading the asset registry cache using memory mapping"));
+			}
+
+			Preload.Wait();
+			UE::Tasks::Launch(UE_SOURCE_LOCATION, [KillAsync = MoveTemp(File)](){});
+		}
+		else
+		{
+			TUniquePtr<FArchive> FileAr(IFileManager::Get().CreateFileReader(*CacheFilename, FILEREAD_Silent));
+			if (FileAr && !FileAr->IsError() && FileAr->TotalSize() > sizeof(uint64))
+			{
+				uint64 Version = 0;
+				*FileAr << Version;
+				if (Version == AssetDataGathererConstants::CurrentVersion)
+				{
+					FChecksumArchiveReader ChecksummingReader(*FileAr);
+					Payload = DoLoad(ChecksummingReader);
+					UE_CLOG(bIsMonolithicCache && Payload.bSucceeded, LogAssetRegistry, Display, TEXT("Asset registry cache read as %.1f MiB from %s"), static_cast<float>(FileAr->TotalSize())/1024.f/1024.f, *CacheFilename);
+					UE_CLOG(!Payload.bSucceeded, LogAssetRegistry, Warning, TEXT("There was an error loading the asset registry cache"));
+				}
 			}
 		}
-	}
 	
-	return Payload;
-}
+		Results[Index] = MoveTemp(Payload);
+	}, EParallelForFlags::BackgroundPriority);
 
+	return MoveTemp(Results);
+}
 }
 
 SIZE_T FAssetDataGatherer::GetAllocatedSize() const
@@ -4431,7 +4575,7 @@ SIZE_T FAssetDataGatherer::GetAllocatedSize() const
 	Result += BlockedResults.GetAllocatedSize();
 	Result += SearchTimes.GetAllocatedSize();
 	Result += GetArrayRecursiveAllocatedSize(DiscoveredPaths);
-	Result += GPreloadSettings.GetMonolithicCacheFilename().GetAllocatedSize();
+	Result += GPreloadSettings.GetMonolithicCacheBaseFilename().GetAllocatedSize();
 
 	Result += NewCachedAssetData.GetAllocatedSize();
 	for (const FDiskCachedAssetData* Value : NewCachedAssetData)
diff --git a/Engine/Source/Runtime/AssetRegistry/Private/AssetDataGatherer.h b/Engine/Source/Runtime/AssetRegistry/Private/AssetDataGatherer.h
--- a/Engine/Source/Runtime/AssetRegistry/Private/AssetDataGatherer.h
+++ b/Engine/Source/Runtime/AssetRegistry/Private/AssetDataGatherer.h
@@ -171,7 +171,7 @@ public:
 	/** Calculate the cache filename that should be used for the given list of package paths. */
 	FString GetCacheFilename(TConstArrayView<FString> CacheFilePackagePaths);
 	/** Attempt to read the cache file at the given LocalPath, and store all of its results in the in-memory cache. */
-	void LoadCacheFile(FStringView CacheFilename);
+	void LoadCacheFiles(TConstArrayView<FString> CacheFilename);
 	/** Return the memory used by the gatherer. Used for performance metrics. */
 	SIZE_T GetAllocatedSize() const;
 
@@ -279,25 +279,35 @@ private:
 		TArray<FString>& CookedPackagesToLoadUponDiscovery, bool& OutCanRetry) const;
 
 	/** Add the given AssetDatas into DiskCachedAssetDataMap and DiskCachedAssetBlocks. */
-	void ConsumeCacheFile(UE::AssetDataGather::Private::FCachePayload&& Payload);
+	void ConsumeCacheFiles(TArray<UE::AssetDataGather::Private::FCachePayload> Payloads);
 	/**
 	 * If a save of the monolithic cache has been triggered, get the cache filename and pointers to all elements that
 	 * should be saved, for later saving outside of the critical section.
 	 */
 	void TryReserveSaveMonolithicCache(bool& bOutShouldSave, TArray<TPair<FName,FDiskCachedAssetData*>>& AssetsToSave);
+	/** 
+	 * Save a monolithic cache for the main asset discovery process, possibly sharded into multiple files.
+	*/
+	void SaveMonolithicCacheFile(const TArray<TPair<FName,FDiskCachedAssetData*>>& AssetsToSave);
 	/**
 	 * If the CacheFilename/AssetsToSave are non empty, save the cache file. 
 	 * This function reads the read-only-after-creation data from each FDiskCachedAssetData*, but otherwise does not use
 	 * data from this Gatherer and so can be run outside any critical section.
+	 * Returns the size of the saved file, or 0 if nothing was saved for any reason.
 	 */
-	void SaveCacheFileInternal(const FString& CacheFilename,
-		const TArray<TPair<FName,FDiskCachedAssetData*>>& AssetsToSave, bool bIsAsyncCacheSave, bool bIsMonolithicCache);
+	int64 SaveCacheFileInternal(const FString& CacheFilename,
+		const TArray<TPair<FName,FDiskCachedAssetData*>>& AssetsToSave);
 	/**
 	 * Get the list of FDiskCachedAssetData* that have been loaded in the gatherer, for saving into a cachefile.
 	 * Filters the list of assets by child paths of the elements in SaveCacheLongPackageNameDirs, if it is non-empty.
 	 */
 	void GetAssetsToSave(TArrayView<const FString> SaveCacheLongPackageNameDirs,
 		TArray<TPair<FName,FDiskCachedAssetData*>>& OutAssetsToSave);
+	/**
+	 * Get the list of FDiskCachedAssetData* for saving into the monolithic cache.
+	 * Includes both assets that were loaded in the gatherer and assets which were loaded from the monolithic cache and have not been pruned.
+	 */
+	void GetMonolithicCacheAssetsToSave(TArray<TPair<FName,FDiskCachedAssetData*>>& OutAssetsToSave);
 
 	/* Adds the given pair into NewCachedAssetDataMap. Detects collisions for multiple files with the same PackageName */
 	void AddToCache(FName PackageName, FDiskCachedAssetData* DiskCachedAssetData);
@@ -358,6 +368,8 @@ private:
 	/** True if dependency data should be gathered. Constant during threading. */
 	bool bGatherDependsData;
 
+	/** Timestamp of the start of the gather for consistent marking of 'last discovered' time in caching */
+	FDateTime GatherStartTime;
 
 	// Variable section for variables that are atomics read/writable from outside critical sections.
 
@@ -450,7 +462,7 @@ private:
 	TArray<TPair<int32, FDiskCachedAssetData*>> DiskCachedAssetBlocks;
 	/**
 	 * Map of PackageName to cached discovered assets that were loaded from disk.
-	 * This should only be modified by ConsumeCacheFile.
+	 * This should only be modified by ConsumeCacheFiles.
 	 */
 	TMap<FName, FDiskCachedAssetData*> DiskCachedAssetDataMap;
 	/** Map of PackageName to cached discovered assets that will be written to disk at shutdown. */
@@ -461,6 +473,8 @@ private:
 	int32 NumCachedAssetFiles;
 	/** How many files in the search results were not in the cache and were read by parsing the file. */
 	int32 NumUncachedAssetFiles;
+	/** How many uncached asset files had been discovered at the last async cache save */
+	int32 LastMonolithicCacheSaveUncachedAssetFiles;
 	/**
 	 * Incremented when a thread is in the middle of saving any cache and therefore the cache cannot be deleted,
 	 * decremented when the thread is done. Only incremented when bCacheEnabled has been recently confirmed to be true.
diff --git a/Engine/Source/Runtime/AssetRegistry/Private/AssetRegistry.cpp b/Engine/Source/Runtime/AssetRegistry/Private/AssetRegistry.cpp
--- a/Engine/Source/Runtime/AssetRegistry/Private/AssetRegistry.cpp
+++ b/Engine/Source/Runtime/AssetRegistry/Private/AssetRegistry.cpp
@@ -30,6 +30,7 @@
 #include "Misc/ScopeRWLock.h"
 #include "Misc/TrackedActivity.h"
 #include "AssetRegistry/PackageReader.h"
+#include "ProfilingDebugging/MiscTrace.h"
 #include "Serialization/ArrayReader.h"
 #include "Serialization/CompactBinarySerialization.h"
 #include "Serialization/CompactBinaryWriter.h"
@@ -1591,6 +1592,7 @@ namespace UE::AssetRegistry
 void FAssetRegistryImpl::SearchAllAssets(Impl::FEventContext& EventContext,
 	Impl::FClassInheritanceContext& InheritanceContext, bool bSynchronousSearch)
 {
+	TRACE_BEGIN_REGION(TEXT("Asset Registry Scan"));
 	ConstructGatherer();
 	FAssetDataGatherer& Gatherer = *GlobalGatherer;
 	if (!Gatherer.IsAsyncEnabled())
@@ -3799,6 +3801,7 @@ bool FAssetRegistryImpl::IsLoadingAssets() const
 void UAssetRegistryImpl::Tick(float DeltaTime)
 {
 	checkf(IsInGameThread(), TEXT("The tick function executes deferred loads and events and must be on the game thread to do so."));
+	TRACE_CPUPROFILER_EVENT_SCOPE_STR("Asset Registry Tick");
 
 	UE::AssetRegistry::Impl::EGatherStatus Status = UE::AssetRegistry::Impl::EGatherStatus::Active;
 	double TickStartTime = -1; // Force a full flush if DeltaTime < 0
@@ -4024,6 +4027,7 @@ Impl::EGatherStatus FAssetRegistryImpl::TickGatherer(Impl::FEventContext& EventC
 #endif
 			RecordTimer();
 			LogSearchDiagnostics();
+			TRACE_END_REGION(TEXT("Asset Registry Scan"));
 
 			bInitialSearchCompleted = true;
 
@@ -4662,7 +4666,7 @@ void FAssetRegistryImpl::ScanPathsSynchronous(Impl::FScanPathContext& Context)
 	if (!CacheFilePackagePaths.IsEmpty())
 	{
 		CacheFilename = Gatherer.GetCacheFilename(CacheFilePackagePaths);
-		Gatherer.LoadCacheFile(CacheFilename);
+		Gatherer.LoadCacheFiles({CacheFilename});
 	}
 
 	Gatherer.ScanPathsSynchronous(Context.LocalPaths, Context.bForceRescan, Context.bIgnoreDenyListScanFilters, CacheFilename, Context.PackageDirs);
@@ -4809,6 +4813,7 @@ void FAssetRegistryImpl::PostLoadAssetRegistryTags(FAssetData* AssetData)
 
 void FAssetRegistryImpl::AssetSearchDataGathered(Impl::FEventContext& EventContext, const double TickStartTime, TMultiMap<FName, FAssetData*>& AssetResults)
 {
+	TRACE_CPUPROFILER_EVENT_SCOPE(AssetSearchDataGathered);
 	const bool bFlushFullBuffer = TickStartTime < 0;
 
 	// Refreshes ClassGeneratorNames if out of date due to module load
@@ -4906,6 +4911,7 @@ void FAssetRegistryImpl::AssetSearchDataGathered(Impl::FEventContext& EventConte
 
 void FAssetRegistryImpl::PathDataGathered(Impl::FEventContext& EventContext, const double TickStartTime, TRingBuffer<FString>& PathResults)
 {
+	TRACE_CPUPROFILER_EVENT_SCOPE(PathDataGathered);
 	const bool bFlushFullBuffer = TickStartTime < 0;
 
 	TSet<FString> MountPoints;
@@ -4941,6 +4947,7 @@ void FAssetRegistryImpl::PathDataGathered(Impl::FEventContext& EventContext, con
 
 void FAssetRegistryImpl::DependencyDataGathered(const double TickStartTime, TMultiMap<FName, FPackageDependencyData>& DependsResults)
 {
+	TRACE_CPUPROFILER_EVENT_SCOPE(DependencyDataGathered);
 	using namespace UE::AssetRegistry;
 	const bool bFlushFullBuffer = TickStartTime < 0;
 
@@ -5067,6 +5074,7 @@ void FAssetRegistryImpl::DependencyDataGathered(const double TickStartTime, TMul
 void FAssetRegistryImpl::CookedPackageNamesWithoutAssetDataGathered(Impl::FEventContext& EventContext,
 	const double TickStartTime, TRingBuffer<FString>& CookedPackageNamesWithoutAssetDataResults, bool& bOutInterrupted)
 {
+	TRACE_CPUPROFILER_EVENT_SCOPE(CookedPackageNamesWithoutAssetDataGathered);
 	bOutInterrupted = false;
 
 	struct FConfigValue
@@ -5110,6 +5118,7 @@ void FAssetRegistryImpl::CookedPackageNamesWithoutAssetDataGathered(Impl::FEvent
 
 void FAssetRegistryImpl::VerseFilesGathered(Impl::FEventContext& EventContext, const double TickStartTime, TRingBuffer<FName>& VerseResults)
 {
+	TRACE_CPUPROFILER_EVENT_SCOPE(VerseFilesGathered);
 	while (VerseResults.Num() > 0)
 	{
 		FName VerseFilePath = VerseResults.PopFrontValue();
diff --git a/Engine/Source/Runtime/AssetRegistry/Private/DiskCachedAssetData.h b/Engine/Source/Runtime/AssetRegistry/Private/DiskCachedAssetData.h
--- a/Engine/Source/Runtime/AssetRegistry/Private/DiskCachedAssetData.h
+++ b/Engine/Source/Runtime/AssetRegistry/Private/DiskCachedAssetData.h
@@ -8,7 +8,8 @@
 class FDiskCachedAssetData
 {
 public:
-	FDateTime Timestamp;
+	FDateTime ModificationTime; // Modification time of the asset last time it was discovered
+	FDateTime LastGatheredTime;  // The last time this asset was discovered during a scan 
 	FName Extension;
 	TArray<FAssetData> AssetDataList;
 	FPackageDependencyData DependencyData;
@@ -16,19 +17,21 @@ public:
 	FDiskCachedAssetData()
 	{}
 
-	FDiskCachedAssetData(const FDateTime& InTimestamp, FName InExtension)
-		: Timestamp(InTimestamp)
+	FDiskCachedAssetData(const FDateTime& InModificationTime, const FDateTime& InLastSeen, FName InExtension)
+		: ModificationTime(InModificationTime)
+		, LastGatheredTime(InLastSeen)
 		, Extension(InExtension)
 	{}
 
 	/**
 	 * Serialize as part of the registry cache. This is not meant to be serialized as part of a package so  it does not handle versions normally
-	 * To version this data change FAssetRegistryVersion or CacheSerializationVersion
+	 * To version this data change FAssetRegistryVersion or AssetDataGathererConstants::CacheSerializationMagic
 	 */
 	template<class Archive>
 	void SerializeForCache(Archive&& Ar)
 	{
-		Ar << Timestamp;
+		Ar << ModificationTime;
+		Ar << LastGatheredTime;
 		Ar << Extension;
 	
 		int32 AssetDataCount = AssetDataList.Num();
