AssetRegistry preloader: fix crash when exiting process early and Preloader is still active; we can not access global systems after main has exited and we are in c++ global variable destructor.
#rnx

diff --git a/Engine/Source/Runtime/AssetRegistry/Private/AssetRegistry.cpp b/Engine/Source/Runtime/AssetRegistry/Private/AssetRegistry.cpp
--- a/Engine/Source/Runtime/AssetRegistry/Private/AssetRegistry.cpp
+++ b/Engine/Source/Runtime/AssetRegistry/Private/AssetRegistry.cpp
@@ -223,7 +223,7 @@ public:
 	{
 		// We are destructed after Main exits, which means that our AsyncThread was either never called
 		// or it was waited on to complete by TaskGraph. Therefore we do not need to handle waiting for it ourselves.
-		Shutdown();
+		Shutdown(true /* bFromGlobalDestructor */);
 	}
 
 	/**
@@ -448,12 +448,17 @@ private:
 	}
 
 	/** Called when the Preloader has no further work to do, to free resources early since destruction occurs at end of process. */
-	void Shutdown()
+	void Shutdown(bool bFromGlobalDestructor = false)
 	{
 		OnTaskGraphReady.Reset();
 		if (PreloadReady)
 		{
-			FPlatformProcess::ReturnSynchEventToPool(PreloadReady);
+			// If we are exiting the process early while PreloadReady is still allocated, the event
+			// system has already been torn down and there is nothing for us to free for PreloadReady.
+			if (!bFromGlobalDestructor)
+			{
+				FPlatformProcess::ReturnSynchEventToPool(PreloadReady);
+			}
 			PreloadReady = nullptr;
 		}
 		ARPath.Reset();
