AssetRegistry and ExternalActors: Ignore ExternalActors that have a packagename that does not match their objectpath. These actors were created by an editor bug during saveas, and they can hide the true actors that have that objectpath.
#rnx

diff --git a/Engine/Source/Runtime/AssetRegistry/Private/AssetRegistry.cpp b/Engine/Source/Runtime/AssetRegistry/Private/AssetRegistry.cpp
--- a/Engine/Source/Runtime/AssetRegistry/Private/AssetRegistry.cpp
+++ b/Engine/Source/Runtime/AssetRegistry/Private/AssetRegistry.cpp
@@ -4816,6 +4816,101 @@ void FAssetRegistryImpl::PostLoadAssetRegistryTags(FAssetData* AssetData)
 }
 #endif
 
+/** "/PackageRoot/Path/Leaf" -> ("/PackageRoot", "Path/Leaf") */
+static void SplitPackageNamePackageRoot(FStringView PackageName, FStringView& OutRootPath, FStringView& OutRelPath)
+{
+	if (!PackageName.StartsWith(TEXT("/")))
+	{
+		OutRootPath = FStringView();
+		OutRelPath = PackageName;
+		return;
+	}
+
+	int32 FirstSlashLen = 1;
+	int32 SecondSlashIndex;
+	PackageName.RightChop(FirstSlashLen).FindChar('/', SecondSlashIndex);
+	if (SecondSlashIndex == INDEX_NONE)
+	{
+		OutRootPath = PackageName;
+		OutRelPath = FStringView();
+	}
+	SecondSlashIndex += FirstSlashLen;
+
+	OutRootPath = PackageName.Left(SecondSlashIndex);
+	OutRelPath = PackageName.RightChop(SecondSlashIndex + 1);
+};
+
+bool FAssetRegistryImpl::ShouldSkipGatheredAsset(FAssetData& AssetData)
+{
+	// TODO: This pruning of invalid ExternalActors is temporary, to handle the fallout from a bug in SaveAs
+	// that is keeping the old ExternalActors as duplicates of the new ones. Remove it after the data has been
+	// cleaned up for all affected licensees. If we need such validation permanently, it should be decoupled
+	// from the AssetRegistry by adding a delegate.
+	// Extra validation for ExternalActors. If duplicate ExternalActors with the same object path exist
+	// then we intermittently will fail to find the correct one and WorldPartition will break.
+	// Validate that the PackageName matches what is expected from the ObjectPath.
+	FStringView ExternalActorsFolderName(FPackagePath::GetExternalActorsFolderName());
+	TStringBuilder<256> PackageNameStr;
+	AssetData.PackageName.ToString(PackageNameStr);
+	if (UE::String::FindFirst(PackageNameStr, ExternalActorsFolderName) != INDEX_NONE)
+	{
+		TStringBuilder<256> ObjectPathString;
+		AssetData.AppendObjectPath(ObjectPathString);
+		FStringView ObjectPathPackageName = FPackageName::ObjectPathToPackageName(ObjectPathString);
+		FStringView PackageNamePackageRoot;
+		FStringView PackageNameRelPath;
+		FStringView ObjectPathPackageRoot;
+		FStringView ObjectPathRelPath;
+
+		// /PackageRoot/__ExternalActors__/RelPathFromPackageRootToMap/#/##/#######
+		// OR
+		// /PackageRoot/__ExternalActors__/ContentBundle/######/RelPathFromPackageRootToMap/#/##/#######
+		// Package roots do not need to be the same; ContentBundles can be injected into /Game maps from plugins
+		SplitPackageNamePackageRoot(PackageNameStr, PackageNamePackageRoot, PackageNameRelPath);
+		SplitPackageNamePackageRoot(ObjectPathPackageName, ObjectPathPackageRoot, ObjectPathRelPath);
+
+		if (!PackageNameRelPath.StartsWith(ExternalActorsFolderName) || !PackageNameRelPath.RightChop(ExternalActorsFolderName.Len()).StartsWith(TEXT("/")))
+		{
+			UE_LOG(LogAssetRegistry, Verbose,
+				TEXT("Invalid ExternalActor: Package %s is an ExternalActor package but is not in the expected root path for ExternalActors %.*s/%.*s. Ignoring this actor."),
+				*PackageNameStr, PackageNamePackageRoot.Len(), PackageNamePackageRoot.GetData(),
+				ExternalActorsFolderName.Len(), ExternalActorsFolderName.GetData());
+			return true;
+		}
+
+		bool bAllowValidation = true;
+		FStringView PackageNameRelPathAfterExternalActorRoot = PackageNameRelPath.RightChop(ExternalActorsFolderName.Len() + 1);
+		FStringView ContentBundleDirName(TEXTVIEW("ContentBundle"));
+		if (PackageNameRelPathAfterExternalActorRoot.StartsWith(ContentBundleDirName))
+		{
+			bAllowValidation = false; // Don't allow validation unless we succeed in finding the new relpath
+			PackageNameRelPathAfterExternalActorRoot.RightChopInline(ContentBundleDirName.Len());
+			if (PackageNameRelPathAfterExternalActorRoot.StartsWith(TEXT("/")))
+			{
+				PackageNameRelPathAfterExternalActorRoot.RightChopInline(1);
+				int32 NextSlash;
+				PackageNameRelPathAfterExternalActorRoot.FindChar('/', NextSlash);
+				if (NextSlash != INDEX_NONE)
+				{
+					PackageNameRelPathAfterExternalActorRoot.RightChopInline(NextSlash + 1);
+					bAllowValidation = true;
+				}
+			}
+		}
+
+		if (bAllowValidation && !PackageNameRelPathAfterExternalActorRoot.StartsWith(ObjectPathRelPath))
+		{
+			TStringBuilder<256> ExpectedPath;
+			ExpectedPath << ObjectPathPackageRoot << "/" << ExternalActorsFolderName << "/" << ObjectPathRelPath;
+			UE_LOG(LogAssetRegistry, Verbose,
+				TEXT("Invalid ExternalActor: Package %s is an ExternalActor package but its path does not match the expected path %s created from its objectpath %s. Ignoring this actor."),
+				*PackageNameStr, *ExpectedPath, *ObjectPathString);
+			return true;
+		}
+	}
+	return false;
+}
+
 void FAssetRegistryImpl::AssetSearchDataGathered(Impl::FEventContext& EventContext, const double TickStartTime, TMultiMap<FName, FAssetData*>& AssetResults)
 {
 	TRACE_CPUPROFILER_EVENT_SCOPE(AssetSearchDataGathered);
@@ -4844,6 +4939,12 @@ void FAssetRegistryImpl::AssetSearchDataGathered(Impl::FEventContext& EventConte
 
 		// Try to update any asset data that may already exist
 		FCachedAssetKey Key(*BackgroundResult);
+
+		if (ShouldSkipGatheredAsset(*BackgroundResult))
+		{
+			continue;
+		}
+
 		FAssetData* const* FoundData = State.CachedAssets.Find(Key);
 		FAssetData* ExistingAssetData = FoundData ? *FoundData : nullptr;
 
diff --git a/Engine/Source/Runtime/AssetRegistry/Private/AssetRegistryImpl.h b/Engine/Source/Runtime/AssetRegistry/Private/AssetRegistryImpl.h
--- a/Engine/Source/Runtime/AssetRegistry/Private/AssetRegistryImpl.h
+++ b/Engine/Source/Runtime/AssetRegistry/Private/AssetRegistryImpl.h
@@ -282,6 +282,8 @@ private:
 	 * If TickStartTime is < 0, the entire list of gathered assets will be cached. Also used in sychronous searches
 	 */
 	void AssetSearchDataGathered(Impl::FEventContext& EventContext, const double TickStartTime, TMultiMap<FName, FAssetData*>& AssetResults);
+	/** Validate assets gathered from disk before adding them to the AssetRegistry. */
+	bool ShouldSkipGatheredAsset(FAssetData& AssetData);
 
 	/**
 	 * Called every tick when data is retrieved by the background path search.
diff --git a/Engine/Source/Runtime/CoreUObject/Private/AssetRegistry/AssetData.cpp b/Engine/Source/Runtime/CoreUObject/Private/AssetRegistry/AssetData.cpp
--- a/Engine/Source/Runtime/CoreUObject/Private/AssetRegistry/AssetData.cpp
+++ b/Engine/Source/Runtime/CoreUObject/Private/AssetRegistry/AssetData.cpp
@@ -802,6 +802,12 @@ void FAssetData::SerializeForCacheOldVersionWithTagsAndBundles(FArchive& Ar, FAs
 	SerializeForCacheInternal(Ar, Version, SerializeTagsAndBundles);
 }
 
+bool FAssetData::IsRedirectorClassName(FTopLevelAssetPath ClassPathName)
+{
+	static const FTopLevelAssetPath ObjectRedirectorClassPathName = UObjectRedirector::StaticClass()->GetClassPathName();
+	return ClassPathName == ObjectRedirectorClassPathName;
+}
+
 FTopLevelAssetPath FAssetData::TryConvertShortClassNameToPathName(FName InClassName, ELogVerbosity::Type FailureMessageVerbosity /*= ELogVerbosity::Warning*/)
 {
 	FTopLevelAssetPath ClassPath;
diff --git a/Engine/Source/Runtime/CoreUObject/Public/AssetRegistry/AssetData.h b/Engine/Source/Runtime/CoreUObject/Public/AssetRegistry/AssetData.h
--- a/Engine/Source/Runtime/CoreUObject/Public/AssetRegistry/AssetData.h
+++ b/Engine/Source/Runtime/CoreUObject/Public/AssetRegistry/AssetData.h
@@ -726,13 +726,9 @@ private:
 		return PackageBaseName.Equals(ObjectPathName, ESearchCase::IgnoreCase);
 	}
 
-	static bool IsRedirectorClassName(FTopLevelAssetPath ClassPathName)
-	{
-		static const FTopLevelAssetPath ObjectRedirectorClassPathName = UObjectRedirector::StaticClass()->GetClassPathName();
-		return ClassPathName == ObjectRedirectorClassPathName;
-	}
-
 public:
+	COREUOBJECT_API static bool IsRedirectorClassName(FTopLevelAssetPath ClassPathName);
+
 	/** Helper function that tries to convert short class name to path name */
 	COREUOBJECT_API static FTopLevelAssetPath TryConvertShortClassNameToPathName(FName InClassName, ELogVerbosity::Type FailureMessageVerbosity = ELogVerbosity::Warning);
 
