fix non-unity build

diff --git a/Engine/Source/Runtime/AssetRegistry/Private/AssetRegistry.cpp b/Engine/Source/Runtime/AssetRegistry/Private/AssetRegistry.cpp
--- a/Engine/Source/Runtime/AssetRegistry/Private/AssetRegistry.cpp
+++ b/Engine/Source/Runtime/AssetRegistry/Private/AssetRegistry.cpp
@@ -40,6 +40,7 @@
 #include "UObject/MetaData.h"
 #include "UObject/UObjectHash.h"
 #include "UObject/UObjectIterator.h"
+#include "UObject/UObjectThreadContext.h"
 
 #include UE_INLINE_GENERATED_CPP_BY_NAME(IAssetRegistry)
 #include UE_INLINE_GENERATED_CPP_BY_NAME(AssetRegistry)
