Faster sorting of entries when saving asset registry

diff --git a/Engine/Source/Runtime/AssetRegistry/Private/AssetRegistryState.cpp b/Engine/Source/Runtime/AssetRegistry/Private/AssetRegistryState.cpp
--- a/Engine/Source/Runtime/AssetRegistry/Private/AssetRegistryState.cpp
+++ b/Engine/Source/Runtime/AssetRegistry/Private/AssetRegistryState.cpp
@@ -1159,12 +1159,22 @@ bool FAssetRegistryState::Save(FArchive& OriginalAr, const FAssetRegistrySeriali
 	Ar << AssetCount;
 
 	// Write asset data first
-	TArray<FAssetData*> SortedAssetsByObjectPath = CachedAssets.Array();
-	Algo::Sort(SortedAssetsByObjectPath, [](const FAssetData* A, const FAssetData* B) { return A->GetSoftObjectPath().LexicalLess(B->GetSoftObjectPath()); });
-	for (FAssetData* AssetData : SortedAssetsByObjectPath)
 	{
-		// Hardcoding FAssetRegistryVersion::LatestVersion here so that branches can get optimized out in the forceinlined SerializeForCache
-		AssetData->SerializeForCache(Ar);
+		TArray<TPair<FAssetData*, FSoftObjectPath>> SortedAssetsByObjectPath;
+		SortedAssetsByObjectPath.Reserve(AssetCount);
+		for (FAssetData* AssetData : CachedAssets)
+		{
+			SortedAssetsByObjectPath.Add({ AssetData, AssetData->GetSoftObjectPath() });
+		}
+		Algo::Sort(SortedAssetsByObjectPath, [](const TPair<FAssetData*, FSoftObjectPath>& A, const TPair<FAssetData*, FSoftObjectPath>& B) {
+			return A.Value.LexicalLess(B.Value);
+		});
+
+		for (TPair<FAssetData*, FSoftObjectPath>& Asset : SortedAssetsByObjectPath)
+		{
+			// Hardcoding FAssetRegistryVersion::LatestVersion here so that branches can get optimized out in the forceinlined SerializeForCache
+			Asset.Key->SerializeForCache(Ar);
+		}
 	}
 
 	// Serialize Dependencies
