Add asset registry events which takes array views of updated assets rather than individual assets.

diff --git a/Engine/Source/Runtime/AssetRegistry/Private/AssetRegistry.cpp b/Engine/Source/Runtime/AssetRegistry/Private/AssetRegistry.cpp
--- a/Engine/Source/Runtime/AssetRegistry/Private/AssetRegistry.cpp
+++ b/Engine/Source/Runtime/AssetRegistry/Private/AssetRegistry.cpp
@@ -1480,6 +1480,10 @@ void UAssetRegistryImpl::FinishDestroy()
 		AssetRenamedEvent.Clear();
 		AssetUpdatedEvent.Clear();
 		AssetUpdatedOnDiskEvent.Clear();
+		for (FAssetsEvent& Event : BatchedAssetEvents)
+		{
+			Event.Clear();
+		}
 		InMemoryAssetCreatedEvent.Clear();
 		InMemoryAssetDeletedEvent.Clear();
 		FileLoadedEvent.Clear();
@@ -3589,7 +3593,9 @@ void UAssetRegistryImpl::AssetCreated(UObject* NewAsset)
 		{
 			checkf(IsInGameThread(), TEXT("AssetCreated is not yet implemented as callable from other threads"));
 			// Let subscribers know that the new asset was added to the registry
-			AssetAddedEvent.Broadcast(FAssetData(NewAsset, true /* bAllowBlueprintClass */));
+			FAssetData AssetData = FAssetData(NewAsset, true /* bAllowBlueprintClass */); 
+			AssetAddedEvent.Broadcast(AssetData);
+			OnAssetsAdded().Broadcast({ AssetData });
 
 			// Notify listeners that an asset was just created
 			InMemoryAssetCreatedEvent.Broadcast(NewAsset);
@@ -3645,6 +3651,7 @@ void UAssetRegistryImpl::AssetDeleted(UObject* DeletedAsset)
 			checkf(IsInGameThread(), TEXT("AssetDeleted is not yet implemented as callable from other threads"));
 			// Let subscribers know that the asset was removed from the registry
 			AssetRemovedEvent.Broadcast(AssetDataDeleted);
+			OnAssetsRemoved().Broadcast({AssetDataDeleted});
 
 			// Notify listeners that an in-memory asset was just deleted
 			InMemoryAssetDeletedEvent.Broadcast(DeletedAsset);
@@ -4303,12 +4310,15 @@ void UAssetRegistryImpl::AppendState(const FAssetRegistryState& InState)
 
 	Broadcast(EventContext);
 	checkf(IsInGameThread(), TEXT("AppendState is not yet implemented as callable from other threads"));
-	InState.EnumerateAllAssets(TSet<FName>(), [this](const FAssetData& AssetData)
+
+	TArray<FAssetData> TempAssets;
+	InState.GetAllAssets(TSet<FName>(), TempAssets, true /* bARFiltering */);
+	for (const FAssetData& AssetData : TempAssets)
 	{
 		// Let subscribers know that the new asset was added to the registry
 		AssetAddedEvent.Broadcast(AssetData);
-		return true;
-	}, true /*bARFiltering*/);
+	}
+	OnAssetsAdded().Broadcast(TempAssets);
 }
 
 namespace UE::AssetRegistry
@@ -7090,6 +7100,40 @@ void UAssetRegistryImpl::Broadcast(UE::AssetRegistry::Impl::FEventContext& Event
 
 	if (EventContext.AssetEvents.Num())
 	{
+		// Batch events so that if adds/updates are interspersed with removes, relative ordering of the add/remove is maintained 
+		constexpr uint32 EventTypeCount = static_cast<uint32>(FEventContext::EEvent::MAX);
+		static_assert(EventTypeCount == 4, "Loop needs to be rewritten to correctly order new event types");
+		TArray<FAssetData> EventBatches[EventTypeCount];
+		FEventContext::EEvent LastEvent = EventContext.AssetEvents[0].Get<1>();
+		auto FlushBatchedEvents = [&EventBatches, &Events=BatchedAssetEvents]() 
+		{
+			for (int32 i=0; i < EventTypeCount; ++i)
+			{
+				if (EventBatches[i].Num())
+				{
+					Events[i].Broadcast(EventBatches[i]);
+					EventBatches[i].Reset();
+				}
+			}
+		};
+
+		for (const TPair<FAssetData, FEventContext::EEvent>& AssetEvent : EventContext.AssetEvents)
+		{
+			const FAssetData& AssetData = AssetEvent.Get<0>();
+			FEventContext::EEvent Event = AssetEvent.Get<1>();
+
+			// Flush events when switching between removed and non-removed events 
+			if ((Event == FEventContext::EEvent::Removed) != (LastEvent == FEventContext::EEvent::Removed))
+			{
+				FlushBatchedEvents();
+			}
+			EventBatches[static_cast<int32>(Event)].Add(AssetData);
+			LastEvent = Event;
+		}
+		// Flush last batch of events 
+		FlushBatchedEvents();
+		
+		// Single events
 		for (const TPair<FAssetData, FEventContext::EEvent>& AssetEvent : EventContext.AssetEvents)
 		{
 			const FAssetData& AssetData = AssetEvent.Get<0>();
@@ -7195,6 +7239,26 @@ UAssetRegistryImpl::FAssetUpdatedEvent& UAssetRegistryImpl::OnAssetUpdatedOnDisk
 	return AssetUpdatedOnDiskEvent;
 }
 
+UAssetRegistryImpl::FAssetsEvent& UAssetRegistryImpl::OnAssetsAdded()
+{
+	return BatchedAssetEvents[static_cast<SIZE_T>(UE::AssetRegistry::Impl::FEventContext::EEvent::Added)];
+}
+
+UAssetRegistryImpl::FAssetsEvent& UAssetRegistryImpl::OnAssetsUpdated()
+{
+	return BatchedAssetEvents[static_cast<SIZE_T>(UE::AssetRegistry::Impl::FEventContext::EEvent::Updated)];
+}
+
+UAssetRegistryImpl::FAssetsEvent& UAssetRegistryImpl::OnAssetsUpdatedOnDisk()
+{
+	return BatchedAssetEvents[static_cast<SIZE_T>(UE::AssetRegistry::Impl::FEventContext::EEvent::UpdatedOnDisk)];
+}
+
+UAssetRegistryImpl::FAssetsEvent& UAssetRegistryImpl::OnAssetsRemoved()
+{
+	return BatchedAssetEvents[static_cast<SIZE_T>(UE::AssetRegistry::Impl::FEventContext::EEvent::Removed)];
+}
+
 UAssetRegistryImpl::FInMemoryAssetCreatedEvent& UAssetRegistryImpl::OnInMemoryAssetCreated()
 {
 	return InMemoryAssetCreatedEvent;
diff --git a/Engine/Source/Runtime/AssetRegistry/Private/AssetRegistry.h b/Engine/Source/Runtime/AssetRegistry/Private/AssetRegistry.h
--- a/Engine/Source/Runtime/AssetRegistry/Private/AssetRegistry.h
+++ b/Engine/Source/Runtime/AssetRegistry/Private/AssetRegistry.h
@@ -148,7 +148,7 @@ public:
 
 	DECLARE_DERIVED_EVENT( UAssetRegistryImpl, IAssetRegistry::FAssetRemovedEvent, FAssetRemovedEvent);
 	virtual FAssetRemovedEvent& OnAssetRemoved() override;
-
+	
 	DECLARE_DERIVED_EVENT( UAssetRegistryImpl, IAssetRegistry::FAssetRenamedEvent, FAssetRenamedEvent);
 	virtual FAssetRenamedEvent& OnAssetRenamed() override;
 
@@ -156,6 +156,13 @@ public:
 	virtual FAssetUpdatedEvent& OnAssetUpdated() override;
 	virtual FAssetUpdatedEvent& OnAssetUpdatedOnDisk() override;
 
+	// Batch events
+	DECLARE_DERIVED_EVENT( UAssetRegistryImpl, IAssetRegistry::FAssetsEvent, FAssetsEvent);
+	virtual FAssetsEvent& OnAssetsAdded() override;
+	virtual FAssetsEvent& OnAssetsRemoved() override;
+	virtual FAssetsEvent& OnAssetsUpdated() override;
+	virtual FAssetsEvent& OnAssetsUpdatedOnDisk() override;
+	
 	DECLARE_DERIVED_EVENT( UAssetRegistryImpl, IAssetRegistry::FInMemoryAssetCreatedEvent, FInMemoryAssetCreatedEvent );
 	virtual FInMemoryAssetCreatedEvent& OnInMemoryAssetCreated() override;
 
@@ -297,7 +304,7 @@ private:
 
 	/** The delegate to execute when an asset is added to the registry */
 	FAssetAddedEvent AssetAddedEvent;
-
+	
 	/** The delegate to execute when an asset is removed from the registry */
 	FAssetRemovedEvent AssetRemovedEvent;
 
@@ -310,6 +317,10 @@ private:
 	/** The delegate to execute when an asset is updated on disk and has been reloaded in assetregistry */
 	FAssetUpdatedEvent AssetUpdatedOnDiskEvent;
 
+	/** The delegates to execute when assets are added/removed/updated in the registry, indexed by FEventContext::EEvent or returned with public accessors */
+	static constexpr SIZE_T NumBatchedEvents = static_cast<SIZE_T>(UE::AssetRegistry::Impl::FEventContext::EEvent::MAX);
+	FAssetsEvent BatchedAssetEvents[NumBatchedEvents];
+
 	/** The delegate to execute when an in-memory asset was just created */
 	FInMemoryAssetCreatedEvent InMemoryAssetCreatedEvent;
 
diff --git a/Engine/Source/Runtime/AssetRegistry/Private/AssetRegistryImpl.h b/Engine/Source/Runtime/AssetRegistry/Private/AssetRegistryImpl.h
--- a/Engine/Source/Runtime/AssetRegistry/Private/AssetRegistryImpl.h
+++ b/Engine/Source/Runtime/AssetRegistry/Private/AssetRegistryImpl.h
@@ -475,12 +475,14 @@ namespace Impl
  */
 struct FEventContext
 {
-	enum class EEvent
+	enum class EEvent : uint32
 	{
 		Added,
 		Removed,
 		Updated,
 		UpdatedOnDisk,
+		
+		MAX
 	};
 	TOptional<IAssetRegistry::FFileLoadProgressUpdateData> ProgressUpdateData;
 	TArray<TPair<FString, EEvent>> PathEvents;
@@ -498,6 +500,7 @@ struct FEventContext
 	void Append(FEventContext&& Other);
 };
 
+
 /*
  * An accessor for the inheritance map and related data for functions that need them; the source
  * of this data is an internal detail unneeded by the functions reading it. Readers of the data call
diff --git a/Engine/Source/Runtime/AssetRegistry/Public/AssetRegistry/IAssetRegistry.h b/Engine/Source/Runtime/AssetRegistry/Public/AssetRegistry/IAssetRegistry.h
--- a/Engine/Source/Runtime/AssetRegistry/Public/AssetRegistry/IAssetRegistry.h
+++ b/Engine/Source/Runtime/AssetRegistry/Public/AssetRegistry/IAssetRegistry.h
@@ -640,14 +640,22 @@ public:
 	/** Informs the asset registry that an Asset has finalized its tags after loading. Ignored if the Asset's package has been modified. */
 	virtual void AssetTagsFinalized(const UObject& FinalizedAsset) = 0;
 
+	DECLARE_TS_MULTICAST_DELEGATE_OneParam( FAssetsEvent, TConstArrayView<FAssetData> );
+
 	/** Event for when assets are added to the registry */
 	DECLARE_TS_MULTICAST_DELEGATE_OneParam( FAssetAddedEvent, const FAssetData& );
 	virtual FAssetAddedEvent& OnAssetAdded() = 0;
 
+	/** Event for when assets are added to the registry to handle batches more efficiently */
+	virtual FAssetsEvent& OnAssetsAdded() = 0;
+
 	/** Event for when assets are removed from the registry */
 	DECLARE_TS_MULTICAST_DELEGATE_OneParam( FAssetRemovedEvent, const FAssetData& );
 	virtual FAssetRemovedEvent& OnAssetRemoved() = 0;
 
+	/** Event for when assets are removed from the registry */
+	virtual FAssetsEvent& OnAssetsRemoved() = 0;
+
 	/** Event for when assets are renamed in the registry */
 	DECLARE_TS_MULTICAST_DELEGATE_TwoParams( FAssetRenamedEvent, const FAssetData&, const FString& );
 	virtual FAssetRenamedEvent& OnAssetRenamed() = 0;
@@ -656,9 +664,15 @@ public:
 	DECLARE_TS_MULTICAST_DELEGATE_OneParam( FAssetUpdatedEvent, const FAssetData&);
 	virtual FAssetUpdatedEvent& OnAssetUpdated() = 0;
 
+	/** Event for when assets are updated in the registry */
+	virtual FAssetsEvent& OnAssetsUpdated() = 0;
+
 	/** Event for when assets are updated on disk and have been refreshed in the assetregistry */
 	virtual FAssetUpdatedEvent& OnAssetUpdatedOnDisk() = 0;
 
+	/** Event for when assets are updated on disk and have been refreshed in the assetregistry */
+	virtual FAssetsEvent& OnAssetsUpdatedOnDisk() = 0;
+
 	/** Event for when in-memory assets are created */
 	DECLARE_TS_MULTICAST_DELEGATE_OneParam( FInMemoryAssetCreatedEvent, UObject* );
 	virtual FInMemoryAssetCreatedEvent& OnInMemoryAssetCreated() = 0;
