Fix bad capture of constexpr variable

diff --git a/Engine/Source/Runtime/AssetRegistry/Private/AssetRegistry.cpp b/Engine/Source/Runtime/AssetRegistry/Private/AssetRegistry.cpp
--- a/Engine/Source/Runtime/AssetRegistry/Private/AssetRegistry.cpp
+++ b/Engine/Source/Runtime/AssetRegistry/Private/AssetRegistry.cpp
@@ -7107,7 +7107,7 @@ void UAssetRegistryImpl::Broadcast(UE::AssetRegistry::Impl::FEventContext& Event
 		FEventContext::EEvent LastEvent = EventContext.AssetEvents[0].Get<1>();
 		auto FlushBatchedEvents = [&EventBatches, &Events=BatchedAssetEvents]() 
 		{
-			for (int32 i=0; i < EventTypeCount; ++i)
+			for (int32 i=0; i < UE_ARRAY_COUNT(EventBatches); ++i)
 			{
 				if (EventBatches[i].Num())
 				{
