Adding FPackageName::SplitPackageNameRoot to access the plugin for a given package.
#rnx

diff --git a/Engine/Source/Runtime/AssetRegistry/Private/AssetRegistry.cpp b/Engine/Source/Runtime/AssetRegistry/Private/AssetRegistry.cpp
--- a/Engine/Source/Runtime/AssetRegistry/Private/AssetRegistry.cpp
+++ b/Engine/Source/Runtime/AssetRegistry/Private/AssetRegistry.cpp
@@ -4830,30 +4830,6 @@ void FAssetRegistryImpl::PostLoadAssetRegistryTags(FAssetData* AssetData)
 }
 #endif
 
-/** "/PackageRoot/Path/Leaf" -> ("/PackageRoot", "Path/Leaf") */
-static void SplitPackageNamePackageRoot(FStringView PackageName, FStringView& OutRootPath, FStringView& OutRelPath)
-{
-	if (!PackageName.StartsWith(TEXT("/")))
-	{
-		OutRootPath = FStringView();
-		OutRelPath = PackageName;
-		return;
-	}
-
-	int32 FirstSlashLen = 1;
-	int32 SecondSlashIndex;
-	PackageName.RightChop(FirstSlashLen).FindChar('/', SecondSlashIndex);
-	if (SecondSlashIndex == INDEX_NONE)
-	{
-		OutRootPath = PackageName;
-		OutRelPath = FStringView();
-	}
-	SecondSlashIndex += FirstSlashLen;
-
-	OutRootPath = PackageName.Left(SecondSlashIndex);
-	OutRelPath = PackageName.RightChop(SecondSlashIndex + 1);
-};
-
 bool FAssetRegistryImpl::ShouldSkipGatheredAsset(FAssetData& AssetData)
 {
 	// TODO: This pruning of invalid ExternalActors is temporary, to handle the fallout from a bug in SaveAs
@@ -4880,13 +4856,13 @@ bool FAssetRegistryImpl::ShouldSkipGatheredAsset(FAssetData& AssetData)
 		// OR
 		// /PackageRoot/__ExternalActors__/ContentBundle/######/RelPathFromPackageRootToMap/#/##/#######
 		// Package roots do not need to be the same; ContentBundles can be injected into /Game maps from plugins
-		SplitPackageNamePackageRoot(PackageNameStr, PackageNamePackageRoot, PackageNameRelPath);
-		SplitPackageNamePackageRoot(ObjectPathPackageName, ObjectPathPackageRoot, ObjectPathRelPath);
+		PackageNamePackageRoot = FPackageName::SplitPackageNameRoot(PackageNameStr, &PackageNameRelPath);
+		ObjectPathPackageRoot = FPackageName::SplitPackageNameRoot(ObjectPathPackageName, &ObjectPathRelPath);
 
 		if (!PackageNameRelPath.StartsWith(ExternalActorsFolderName) || !PackageNameRelPath.RightChop(ExternalActorsFolderName.Len()).StartsWith(TEXT("/")))
 		{
 			UE_LOG(LogAssetRegistry, Verbose,
-				TEXT("Invalid ExternalActor: Package %s is an ExternalActor package but is not in the expected root path for ExternalActors %.*s/%.*s. Ignoring this actor."),
+				TEXT("Invalid ExternalActor: Package %s is an ExternalActor package but is not in the expected root path for ExternalActors /%.*s/%.*s. Ignoring this actor."),
 				*PackageNameStr, PackageNamePackageRoot.Len(), PackageNamePackageRoot.GetData(),
 				ExternalActorsFolderName.Len(), ExternalActorsFolderName.GetData());
 			return true;
@@ -4915,7 +4891,7 @@ bool FAssetRegistryImpl::ShouldSkipGatheredAsset(FAssetData& AssetData)
 		if (bAllowValidation && !PackageNameRelPathAfterExternalActorRoot.StartsWith(ObjectPathRelPath))
 		{
 			TStringBuilder<256> ExpectedPath;
-			ExpectedPath << ObjectPathPackageRoot << "/" << ExternalActorsFolderName << "/" << ObjectPathRelPath;
+			ExpectedPath << "/" << ObjectPathPackageRoot << "/" << ExternalActorsFolderName << "/" << ObjectPathRelPath;
 			UE_LOG(LogAssetRegistry, Verbose,
 				TEXT("Invalid ExternalActor: Package %s is an ExternalActor package but its path does not match the expected path %s created from its objectpath %s. Ignoring this actor."),
 				*PackageNameStr, *ExpectedPath, *ObjectPathString);
diff --git a/Engine/Source/Runtime/CoreUObject/Private/Misc/PackageName.cpp b/Engine/Source/Runtime/CoreUObject/Private/Misc/PackageName.cpp
--- a/Engine/Source/Runtime/CoreUObject/Private/Misc/PackageName.cpp
+++ b/Engine/Source/Runtime/CoreUObject/Private/Misc/PackageName.cpp
@@ -2464,6 +2464,39 @@ FString FPackageName::ExportTextPathToObjectPath(const TCHAR* InExportTextPath)
 	return ExportTextPathToObjectPath(FString(InExportTextPath));
 }
 
+FStringView FPackageName::SplitPackageNameRoot(FStringView InPackageName, FStringView* OutRelativePath)
+{
+	if (!InPackageName.StartsWith(TEXT("/")))
+	{
+		if (OutRelativePath)
+		{
+			*OutRelativePath = InPackageName;
+		}
+		return FStringView();
+	}
+
+	// Strip the first slash.
+	InPackageName.RightChopInline(1);
+
+	int32 SecondSlashIndex;
+	InPackageName.FindChar('/', SecondSlashIndex);
+	if (SecondSlashIndex == INDEX_NONE)
+	{
+		if (OutRelativePath)
+		{
+			*OutRelativePath = FStringView();
+		}
+		return InPackageName;
+	}
+
+	if (OutRelativePath)
+	{
+		*OutRelativePath = InPackageName.RightChop(SecondSlashIndex + 1);
+	}
+	return InPackageName.Left(SecondSlashIndex);
+};
+
+
 template<class T>
 T ObjectPathToPackageNameImpl(const T& InObjectPath)
 {
diff --git a/Engine/Source/Runtime/CoreUObject/Public/Misc/PackageName.h b/Engine/Source/Runtime/CoreUObject/Public/Misc/PackageName.h
--- a/Engine/Source/Runtime/CoreUObject/Public/Misc/PackageName.h
+++ b/Engine/Source/Runtime/CoreUObject/Public/Misc/PackageName.h
@@ -648,6 +648,17 @@ public:
 	static FString			ExportTextPathToObjectPath(const FString& InExportTextPath);
 	static FString			ExportTextPathToObjectPath(const TCHAR* InExportTextPath);
 
+	/**
+	 * Returns the top level 'directory' in a package name. If the package is part of a
+	 * plugin, the root will be the name of the associated plugin, otherwise it will be
+	 * 'Engine' or 'Game'.
+	 * 
+	 * Note the slashes are removed!
+	 * 
+	 * "/PackageRoot/Path/Leaf" -> (return "PackageRoot"; OutRelativePath = "Path/Leaf")
+	 */
+	static FStringView SplitPackageNameRoot(FStringView InPackageName, FStringView* OutRelativePath);
+	
 	/** 
 	 * Returns the name of the package referred to by the specified object path
 	 * 
