Fixed vtable constructor to UAssetRegistryImpl to invoke the proper base class constructor.  Code would asset during a Live Coding session.  However, the asset registry still doesn't support live coding.

diff --git a/Engine/Source/Runtime/AssetRegistry/Private/AssetRegistry.cpp b/Engine/Source/Runtime/AssetRegistry/Private/AssetRegistry.cpp
--- a/Engine/Source/Runtime/AssetRegistry/Private/AssetRegistry.cpp
+++ b/Engine/Source/Runtime/AssetRegistry/Private/AssetRegistry.cpp
@@ -962,6 +962,7 @@ void UAssetRegistryImpl::InitializeEvents(UE::AssetRegistry::Impl::FInitializeCo
 }
 
 UAssetRegistryImpl::UAssetRegistryImpl(FVTableHelper& Helper)
+	: Super(Helper)
 {
 }
 
