[Engine]

diff --git a/Engine/Source/Runtime/ApplicationCore/Public/Windows/WindowsPlatformApplicationMisc.h b/Engine/Source/Runtime/ApplicationCore/Public/Windows/WindowsPlatformApplicationMisc.h
--- a/Engine/Source/Runtime/ApplicationCore/Public/Windows/WindowsPlatformApplicationMisc.h
+++ b/Engine/Source/Runtime/ApplicationCore/Public/Windows/WindowsPlatformApplicationMisc.h
@@ -4,29 +4,29 @@
 
 #include "GenericPlatform/GenericPlatformApplicationMisc.h"
 
-struct APPLICATIONCORE_API FWindowsPlatformApplicationMisc : public FGenericPlatformApplicationMisc
+struct FWindowsPlatformApplicationMisc : public FGenericPlatformApplicationMisc
 {
-	static void PreInit();
-	static void LoadStartupModules();
-	static class FOutputDeviceConsole* CreateConsoleOutputDevice();
-	static class FOutputDeviceError* GetErrorOutputDevice();
-	static class FFeedbackContext* GetFeedbackContext();
-	static class GenericApplication* CreateApplication();
-	static void RequestMinimize();
-	static bool IsThisApplicationForeground();
-	static int32 GetAppIcon();
-	static void PumpMessages(bool bFromMainLoop);
-	static void PreventScreenSaver();
-	static struct FLinearColor GetScreenPixelColor(const FVector2D& InScreenPos, float InGamma = 1.0f);
-	static void SetHighDPIMode();
-	static bool GetWindowTitleMatchingText(const TCHAR* TitleStartsWith, FString& OutTitle);
-	static float GetDPIScaleFactorAtPoint(float X, float Y);
-	static void ClipboardCopy(const TCHAR* Str);
-	static void ClipboardPaste(class FString& Dest);
+	static APPLICATIONCORE_API void PreInit();
+	static APPLICATIONCORE_API void LoadStartupModules();
+	static APPLICATIONCORE_API class FOutputDeviceConsole* CreateConsoleOutputDevice();
+	static APPLICATIONCORE_API class FOutputDeviceError* GetErrorOutputDevice();
+	static APPLICATIONCORE_API class FFeedbackContext* GetFeedbackContext();
+	static APPLICATIONCORE_API class GenericApplication* CreateApplication();
+	static APPLICATIONCORE_API void RequestMinimize();
+	static APPLICATIONCORE_API bool IsThisApplicationForeground();
+	static APPLICATIONCORE_API int32 GetAppIcon();
+	static APPLICATIONCORE_API void PumpMessages(bool bFromMainLoop);
+	static APPLICATIONCORE_API void PreventScreenSaver();
+	static APPLICATIONCORE_API struct FLinearColor GetScreenPixelColor(const FVector2D& InScreenPos, float InGamma = 1.0f);
+	static APPLICATIONCORE_API void SetHighDPIMode();
+	static APPLICATIONCORE_API bool GetWindowTitleMatchingText(const TCHAR* TitleStartsWith, FString& OutTitle);
+	static APPLICATIONCORE_API float GetDPIScaleFactorAtPoint(float X, float Y);
+	static APPLICATIONCORE_API void ClipboardCopy(const TCHAR* Str);
+	static APPLICATIONCORE_API void ClipboardPaste(class FString& Dest);
 
 	/** Windows platform only */
 	/** Function should retrieve the DPI value for the provided monitor information structure */
-	static int32 GetMonitorDPI(const FMonitorInfo& MonitorInfo);
+	static APPLICATIONCORE_API int32 GetMonitorDPI(const FMonitorInfo& MonitorInfo);
 
 	struct FGPUInfo
 	{
@@ -35,7 +35,7 @@ struct APPLICATIONCORE_API FWindowsPlatformApplicationMisc : public FGenericPlat
 		uint64 DedicatedVideoMemory = 0;
 	};
 
-	static FGPUInfo GetBestGPUInfo();
+	static APPLICATIONCORE_API FGPUInfo GetBestGPUInfo();
 	/** End Windows platform only */
 };
 
diff --git a/Engine/Source/Runtime/ApplicationCore/Public/Windows/WindowsPlatformSplash.h b/Engine/Source/Runtime/ApplicationCore/Public/Windows/WindowsPlatformSplash.h
--- a/Engine/Source/Runtime/ApplicationCore/Public/Windows/WindowsPlatformSplash.h
+++ b/Engine/Source/Runtime/ApplicationCore/Public/Windows/WindowsPlatformSplash.h
@@ -9,21 +9,21 @@
 /**
  * Windows splash implementation.
  */
-struct APPLICATIONCORE_API FWindowsPlatformSplash
+struct FWindowsPlatformSplash
 	: public FGenericPlatformSplash
 {
 	/** Show the splash screen. */
-	static void Show();
+	static APPLICATIONCORE_API void Show();
 
 	/** Hide the splash screen. */
-	static void Hide();
+	static APPLICATIONCORE_API void Hide();
 
 	/**
 	 * Sets the progress displayed on the application icon (for startup/loading progress).
 	 *
 	 * @param InType Progress value in percent.
 	 */
-	static void SetProgress(int ProgressPercent);
+	static APPLICATIONCORE_API void SetProgress(int ProgressPercent);
 
 	/**
 	 * Sets the text displayed on the splash screen (for startup/loading progress).
@@ -31,12 +31,12 @@ struct APPLICATIONCORE_API FWindowsPlatformSplash
 	 * @param InType Type of text to change.
 	 * @param InText Text to display.
 	 */
-	static void SetSplashText( const SplashTextType::Type InType, const TCHAR* InText );
+	static APPLICATIONCORE_API void SetSplashText( const SplashTextType::Type InType, const TCHAR* InText );
 
 	/**
 	 * Return whether the splash screen is being shown or not
 	 */
-	static bool IsShown();
+	static APPLICATIONCORE_API bool IsShown();
 };
 
 
diff --git a/Engine/Source/Runtime/ApplicationCore/Public/Windows/WindowsPlatformSurvey.h b/Engine/Source/Runtime/ApplicationCore/Public/Windows/WindowsPlatformSurvey.h
--- a/Engine/Source/Runtime/ApplicationCore/Public/Windows/WindowsPlatformSurvey.h
+++ b/Engine/Source/Runtime/ApplicationCore/Public/Windows/WindowsPlatformSurvey.h
@@ -9,11 +9,11 @@
 /**
  * Windows implementation of FGenericPlatformSurvey.
  */
-struct APPLICATIONCORE_API FWindowsPlatformSurvey
+struct FWindowsPlatformSurvey
 	: public FGenericPlatformSurvey
 {
 	/** Start, or check on, the hardware survey */
-	static bool GetSurveyResults( FHardwareSurveyResults& OutResults, bool bWait = false );
+	static APPLICATIONCORE_API bool GetSurveyResults( FHardwareSurveyResults& OutResults, bool bWait = false );
 
 private:
 	/**
diff --git a/Engine/Source/Runtime/ApplicationCore/Public/Windows/WindowsWindow.h b/Engine/Source/Runtime/ApplicationCore/Public/Windows/WindowsWindow.h
--- a/Engine/Source/Runtime/ApplicationCore/Public/Windows/WindowsWindow.h
+++ b/Engine/Source/Runtime/ApplicationCore/Public/Windows/WindowsWindow.h
@@ -21,38 +21,38 @@ enum class EWindowTransparency;
  *
  * Native Windows provide platform-specific backing for and are always owned by an SWindow.
  */
-class APPLICATIONCORE_API FWindowsWindow
+class FWindowsWindow
 	: public FGenericWindow
 	, public IDropTarget
 {
 public:
 
 	/** Win32 requirement: see CreateWindowEx and RegisterClassEx. */
-	static const TCHAR AppWindowClass[];
+	static APPLICATIONCORE_API const TCHAR AppWindowClass[];
 
 public:
 
 	/** Destructor. */
-	~FWindowsWindow();
+	APPLICATIONCORE_API ~FWindowsWindow();
 
 	/** Create a new FWin32Window. */
-	static TSharedRef<FWindowsWindow> Make();
+	static APPLICATIONCORE_API TSharedRef<FWindowsWindow> Make();
 
 	/**
 	 * Gets the Window's handle.
 	 *
 	 * @return The window's HWND handle.
 	 */
-	HWND GetHWnd() const;
+	APPLICATIONCORE_API HWND GetHWnd() const;
 
-	void Initialize( class FWindowsApplication* const Application, const TSharedRef<FGenericWindowDefinition>& InDefinition, HINSTANCE InHInstance, const TSharedPtr<FWindowsWindow>& InParent, const bool bShowImmediately );
+	APPLICATIONCORE_API void Initialize( class FWindowsApplication* const Application, const TSharedRef<FGenericWindowDefinition>& InDefinition, HINSTANCE InHInstance, const TSharedPtr<FWindowsWindow>& InParent, const bool bShowImmediately );
 
 	/**
 	* Sets the window text (usually the title but can also be text content for things like controls).
 	*
 	* @param Text The window's title or content text
 	*/
-	bool IsRegularWindow() const;
+	APPLICATIONCORE_API bool IsRegularWindow() const;
 
 	/**
      * Sets the window region to specified dimensions.
@@ -60,10 +60,10 @@ public:
 	 * @param Width The width of the window region (in pixels).
 	 * @param Height The height of the window region (in pixels).
 	 */
-	void AdjustWindowRegion(int32 Width, int32 Height);
+	APPLICATIONCORE_API void AdjustWindowRegion(int32 Width, int32 Height);
 
 	/** @return	Gives the native window a chance to adjust our stored window size before we cache it off. */
-	virtual void AdjustCachedSize(FVector2D& Size) const override;
+	APPLICATIONCORE_API virtual void AdjustCachedSize(FVector2D& Size) const override;
 
 	virtual float GetDPIScaleFactor() const override
 	{
@@ -76,87 +76,87 @@ public:
 	}
 
 	/** determines whether or not this window does its own DPI management */
-	virtual bool IsManualManageDPIChanges() const override;
+	APPLICATIONCORE_API virtual bool IsManualManageDPIChanges() const override;
 
 	/** call with a true argument if this window need to do its custom size management in response to DPI variations */
-	virtual void SetManualManageDPIChanges(const bool bManualDPIChanges) override;
+	APPLICATIONCORE_API virtual void SetManualManageDPIChanges(const bool bManualDPIChanges) override;
 
 	/** Called when our parent window is minimized (which will in turn cause us to become minimized). */
-	void OnParentWindowMinimized();
+	APPLICATIONCORE_API void OnParentWindowMinimized();
 
 	/** Called when our parent window is restored (which will in turn cause us to become restored). */
-	void OnParentWindowRestored();
+	APPLICATIONCORE_API void OnParentWindowRestored();
 
 	/** Called by the owning application when the level of transparency support has changed */
-	void OnTransparencySupportChanged(EWindowTransparency NewTransparency);
+	APPLICATIONCORE_API void OnTransparencySupportChanged(EWindowTransparency NewTransparency);
 
 	float GetAspectRatio() const { return AspectRatio; }
 
 	/** @return True if the window is enabled */
-	bool IsEnabled();
+	APPLICATIONCORE_API bool IsEnabled();
 
 public:
 
 	// FGenericWindow interface
 
-	virtual void ReshapeWindow( int32 X, int32 Y, int32 Width, int32 Height ) override;
-	virtual bool GetFullScreenInfo( int32& X, int32& Y, int32& Width, int32& Height ) const override;
-	virtual void MoveWindowTo ( int32 X, int32 Y ) override;
-	virtual void BringToFront( bool bForce = false ) override;
-	virtual void HACK_ForceToFront() override;
-	virtual void Destroy() override;
-	virtual void Minimize() override;
-	virtual void Maximize() override;
-	virtual void Restore() override;
-	virtual void Show() override;
-	virtual void Hide() override;
-	virtual void SetWindowMode( EWindowMode::Type NewWindowMode ) override;
+	APPLICATIONCORE_API virtual void ReshapeWindow( int32 X, int32 Y, int32 Width, int32 Height ) override;
+	APPLICATIONCORE_API virtual bool GetFullScreenInfo( int32& X, int32& Y, int32& Width, int32& Height ) const override;
+	APPLICATIONCORE_API virtual void MoveWindowTo ( int32 X, int32 Y ) override;
+	APPLICATIONCORE_API virtual void BringToFront( bool bForce = false ) override;
+	APPLICATIONCORE_API virtual void HACK_ForceToFront() override;
+	APPLICATIONCORE_API virtual void Destroy() override;
+	APPLICATIONCORE_API virtual void Minimize() override;
+	APPLICATIONCORE_API virtual void Maximize() override;
+	APPLICATIONCORE_API virtual void Restore() override;
+	APPLICATIONCORE_API virtual void Show() override;
+	APPLICATIONCORE_API virtual void Hide() override;
+	APPLICATIONCORE_API virtual void SetWindowMode( EWindowMode::Type NewWindowMode ) override;
 	virtual EWindowMode::Type GetWindowMode() const override { return WindowMode; } 
-	virtual bool IsMaximized() const override;
-	virtual bool IsMinimized() const override;
-	virtual bool IsVisible() const override;
-	virtual bool GetRestoredDimensions(int32& X, int32& Y, int32& Width, int32& Height) override;
-	virtual void SetWindowFocus() override;
-	virtual void SetOpacity( const float InOpacity ) override;
-	virtual void Enable( bool bEnable ) override;
-	virtual bool IsPointInWindow( int32 X, int32 Y ) const override;
-	virtual int32 GetWindowBorderSize() const override;
-	virtual int32 GetWindowTitleBarSize() const override;
+	APPLICATIONCORE_API virtual bool IsMaximized() const override;
+	APPLICATIONCORE_API virtual bool IsMinimized() const override;
+	APPLICATIONCORE_API virtual bool IsVisible() const override;
+	APPLICATIONCORE_API virtual bool GetRestoredDimensions(int32& X, int32& Y, int32& Width, int32& Height) override;
+	APPLICATIONCORE_API virtual void SetWindowFocus() override;
+	APPLICATIONCORE_API virtual void SetOpacity( const float InOpacity ) override;
+	APPLICATIONCORE_API virtual void Enable( bool bEnable ) override;
+	APPLICATIONCORE_API virtual bool IsPointInWindow( int32 X, int32 Y ) const override;
+	APPLICATIONCORE_API virtual int32 GetWindowBorderSize() const override;
+	APPLICATIONCORE_API virtual int32 GetWindowTitleBarSize() const override;
 	virtual void* GetOSWindowHandle() const  override { return HWnd; }
-	virtual bool IsForegroundWindow() const override;
-	virtual bool IsFullscreenSupported() const override;
-	virtual void SetText(const TCHAR* const Text) override;
-	virtual void DrawAttention(const FWindowDrawAttentionParameters& Parameters) override;
+	APPLICATIONCORE_API virtual bool IsForegroundWindow() const override;
+	APPLICATIONCORE_API virtual bool IsFullscreenSupported() const override;
+	APPLICATIONCORE_API virtual void SetText(const TCHAR* const Text) override;
+	APPLICATIONCORE_API virtual void DrawAttention(const FWindowDrawAttentionParameters& Parameters) override;
 
 public:
 
 	// IUnknown interface
 
-	HRESULT STDCALL QueryInterface(REFIID iid, void ** ppvObject) override;
-	ULONG STDCALL AddRef(void) override;
-	ULONG STDCALL Release(void) override;
+	APPLICATIONCORE_API HRESULT STDCALL QueryInterface(REFIID iid, void ** ppvObject) override;
+	APPLICATIONCORE_API ULONG STDCALL AddRef(void) override;
+	APPLICATIONCORE_API ULONG STDCALL Release(void) override;
 
 public:
 
 	// IDropTarget interface
 
-	virtual HRESULT STDCALL DragEnter( __RPC__in_opt IDataObject *DataObjectPointer, ::DWORD KeyState, POINTL CursorPosition, __RPC__inout ::DWORD *CursorEffect) override;
-	virtual HRESULT STDCALL DragOver( ::DWORD KeyState, POINTL CursorPosition, __RPC__inout ::DWORD *CursorEffect) override;
-	virtual HRESULT STDCALL DragLeave( void ) override;
-	virtual HRESULT STDCALL Drop( __RPC__in_opt IDataObject *DataObjectPointer, ::DWORD KeyState, POINTL CursorPosition, __RPC__inout ::DWORD *CursorEffect) override;
+	APPLICATIONCORE_API virtual HRESULT STDCALL DragEnter( __RPC__in_opt IDataObject *DataObjectPointer, ::DWORD KeyState, POINTL CursorPosition, __RPC__inout ::DWORD *CursorEffect) override;
+	APPLICATIONCORE_API virtual HRESULT STDCALL DragOver( ::DWORD KeyState, POINTL CursorPosition, __RPC__inout ::DWORD *CursorEffect) override;
+	APPLICATIONCORE_API virtual HRESULT STDCALL DragLeave( void ) override;
+	APPLICATIONCORE_API virtual HRESULT STDCALL Drop( __RPC__in_opt IDataObject *DataObjectPointer, ::DWORD KeyState, POINTL CursorPosition, __RPC__inout ::DWORD *CursorEffect) override;
 
 private:
 
 	/** Protect the constructor; only TSharedRefs of this class can be made. */
-	FWindowsWindow();
+	APPLICATIONCORE_API FWindowsWindow();
 
-	void UpdateVisibility();
+	APPLICATIONCORE_API void UpdateVisibility();
 
 	/** Creates an HRGN for the window's current region.  Remember to delete this when you're done with it using
 	   ::DeleteObject, unless you're passing it to SetWindowRgn(), which will absorb the reference itself. */
-	HRGN MakeWindowRegionObject(bool bIncludeBorderWhenMaximized) const;
+	APPLICATIONCORE_API HRGN MakeWindowRegionObject(bool bIncludeBorderWhenMaximized) const;
 
-	void DisableTouchFeedback();
+	APPLICATIONCORE_API void DisableTouchFeedback();
 
 private:
 
diff --git a/Engine/Source/Runtime/AssetRegistry/Public/AssetRegistry/AssetDependencyGatherer.h b/Engine/Source/Runtime/AssetRegistry/Public/AssetRegistry/AssetDependencyGatherer.h
--- a/Engine/Source/Runtime/AssetRegistry/Public/AssetRegistry/AssetDependencyGatherer.h
+++ b/Engine/Source/Runtime/AssetRegistry/Public/AssetRegistry/AssetDependencyGatherer.h
@@ -51,25 +51,25 @@ public:
 
 namespace UE::AssetDependencyGatherer::Private
 {
-	class ASSETREGISTRY_API FRegisteredAssetDependencyGatherer
+	class FRegisteredAssetDependencyGatherer
 	{
 	public:
 		DECLARE_MULTICAST_DELEGATE(FOnAssetDependencyGathererRegistered);
-		static FOnAssetDependencyGathererRegistered OnAssetDependencyGathererRegistered;
+		static ASSETREGISTRY_API FOnAssetDependencyGathererRegistered OnAssetDependencyGathererRegistered;
 
-		FRegisteredAssetDependencyGatherer();
-		virtual ~FRegisteredAssetDependencyGatherer();
+		ASSETREGISTRY_API FRegisteredAssetDependencyGatherer();
+		ASSETREGISTRY_API virtual ~FRegisteredAssetDependencyGatherer();
 
 		virtual UClass* GetAssetClass() const = 0;
 		virtual void GatherDependencies(const FAssetData& AssetData, const FAssetRegistryState& AssetRegistryState, 
 			TFunctionRef<FARCompiledFilter(const FARFilter&)> CompileFilterFunc, TArray<IAssetDependencyGatherer::FGathereredDependency>& OutDependencies, 
 			TArray<FString>& OutDependencyDirectories) const = 0;
 
-		static void ForEach(TFunctionRef<void(FRegisteredAssetDependencyGatherer*)> Func);
+		static ASSETREGISTRY_API void ForEach(TFunctionRef<void(FRegisteredAssetDependencyGatherer*)> Func);
 		static bool IsEmpty() { return !GetRegisteredList(); }
 	private:
 
-		static TLinkedList<FRegisteredAssetDependencyGatherer*>*& GetRegisteredList();
+		static ASSETREGISTRY_API TLinkedList<FRegisteredAssetDependencyGatherer*>*& GetRegisteredList();
 		TLinkedList<FRegisteredAssetDependencyGatherer*> GlobalListLink;
 	};
 }
@@ -98,4 +98,4 @@ namespace PREPROCESSOR_JOIN(GathererClass, AssetClass) \
 	static PREPROCESSOR_JOIN(PREPROCESSOR_JOIN(GathererClass, AssetClass), _Register) DefaultObject; \
 }
 
-#endif // WITH_EDITOR
\ No newline at end of file
+#endif // WITH_EDITOR
diff --git a/Engine/Source/Runtime/AssetRegistry/Public/AssetRegistry/AssetRegistryHelpers.h b/Engine/Source/Runtime/AssetRegistry/Public/AssetRegistry/AssetRegistryHelpers.h
--- a/Engine/Source/Runtime/AssetRegistry/Public/AssetRegistry/AssetRegistryHelpers.h
+++ b/Engine/Source/Runtime/AssetRegistry/Public/AssetRegistry/AssetRegistryHelpers.h
@@ -125,10 +125,10 @@ public:
 	static ASSETREGISTRY_API void FindReferencersOfAssetOfClass(const FAssetIdentifier& InAssetIdentifier, TConstArrayView<UClass*> InMatchClasses, TArray<FAssetData>& OutAssetDatas);
 	
 	/** Enable/disable asset registry caching mode for the duration of the scope */
-	struct ASSETREGISTRY_API FTemporaryCachingModeScope
+	struct FTemporaryCachingModeScope
 	{
-		FTemporaryCachingModeScope(bool InTempCachingMode);
-		~FTemporaryCachingModeScope();
+		ASSETREGISTRY_API FTemporaryCachingModeScope(bool InTempCachingMode);
+		ASSETREGISTRY_API ~FTemporaryCachingModeScope();
 
 	private:
 		bool PreviousCachingMode;
diff --git a/Engine/Source/Runtime/AssetRegistry/Public/AssetRegistry/AssetRegistryState.h b/Engine/Source/Runtime/AssetRegistry/Public/AssetRegistry/AssetRegistryState.h
--- a/Engine/Source/Runtime/AssetRegistry/Public/AssetRegistry/AssetRegistryState.h
+++ b/Engine/Source/Runtime/AssetRegistry/Public/AssetRegistry/AssetRegistryState.h
@@ -188,7 +188,7 @@ namespace UE::AssetRegistry::Private
 }
 
 /** The state of an asset registry, this is used internally by IAssetRegistry to represent the disk cache, and is also accessed directly to save/load cooked caches */
-class ASSETREGISTRY_API FAssetRegistryState
+class FAssetRegistryState
 {
 	using FCachedAssetKey = UE::AssetRegistry::Private::FCachedAssetKey;
 public:
@@ -199,10 +199,10 @@ public:
 	FAssetRegistryState() = default;
 	FAssetRegistryState(const FAssetRegistryState&) = delete;
 	FAssetRegistryState(FAssetRegistryState&& Rhs) { *this = MoveTemp(Rhs); }
-	~FAssetRegistryState();
+	ASSETREGISTRY_API ~FAssetRegistryState();
 
 	FAssetRegistryState& operator=(const FAssetRegistryState&) = delete;
-	FAssetRegistryState& operator=(FAssetRegistryState&& O);
+	ASSETREGISTRY_API FAssetRegistryState& operator=(FAssetRegistryState&& O);
 
 	/**
 	 * Enum controlling how we initialize this state
@@ -220,7 +220,7 @@ public:
 	 * @param bARFiltering Whether to apply filtering from UE::AssetRegistry::FFiltering (false by default)
 	 * @note This function doesn't recurse into sub-paths.
 	 */
-	bool HasAssets(const FName PackagePath, bool bARFiltering=false) const;
+	ASSETREGISTRY_API bool HasAssets(const FName PackagePath, bool bARFiltering=false) const;
 
 	/**
 	 * Gets asset data for all assets that match the filter.
@@ -232,7 +232,7 @@ public:
 	 * @param OutAssetData the list of assets in this path
 	 * @param bSkipARFilteredAssets If true, skip assets that are skipped by UE::AssetRegistry::FFiltering (false by default)
 	 */
-	bool GetAssets(const FARCompiledFilter& Filter, const TSet<FName>& PackageNamesToSkip, TArray<FAssetData>& OutAssetData, bool bSkipARFilteredAssets = false) const;
+	ASSETREGISTRY_API bool GetAssets(const FARCompiledFilter& Filter, const TSet<FName>& PackageNamesToSkip, TArray<FAssetData>& OutAssetData, bool bSkipARFilteredAssets = false) const;
 
 	/**
 	 * Enumerate asset data for all assets that match the filter.
@@ -244,7 +244,7 @@ public:
 	 * @param Callback function to call for each asset data enumerated
 	 * @param bARFiltering Whether to apply filtering from UE::AssetRegistry::FFiltering (false by default)
 	 */
-	bool EnumerateAssets(const FARCompiledFilter& Filter, const TSet<FName>& PackageNamesToSkip, TFunctionRef<bool(const FAssetData&)> Callback, bool bARFiltering = false) const;
+	ASSETREGISTRY_API bool EnumerateAssets(const FARCompiledFilter& Filter, const TSet<FName>& PackageNamesToSkip, TFunctionRef<bool(const FAssetData&)> Callback, bool bARFiltering = false) const;
 
 	/**
 	 * Gets asset data for all assets in the registry state.
@@ -253,7 +253,7 @@ public:
 	 * @param OutAssetData the list of assets
 	 * @param bARFiltering Whether to apply filtering from UE::AssetRegistry::FFiltering (false by default)
 	 */
-	bool GetAllAssets(const TSet<FName>& PackageNamesToSkip, TArray<FAssetData>& OutAssetData, bool bARFiltering = false) const;
+	ASSETREGISTRY_API bool GetAllAssets(const TSet<FName>& PackageNamesToSkip, TArray<FAssetData>& OutAssetData, bool bARFiltering = false) const;
 
 	/**
 	 * Enumerates asset data for all assets in the registry state.
@@ -262,8 +262,8 @@ public:
 	 * @param Callback function to call for each asset data enumerated
 	 * @param bARFiltering Whether to apply filtering from UE::AssetRegistry::FFiltering (false by default)
 	 */
-	bool EnumerateAllAssets(const TSet<FName>& PackageNamesToSkip, TFunctionRef<bool(const FAssetData&)> Callback, bool bARFiltering = false) const;
-	void EnumerateAllAssets(TFunctionRef<void(const FAssetData&)> Callback) const;
+	ASSETREGISTRY_API bool EnumerateAllAssets(const TSet<FName>& PackageNamesToSkip, TFunctionRef<bool(const FAssetData&)> Callback, bool bARFiltering = false) const;
+	ASSETREGISTRY_API void EnumerateAllAssets(TFunctionRef<void(const FAssetData&)> Callback) const;
 
 	/**
 	 * Gets the LongPackageNames for all packages with the given PackageName.
@@ -272,7 +272,7 @@ public:
 	 * @param PackageName Name of the package to find, may be a LongPackageName or ShortPackageName.
 	 * @param OutPackageNames All discovered matching LongPackageNames are appended to this array.
 	 */
-	void GetPackagesByName(FStringView PackageName, TArray<FName>& OutPackageNames) const;
+	ASSETREGISTRY_API void GetPackagesByName(FStringView PackageName, TArray<FName>& OutPackageNames) const;
 
 	/**
 	 * Returns the first LongPackageName found for the given PackageName.
@@ -282,7 +282,7 @@ public:
 	 * @param PackageName Name of the package to find, may be a LongPackageName or ShortPackageName.
 	 * @return The first LongPackageName of the matching package, or NAME_None if not found.
 	 */
-	FName GetFirstPackageByName(FStringView PackageName) const;
+	ASSETREGISTRY_API FName GetFirstPackageByName(FStringView PackageName) const;
 
 	/**
 	 * Appends a list of packages and searchable names that are referenced by the supplied package or name. (On disk references ONLY)
@@ -292,8 +292,8 @@ public:
 	 * @param Category	which category(ies) of dependencies to include in the output list. Dependencies matching ANY of the OR'd categories will be returned.
 	 * @param Flags	which flags are required present or not present on the dependencies. Dependencies matching ALL required and NONE excluded bits will be returned. For each potentially returned dependency, flags not applicable to their category are ignored.
 	 */
-	bool GetDependencies(const FAssetIdentifier& AssetIdentifier, TArray<FAssetIdentifier>& OutDependencies, UE::AssetRegistry::EDependencyCategory Category = UE::AssetRegistry::EDependencyCategory::All, const UE::AssetRegistry::FDependencyQuery& Flags = UE::AssetRegistry::FDependencyQuery()) const;
-	bool GetDependencies(const FAssetIdentifier& AssetIdentifier, TArray<FAssetDependency>& OutDependencies, UE::AssetRegistry::EDependencyCategory Category = UE::AssetRegistry::EDependencyCategory::All, const UE::AssetRegistry::FDependencyQuery& Flags = UE::AssetRegistry::FDependencyQuery()) const;
+	ASSETREGISTRY_API bool GetDependencies(const FAssetIdentifier& AssetIdentifier, TArray<FAssetIdentifier>& OutDependencies, UE::AssetRegistry::EDependencyCategory Category = UE::AssetRegistry::EDependencyCategory::All, const UE::AssetRegistry::FDependencyQuery& Flags = UE::AssetRegistry::FDependencyQuery()) const;
+	ASSETREGISTRY_API bool GetDependencies(const FAssetIdentifier& AssetIdentifier, TArray<FAssetDependency>& OutDependencies, UE::AssetRegistry::EDependencyCategory Category = UE::AssetRegistry::EDependencyCategory::All, const UE::AssetRegistry::FDependencyQuery& Flags = UE::AssetRegistry::FDependencyQuery()) const;
 
 	/**
 	 * Appends a list of packages and searchable names that reference the supplied package or name. (On disk references ONLY)
@@ -303,8 +303,8 @@ public:
 	 * @param Category	which category(ies) of dependencies to include in the output list. Dependencies matching ANY of the OR'd categories will be returned.
 	 * @param Flags	which flags are required present or not present on the dependencies. Dependencies matching ALL required and NONE excluded bits will be returned. For each potentially returned dependency, flags not applicable to their category are ignored.
 	 */
-	bool GetReferencers(const FAssetIdentifier& AssetIdentifier, TArray<FAssetIdentifier>& OutReferencers, UE::AssetRegistry::EDependencyCategory Category = UE::AssetRegistry::EDependencyCategory::All, const UE::AssetRegistry::FDependencyQuery& Flags = UE::AssetRegistry::FDependencyQuery()) const;
-	bool GetReferencers(const FAssetIdentifier& AssetIdentifier, TArray<FAssetDependency>& OutReferencers, UE::AssetRegistry::EDependencyCategory Category = UE::AssetRegistry::EDependencyCategory::All, const UE::AssetRegistry::FDependencyQuery& Flags = UE::AssetRegistry::FDependencyQuery()) const;
+	ASSETREGISTRY_API bool GetReferencers(const FAssetIdentifier& AssetIdentifier, TArray<FAssetIdentifier>& OutReferencers, UE::AssetRegistry::EDependencyCategory Category = UE::AssetRegistry::EDependencyCategory::All, const UE::AssetRegistry::FDependencyQuery& Flags = UE::AssetRegistry::FDependencyQuery()) const;
+	ASSETREGISTRY_API bool GetReferencers(const FAssetIdentifier& AssetIdentifier, TArray<FAssetDependency>& OutReferencers, UE::AssetRegistry::EDependencyCategory Category = UE::AssetRegistry::EDependencyCategory::All, const UE::AssetRegistry::FDependencyQuery& Flags = UE::AssetRegistry::FDependencyQuery()) const;
 
 	/**
 	 * Gets the asset data for the specified object path
@@ -359,7 +359,7 @@ public:
 	 * @return An array of AssetData*, empty if nothing found
 	 */
 	UE_DEPRECATED(5.1, "Class names are now represented by path names. Please use GetAssetsByClassPathName")
-	const TArray<const FAssetData*>& GetAssetsByClassName(const FName ClassName) const;
+	ASSETREGISTRY_API const TArray<const FAssetData*>& GetAssetsByClassName(const FName ClassName) const;
 
 	/**
 	 * Gets the asset data for the specified asset class
@@ -416,11 +416,11 @@ public:
 	}
 
 	/** Get the set of primary assets contained in this state */
-	void GetPrimaryAssetsIds(TSet<FPrimaryAssetId>& OutPrimaryAssets) const;
+	ASSETREGISTRY_API void GetPrimaryAssetsIds(TSet<FPrimaryAssetId>& OutPrimaryAssets) const;
 
 	/** Returns pointer to the asset package data */
-	const FAssetPackageData* GetAssetPackageData(FName PackageName) const;
-	FAssetPackageData* GetAssetPackageData(FName PackageName);
+	ASSETREGISTRY_API const FAssetPackageData* GetAssetPackageData(FName PackageName) const;
+	ASSETREGISTRY_API FAssetPackageData* GetAssetPackageData(FName PackageName);
 
 	/** Returns all package names */
 	void GetPackageNames(TArray<FName>& OutPackageNames) const
@@ -433,24 +433,24 @@ public:
 	}
 
 	/** Finds an existing package data, or creates a new one to modify */
-	FAssetPackageData* CreateOrGetAssetPackageData(FName PackageName);
+	ASSETREGISTRY_API FAssetPackageData* CreateOrGetAssetPackageData(FName PackageName);
 
 	/** Removes existing package data */
-	bool RemovePackageData(FName PackageName);
+	ASSETREGISTRY_API bool RemovePackageData(FName PackageName);
 
 	/** Adds the asset data to the lookup maps */
-	void AddAssetData(FAssetData* AssetData);
+	ASSETREGISTRY_API void AddAssetData(FAssetData* AssetData);
 
 	/** Add the given tags/values to the asset data associated with the given object path, if it exists */
-	void AddTagsToAssetData(const FSoftObjectPath& InObjectPath, FAssetDataTagMap&& InTagsAndValues);
+	ASSETREGISTRY_API void AddTagsToAssetData(const FSoftObjectPath& InObjectPath, FAssetDataTagMap&& InTagsAndValues);
 
 	/** Finds an existing asset data based on object path and updates it with the new value and updates lookup maps */
-	void UpdateAssetData(const FAssetData& NewAssetData, bool bCreateIfNotExists=false);
-	void UpdateAssetData(FAssetData&& NewAssetData, bool bCreateIfNotExists = false);
+	ASSETREGISTRY_API void UpdateAssetData(const FAssetData& NewAssetData, bool bCreateIfNotExists=false);
+	ASSETREGISTRY_API void UpdateAssetData(FAssetData&& NewAssetData, bool bCreateIfNotExists = false);
 
 	/** Updates an existing asset data with the new value and updates lookup maps */
-	void UpdateAssetData(FAssetData* AssetData, const FAssetData& NewAssetData, bool* bOutModified = nullptr);
-	void UpdateAssetData(FAssetData* AssetData, FAssetData&& NewAssetData, bool* bOutModified = nullptr);
+	ASSETREGISTRY_API void UpdateAssetData(FAssetData* AssetData, const FAssetData& NewAssetData, bool* bOutModified = nullptr);
+	ASSETREGISTRY_API void UpdateAssetData(FAssetData* AssetData, FAssetData&& NewAssetData, bool* bOutModified = nullptr);
 
 	/**
 	 * Updates all asset data package flags in the specified package
@@ -459,49 +459,49 @@ public:
 	 * @param PackageFlags the package flags to set
 	 * @return True if any assets exists in the package
 	 */
-	bool UpdateAssetDataPackageFlags(FName PackageName, uint32 PackageFlags);
+	ASSETREGISTRY_API bool UpdateAssetDataPackageFlags(FName PackageName, uint32 PackageFlags);
 
 	/** Removes the asset data from the lookup maps */
-	void RemoveAssetData(FAssetData* AssetData, bool bRemoveDependencyData, bool& bOutRemovedAssetData, bool& bOutRemovedPackageData);
+	ASSETREGISTRY_API void RemoveAssetData(FAssetData* AssetData, bool bRemoveDependencyData, bool& bOutRemovedAssetData, bool& bOutRemovedPackageData);
 
 	/**
 	 * Clear all dependencies of the given category from the given AssetIdentifier (e.g. package).
 	 * Also clears the referencer link from each of the dependencies.
 	 */
-	void ClearDependencies(const FAssetIdentifier& AssetIdentifier, UE::AssetRegistry::EDependencyCategory Category);
+	ASSETREGISTRY_API void ClearDependencies(const FAssetIdentifier& AssetIdentifier, UE::AssetRegistry::EDependencyCategory Category);
 	/**
 	 * Add the given dependencies to the given AssetIdentifier (e.g. package).
 	 * Also adds a referencer link on each of the dependencies.
 	 */
-	void AddDependencies(const FAssetIdentifier& AssetIdentifier, TConstArrayView<FAssetDependency> Dependencies);
+	ASSETREGISTRY_API void AddDependencies(const FAssetIdentifier& AssetIdentifier, TConstArrayView<FAssetDependency> Dependencies);
 	/**
 	 * Clears existing dependencies of the given Category(s) and assigns the input Dependencies. Gives an error if 
 	 * any elements of Dependencies are outside of the Category(s).
 	 */
-	void SetDependencies(const FAssetIdentifier& AssetIdentifier, TConstArrayView<FAssetDependency> Dependencies,
+	ASSETREGISTRY_API void SetDependencies(const FAssetIdentifier& AssetIdentifier, TConstArrayView<FAssetDependency> Dependencies,
 		UE::AssetRegistry::EDependencyCategory Category = UE::AssetRegistry::EDependencyCategory::All);
 	/**
 	 * Clear all referencers of the given category from the given AssetIdentifier (e.g. package).
 	 * Also clears the dependency link from each of the referencers.
 	 */
-	void ClearReferencers(const FAssetIdentifier& AssetIdentifier, UE::AssetRegistry::EDependencyCategory Category);
+	ASSETREGISTRY_API void ClearReferencers(const FAssetIdentifier& AssetIdentifier, UE::AssetRegistry::EDependencyCategory Category);
 	/**
 	 * Add a dependency on the given AssetIdentifier (e.g. package) from each of the Referencers.
 	 * Also adds a referencer link to each referencer on the AssetIdentifer's node.
 	 */
-	void AddReferencers(const FAssetIdentifier& AssetIdentifier, TConstArrayView<FAssetDependency> Referencers);
+	ASSETREGISTRY_API void AddReferencers(const FAssetIdentifier& AssetIdentifier, TConstArrayView<FAssetDependency> Referencers);
 	/**
 	 * Clears existing referencers of the given Category(s) and assigns the input Referencers. Gives an error if
 	 * any elements of Referencers are outside of the Category(s).
 	 */
-	void SetReferencers(const FAssetIdentifier& AssetIdentifier, TConstArrayView<FAssetDependency> Referencers,
+	ASSETREGISTRY_API void SetReferencers(const FAssetIdentifier& AssetIdentifier, TConstArrayView<FAssetDependency> Referencers,
 		UE::AssetRegistry::EDependencyCategory Category = UE::AssetRegistry::EDependencyCategory::All);
 
 	/** Resets to default state */
-	void Reset();
+	ASSETREGISTRY_API void Reset();
 
 	/** Initializes cache from existing set of asset data and depends nodes */
-	void InitializeFromExisting(const FAssetDataMap& AssetDataMap, const TMap<FAssetIdentifier, FDependsNode*>& DependsNodeMap, const TMap<FName, FAssetPackageData*>& AssetPackageDataMap, const FAssetRegistrySerializationOptions& Options, EInitializationMode InitializationMode = EInitializationMode::Rebuild);
+	ASSETREGISTRY_API void InitializeFromExisting(const FAssetDataMap& AssetDataMap, const TMap<FAssetIdentifier, FDependsNode*>& DependsNodeMap, const TMap<FName, FAssetPackageData*>& AssetPackageDataMap, const FAssetRegistrySerializationOptions& Options, EInitializationMode InitializationMode = EInitializationMode::Rebuild);
 	void InitializeFromExisting(const FAssetRegistryState& Existing, const FAssetRegistrySerializationOptions& Options, EInitializationMode InitializationMode = EInitializationMode::Rebuild)
 	{
 		InitializeFromExisting(Existing.CachedAssets, Existing.CachedDependsNodes, Existing.CachedPackageData, Options, InitializationMode);
@@ -514,9 +514,9 @@ public:
 	 * @param ChunksToKeep The list of chunks that are allowed to remain. Any assets in other chunks are pruned. If empty, all assets are kept regardless of chunk
 	 * @param Options Serialization options to read filter info from
 	 */
-	void PruneAssetData(const TSet<FName>& RequiredPackages, const TSet<FName>& RemovePackages, const TSet<int32> ChunksToKeep, const FAssetRegistrySerializationOptions& Options);
-	void PruneAssetData(const TSet<FName>& RequiredPackages, const TSet<FName>& RemovePackages, const FAssetRegistrySerializationOptions& Options);
-	void Prune(const FAssetRegistryPruneOptions& PruneOptions);
+	ASSETREGISTRY_API void PruneAssetData(const TSet<FName>& RequiredPackages, const TSet<FName>& RemovePackages, const TSet<int32> ChunksToKeep, const FAssetRegistrySerializationOptions& Options);
+	ASSETREGISTRY_API void PruneAssetData(const TSet<FName>& RequiredPackages, const TSet<FName>& RemovePackages, const FAssetRegistrySerializationOptions& Options);
+	ASSETREGISTRY_API void Prune(const FAssetRegistryPruneOptions& PruneOptions);
 
 	
 	/**
@@ -527,31 +527,31 @@ public:
 	 * @param ChunksToKeep The list of chunks that are allowed to remain. Any assets in other chunks are pruned. If empty, all assets are kept regardless of chunk
 	 * @param Options Serialization options to read filter info from
 	 */
-	void InitializeFromExistingAndPrune(const FAssetRegistryState& ExistingState, const TSet<FName>& RequiredPackages, const TSet<FName>& RemovePackages, const TSet<int32> ChunksToKeep, const FAssetRegistrySerializationOptions& Options);
+	ASSETREGISTRY_API void InitializeFromExistingAndPrune(const FAssetRegistryState& ExistingState, const TSet<FName>& RequiredPackages, const TSet<FName>& RemovePackages, const TSet<int32> ChunksToKeep, const FAssetRegistrySerializationOptions& Options);
 
 	/** Edit every AssetData's Tags to remove Tags that are filtered out by the filtering rules in Options */
-	void FilterTags(const FAssetRegistrySerializationOptions& Options);
+	ASSETREGISTRY_API void FilterTags(const FAssetRegistrySerializationOptions& Options);
 
 
 	/** Serialize the registry to/from a file, skipping editor only data */
-	bool Serialize(FArchive& Ar, const FAssetRegistrySerializationOptions& Options);
+	ASSETREGISTRY_API bool Serialize(FArchive& Ar, const FAssetRegistrySerializationOptions& Options);
 
 	/** Save without editor-only data */
-	bool Save(FArchive& Ar, const FAssetRegistrySerializationOptions& Options);
-	bool Load(FArchive& Ar, const FAssetRegistryLoadOptions& Options = FAssetRegistryLoadOptions(), FAssetRegistryVersion::Type* OutVersion = nullptr);
+	ASSETREGISTRY_API bool Save(FArchive& Ar, const FAssetRegistrySerializationOptions& Options);
+	ASSETREGISTRY_API bool Load(FArchive& Ar, const FAssetRegistryLoadOptions& Options = FAssetRegistryLoadOptions(), FAssetRegistryVersion::Type* OutVersion = nullptr);
 
 	/** 
 	* Example Usage:
 	*	FAssetRegistryState AssetRegistry;
 	*	bool bSucceeded = FAssetRegistryState::LoadFromDisk(TEXT("Path/To/AR"), FAssetRegistryLoadOptions(), AssetRegistry);
 	*/
-	static bool LoadFromDisk(const TCHAR* InPath, const FAssetRegistryLoadOptions& InOptions, FAssetRegistryState& OutState, FAssetRegistryVersion::Type* OutVersion = nullptr);
+	static ASSETREGISTRY_API bool LoadFromDisk(const TCHAR* InPath, const FAssetRegistryLoadOptions& InOptions, FAssetRegistryState& OutState, FAssetRegistryVersion::Type* OutVersion = nullptr);
 
 	/** Returns memory size of entire registry, optionally logging sizes */
-	SIZE_T GetAllocatedSize(bool bLogDetailed = false) const;
+	ASSETREGISTRY_API SIZE_T GetAllocatedSize(bool bLogDetailed = false) const;
 
 	/** Checks a filter to make sure there are no illegal entries */
-	static bool IsFilterValid(const FARCompiledFilter& Filter);
+	static ASSETREGISTRY_API bool IsFilterValid(const FARCompiledFilter& Filter);
 
 	/** Returns the number of assets in this state */
 	int32 GetNumAssets() const { return NumAssets; }
@@ -567,7 +567,7 @@ public:
 	 * @param LinesPerPage - how many lines should be combined into each string element of OutPages, for e.g. breaking up the dump into separate files.
 	 *        To facilitate diffing between similar-but-different registries, the actual number of lines per page will be slightly less than LinesPerPage; we introduce partially deterministic pagebreaks near the end of each page.
 	 */
-	void Dump(const TArray<FString>& Arguments, TArray<FString>& OutPages, int32 LinesPerPage=1) const;
+	ASSETREGISTRY_API void Dump(const TArray<FString>& Arguments, TArray<FString>& OutPages, int32 LinesPerPage=1) const;
 #endif
 
 private:
@@ -575,27 +575,27 @@ private:
 	void Load(Archive&& Ar, const FAssetRegistryHeader& Header, const FAssetRegistryLoadOptions& Options);
 
 	/** Initialize the lookup maps */
-	void SetAssetDatas(TArrayView<FAssetData> AssetDatas, const FAssetRegistryLoadOptions& Options);
+	ASSETREGISTRY_API void SetAssetDatas(TArrayView<FAssetData> AssetDatas, const FAssetRegistryLoadOptions& Options);
 
 	/** Find the first non-redirector dependency node starting from InDependency. */
-	FDependsNode* ResolveRedirector(FDependsNode* InDependency, const FAssetDataMap& InAllowedAssets, TMap<FDependsNode*, FDependsNode*>& InCache);
+	ASSETREGISTRY_API FDependsNode* ResolveRedirector(FDependsNode* InDependency, const FAssetDataMap& InAllowedAssets, TMap<FDependsNode*, FDependsNode*>& InCache);
 
 	/** Finds an existing node for the given package and returns it, or returns null if one isn't found */
-	FDependsNode* FindDependsNode(const FAssetIdentifier& Identifier) const;
+	ASSETREGISTRY_API FDependsNode* FindDependsNode(const FAssetIdentifier& Identifier) const;
 
 	/** Creates a node in the CachedDependsNodes map or finds the existing node and returns it */
-	FDependsNode* CreateOrFindDependsNode(const FAssetIdentifier& Identifier);
+	ASSETREGISTRY_API FDependsNode* CreateOrFindDependsNode(const FAssetIdentifier& Identifier);
 
 	/** Removes the depends node and updates the dependencies to no longer contain it as as a referencer. */
-	bool RemoveDependsNode(const FAssetIdentifier& Identifier);
+	ASSETREGISTRY_API bool RemoveDependsNode(const FAssetIdentifier& Identifier);
 
 	/** Filter a set of tags and output a copy of the filtered set. */
-	static void FilterTags(const FAssetDataTagMapSharedView& InTagsAndValues, FAssetDataTagMap& OutTagsAndValues, const TSet<FName>* ClassSpecificFilterList, const FAssetRegistrySerializationOptions & Options);
+	static ASSETREGISTRY_API void FilterTags(const FAssetDataTagMapSharedView& InTagsAndValues, FAssetDataTagMap& OutTagsAndValues, const TSet<FName>* ClassSpecificFilterList, const FAssetRegistrySerializationOptions & Options);
 
-	void LoadDependencies(FArchive& Ar);
-	void LoadDependencies_BeforeFlags(FArchive& Ar, bool bSerializeDependencies, FAssetRegistryVersion::Type Version);
+	ASSETREGISTRY_API void LoadDependencies(FArchive& Ar);
+	ASSETREGISTRY_API void LoadDependencies_BeforeFlags(FArchive& Ar, bool bSerializeDependencies, FAssetRegistryVersion::Type Version);
 
-	void SetTagsOnExistingAsset(FAssetData* AssetData, FAssetDataTagMap&& NewTags);
+	ASSETREGISTRY_API void SetTagsOnExistingAsset(FAssetData* AssetData, FAssetDataTagMap&& NewTags);
 
 	/** Set of asset data for assets saved to disk. Searched via path name types, implicitly converted to FCachedAssetKey. */
 	FAssetDataMap CachedAssets;
@@ -743,4 +743,4 @@ namespace UE::AssetRegistry::Private
 	{
 		return HashCombineQuick(GetTypeHash(A.OuterPath), GetTypeHash(A.ObjectName));
 	}
-}
\ No newline at end of file
+}
diff --git a/Engine/Source/Runtime/AssetRegistry/Public/AssetRegistry/IAssetRegistry.h b/Engine/Source/Runtime/AssetRegistry/Public/AssetRegistry/IAssetRegistry.h
--- a/Engine/Source/Runtime/AssetRegistry/Public/AssetRegistry/IAssetRegistry.h
+++ b/Engine/Source/Runtime/AssetRegistry/Public/AssetRegistry/IAssetRegistry.h
@@ -844,29 +844,29 @@ namespace AssetRegistry
 		InvalidTag = 3,
 	};
 
-	struct ASSETREGISTRY_API FDeserializePackageData
+	struct FDeserializePackageData
 	{
 		int64 DependencyDataOffset = INDEX_NONE;
 		int32 ObjectCount = 0;
 
-		bool DoSerialize(FArchive& BinaryArchive, const FPackageFileSummary& PackageFileSummary, EReadPackageDataMainErrorCode& OutError);
+		ASSETREGISTRY_API bool DoSerialize(FArchive& BinaryArchive, const FPackageFileSummary& PackageFileSummary, EReadPackageDataMainErrorCode& OutError);
 	};
 
-	struct ASSETREGISTRY_API FDeserializeObjectPackageData
+	struct FDeserializeObjectPackageData
 	{
 		FString ObjectPath;
 		FString ObjectClassName;
 		int32 TagCount = 0;
 
-		bool DoSerialize(FArchive& BinaryArchive, EReadPackageDataMainErrorCode& OutError);
+		ASSETREGISTRY_API bool DoSerialize(FArchive& BinaryArchive, EReadPackageDataMainErrorCode& OutError);
 	};
 
-	struct ASSETREGISTRY_API FDeserializeTagData
+	struct FDeserializeTagData
 	{
 		FString Key;
 		FString Value;
 
-		bool DoSerialize(FArchive& BinaryArchive, EReadPackageDataMainErrorCode& OutError);
+		ASSETREGISTRY_API bool DoSerialize(FArchive& BinaryArchive, EReadPackageDataMainErrorCode& OutError);
 	};
 
 	// Functions to read and write the data used by the AssetRegistry in each package; the format of this data is separate from the format of the data in the asset registry
diff --git a/Engine/Source/Runtime/AssetRegistry/Public/AssetRegistry/PathTree.h b/Engine/Source/Runtime/AssetRegistry/Public/AssetRegistry/PathTree.h
--- a/Engine/Source/Runtime/AssetRegistry/Public/AssetRegistry/PathTree.h
+++ b/Engine/Source/Runtime/AssetRegistry/Public/AssetRegistry/PathTree.h
@@ -13,29 +13,29 @@
 
 template <typename FuncType> class TFunctionRef;
 
-class ASSETREGISTRY_API FPathTree
+class FPathTree
 {
 public:
 	/** Adds the specified path to the tree, creating nodes as needed and calling OnPathAdded for any new paths added. Returns true if the specified path was actually added (as opposed to already existed) */
-	bool CachePath(FName Path, TFunctionRef<void(FName)> OnPathAdded);
+	ASSETREGISTRY_API bool CachePath(FName Path, TFunctionRef<void(FName)> OnPathAdded);
 
 	/** Removes the specified path from the tree, calling OnPathRemoved for any existing paths removed. Returns true if the path was found and removed. */
-	bool RemovePath(FName Path, TFunctionRef<void(FName)> OnPathRemoved);
+	ASSETREGISTRY_API bool RemovePath(FName Path, TFunctionRef<void(FName)> OnPathRemoved);
 
 	/** Checks whether the given path is one that we know about */
-	bool PathExists(FName Path) const;
+	ASSETREGISTRY_API bool PathExists(FName Path) const;
 
 	/** Get all of the paths we know about */
-	bool GetAllPaths(TSet<FName>& OutPaths) const;
+	ASSETREGISTRY_API bool GetAllPaths(TSet<FName>& OutPaths) const;
 
 	/** Enumerate all of the paths we know about */
-	void EnumerateAllPaths(TFunctionRef<bool(FName)> Callback) const;
+	ASSETREGISTRY_API void EnumerateAllPaths(TFunctionRef<bool(FName)> Callback) const;
 
 	/** Recursively gathers all child paths from the specified base path relative to this node */
-	bool GetSubPaths(FName BasePath, TSet<FName>& OutPaths, bool bRecurse = true) const;
+	ASSETREGISTRY_API bool GetSubPaths(FName BasePath, TSet<FName>& OutPaths, bool bRecurse = true) const;
 
 	/** Recursively enumerates all child paths from the specified base path relative to this node */
-	bool EnumerateSubPaths(FName BasePath, TFunctionRef<bool(FName)> Callback, bool bRecurse = true) const;
+	ASSETREGISTRY_API bool EnumerateSubPaths(FName BasePath, TFunctionRef<bool(FName)> Callback, bool bRecurse = true) const;
 
 	int32 NumPaths() const
 	{
@@ -55,7 +55,7 @@ public:
 	}
 
 	/** Normalize the given PackagePath (/Game/SubDir) by removing a trailing slash if it exists. */
-	static FName NormalizePackagePath(FName In);
+	static ASSETREGISTRY_API FName NormalizePackagePath(FName In);
 
 private:
 	/** A one-to-many mapping between a parent path and its child paths. */
diff --git a/Engine/Source/Runtime/AudioAnalyzer/Classes/AudioAnalyzer.h b/Engine/Source/Runtime/AudioAnalyzer/Classes/AudioAnalyzer.h
--- a/Engine/Source/Runtime/AudioAnalyzer/Classes/AudioAnalyzer.h
+++ b/Engine/Source/Runtime/AudioAnalyzer/Classes/AudioAnalyzer.h
@@ -19,8 +19,8 @@ class UAudioAnalyzerSubsystem;
  * across multiple analyzers. 
  *
  */
-UCLASS(Abstract, EditInlineNew, BlueprintType)
-class AUDIOANALYZER_API UAudioAnalyzerSettings : public UAudioAnalyzerAssetBase
+UCLASS(Abstract, EditInlineNew, BlueprintType, MinimalAPI)
+class UAudioAnalyzerSettings : public UAudioAnalyzerAssetBase
 {
 	GENERATED_BODY()
 };
@@ -64,8 +64,8 @@ private:
  * To support blueprint access, subclasses can implement UFUNCTIONs to expose the data
  * returned by GetResult().
  */
-UCLASS(Abstract, EditInlineNew, BlueprintType)
-class AUDIOANALYZER_API UAudioAnalyzer : public UObject
+UCLASS(Abstract, EditInlineNew, BlueprintType, MinimalAPI)
+class UAudioAnalyzer : public UObject
 {
 	GENERATED_BODY()
 	
@@ -87,34 +87,34 @@ public:
 
 	/** Starts analyzing audio from the given audio bus. Optionally override the audio bus desired to analyze. */
 	UFUNCTION(BlueprintCallable, BlueprintCosmetic, Category = AudioAnalyzer, meta = (WorldContext = "WorldContextObject"))
-	void StartAnalyzing(const UObject* WorldContextObject, UAudioBus* AudioBusToAnalyze);
+	AUDIOANALYZER_API void StartAnalyzing(const UObject* WorldContextObject, UAudioBus* AudioBusToAnalyze);
 
 	/** Starts analyzing using the given world.*/
-	void StartAnalyzing(UWorld* InWorld, UAudioBus* AudioBusToAnalyze);
+	AUDIOANALYZER_API void StartAnalyzing(UWorld* InWorld, UAudioBus* AudioBusToAnalyze);
 
 	/** Stops analyzing audio. */
 	UFUNCTION(BlueprintCallable, BlueprintCosmetic, Category = AudioAnalyzer, meta = (WorldContext = "WorldContextObject"))
-	void StopAnalyzing(const UObject* WorldContextObject = nullptr);
+	AUDIOANALYZER_API void StopAnalyzing(const UObject* WorldContextObject = nullptr);
 
 	/**
 	 * Implementations can override this method to create settings objects
 	 * specific for their analyzer.
 	 */
-	virtual TUniquePtr<Audio::IAnalyzerSettings> GetSettings(const int32 InSampleRate, const int32 InNumChannels) const;
+	AUDIOANALYZER_API virtual TUniquePtr<Audio::IAnalyzerSettings> GetSettings(const int32 InSampleRate, const int32 InNumChannels) const;
 
 	/** Function to broadcast results. */
 	virtual void BroadcastResults() {}
 
 	//~ Begin UObject Interface.
-	virtual void BeginDestroy() override;
+	AUDIOANALYZER_API virtual void BeginDestroy() override;
 	//~ End UObject interface
 
 private: 
 	/** Returns if this audio analyzer has enough audio queued up and is ready for analysis. */
-	bool IsReadyForAnalysis() const;
+	AUDIOANALYZER_API bool IsReadyForAnalysis() const;
 
 	/** Does the actual analysis and casts the results to the derived classes type. Returns true if there results to broadcast. */
-	bool DoAnalysis();
+	AUDIOANALYZER_API bool DoAnalysis();
 
 protected:
 
@@ -125,7 +125,7 @@ protected:
 	}
 
 	/* Subclasses must override this method in order to inform this object which AnalyzerFactory to use for analysis */
-	virtual FName GetAnalyzerFactoryName() const PURE_VIRTUAL(UAudioAnalyzer::GetAnalyzerFactoryName, return FName(););
+	AUDIOANALYZER_API virtual FName GetAnalyzerFactoryName() const PURE_VIRTUAL(UAudioAnalyzer::GetAnalyzerFactoryName, return FName(););
 
 	/** How many frames of audio to wait before analyzing the audio. */
 	int32 NumFramesPerBufferToAnalyze = 1024;
@@ -133,7 +133,7 @@ protected:
 private:
 
 	// Returns UAudioAnalyzerSettings* if property points to a valid UAudioAnalyzerSettings, otherwise returns nullptr.
-	UAudioAnalyzerSettings* GetSettingsFromProperty(FProperty* Property);
+	AUDIOANALYZER_API UAudioAnalyzerSettings* GetSettingsFromProperty(FProperty* Property);
 
 	// Audio analysis subsystem used with this audio analyzer
 	UPROPERTY(Transient)
diff --git a/Engine/Source/Runtime/AudioAnalyzer/Classes/AudioAnalyzerAsset.h b/Engine/Source/Runtime/AudioAnalyzer/Classes/AudioAnalyzerAsset.h
--- a/Engine/Source/Runtime/AudioAnalyzer/Classes/AudioAnalyzerAsset.h
+++ b/Engine/Source/Runtime/AudioAnalyzer/Classes/AudioAnalyzerAsset.h
@@ -13,8 +13,8 @@ DECLARE_MULTICAST_DELEGATE(FAnalyzeAudioDelegate);
  *
  * UAudioAnalyzerAssetBase provides the base interface for controlling asset actions within the editor.
  */
-UCLASS(Abstract, EditInlineNew)
-class AUDIOANALYZER_API UAudioAnalyzerAssetBase : public UObject
+UCLASS(Abstract, EditInlineNew, MinimalAPI)
+class UAudioAnalyzerAssetBase : public UObject
 {
 	GENERATED_BODY()
 
@@ -26,12 +26,12 @@ class AUDIOANALYZER_API UAudioAnalyzerAssetBase : public UObject
 		/**
 		 * GetAssetActionName() returns the FText displayed in the editor.
 		 */ 
-		virtual FText GetAssetActionName() const PURE_VIRTUAL(UAudioAnalyzerAsset::GetAssetActionName, return FText(););
+		AUDIOANALYZER_API virtual FText GetAssetActionName() const PURE_VIRTUAL(UAudioAnalyzerAsset::GetAssetActionName, return FText(););
 
 		/**
 		 * GetSupportedClass() returns the class which should be associated with these asset actions.
 		 */ 
-		virtual UClass* GetSupportedClass() const PURE_VIRTUAL(UAudioAnalyzerAsset::GetSupportedClass, return nullptr;);
+		AUDIOANALYZER_API virtual UClass* GetSupportedClass() const PURE_VIRTUAL(UAudioAnalyzerAsset::GetSupportedClass, return nullptr;);
 
 		/**
 		 * GetTypeColor() returns the color used to display asset icons within the editor.
diff --git a/Engine/Source/Runtime/AudioAnalyzer/Classes/AudioAnalyzerNRT.h b/Engine/Source/Runtime/AudioAnalyzer/Classes/AudioAnalyzerNRT.h
--- a/Engine/Source/Runtime/AudioAnalyzer/Classes/AudioAnalyzerNRT.h
+++ b/Engine/Source/Runtime/AudioAnalyzer/Classes/AudioAnalyzerNRT.h
@@ -20,8 +20,8 @@
  * setting when when a UPROPERTY in this setting object is edited.
  *
  */
-UCLASS(Abstract, EditInlineNew, BlueprintType)
-class AUDIOANALYZER_API UAudioAnalyzerNRTSettings : public UAudioAnalyzerAssetBase
+UCLASS(Abstract, EditInlineNew, BlueprintType, MinimalAPI)
+class UAudioAnalyzerNRTSettings : public UAudioAnalyzerAssetBase
 {
 	GENERATED_BODY()
 
@@ -37,13 +37,13 @@ class AUDIOANALYZER_API UAudioAnalyzerNRTSettings : public UAudioAnalyzerAssetBa
 		 *
 		 * This determines whether to trigger analysis by calling ShouldEventTriggerAnalysis(...)
 		 */
-		void PostEditChangeProperty (struct FPropertyChangedEvent & PropertyChangedEvent) override;
+		AUDIOANALYZER_API void PostEditChangeProperty (struct FPropertyChangedEvent & PropertyChangedEvent) override;
 
 		/** 
 		 * This returns true when the PropertyChangeEvent is due to update a setting property. 
 		 * Override this method in order to customize this behavior.
 		 */
-		virtual bool ShouldEventTriggerAnalysis(struct FPropertyChangedEvent & PropertyChangeEvent);
+		AUDIOANALYZER_API virtual bool ShouldEventTriggerAnalysis(struct FPropertyChangedEvent & PropertyChangeEvent);
 #endif
 
 };
@@ -59,8 +59,8 @@ class AUDIOANALYZER_API UAudioAnalyzerNRTSettings : public UAudioAnalyzerAssetBa
  * To support blueprint access, subclasses can implement UFUNCTIONs to expose the data
  * returned by GetResult().
  */
-UCLASS(Abstract, EditInlineNew, BlueprintType)
-class AUDIOANALYZER_API UAudioAnalyzerNRT : public UAudioAnalyzerAssetBase
+UCLASS(Abstract, EditInlineNew, BlueprintType, MinimalAPI)
+class UAudioAnalyzerNRT : public UAudioAnalyzerAssetBase
 {
 	GENERATED_BODY()
 
@@ -120,26 +120,26 @@ class AUDIOANALYZER_API UAudioAnalyzerNRT : public UAudioAnalyzerAssetBase
 		 * Implementations can override this method to create settings objects
 		 * specific for their analyzer. 
 		 */
-		virtual TUniquePtr<Audio::IAnalyzerNRTSettings> GetSettings(const float InSampleRate, const int32 InNumChannels) const;
+		AUDIOANALYZER_API virtual TUniquePtr<Audio::IAnalyzerNRTSettings> GetSettings(const float InSampleRate, const int32 InNumChannels) const;
 
 		/**
 		 * Performs serialization of results.
 		 */
-		virtual void Serialize(FArchive& Ar) override;
+		AUDIOANALYZER_API virtual void Serialize(FArchive& Ar) override;
 
 #if WITH_EDITOR
 
-		void SetResult(FResultSharedPtr NewResult);
+		AUDIOANALYZER_API void SetResult(FResultSharedPtr NewResult);
 
 		/** This will only store the result if the passed InResultId matches the CurrentResultId. */
-		void SetResultIfLatest(FResultSharedPtr NewResult, FResultId InResultId);
+		AUDIOANALYZER_API void SetResultIfLatest(FResultSharedPtr NewResult, FResultId InResultId);
 
 		/**
 		 * Called before a UPROPERTY of this class is edited. This checks to see
 		 * if the UPROPERTY is a UAudioAnalyzerNRTSettings. If so, the previous settings' 
 		 * AnalyzeAudioDelegate will be unbound from this object's AnalyzeAudio()
 		 */
-		void PreEditChange(FProperty* PropertyAboutToChange) override;
+		AUDIOANALYZER_API void PreEditChange(FProperty* PropertyAboutToChange) override;
 
 		/**
 		 * Called when a UPROPERTY of this class is edited. Triggering 
@@ -148,16 +148,16 @@ class AUDIOANALYZER_API UAudioAnalyzerNRT : public UAudioAnalyzerAssetBase
 		 *
 		 * This determines whether to trigger analysis by calling ShouldEventTriggerAnalysis(...)
 		 */
-		void PostEditChangeProperty (struct FPropertyChangedEvent & PropertyChangedEvent) override;
+		AUDIOANALYZER_API void PostEditChangeProperty (struct FPropertyChangedEvent & PropertyChangedEvent) override;
 
 		/**
 		 * Returns true when the PropertyChangeEvent is due to update SoundWave or Settings.
 		 */
-		virtual bool ShouldEventTriggerAnalysis(struct FPropertyChangedEvent & PropertyChangeEvent);
+		AUDIOANALYZER_API virtual bool ShouldEventTriggerAnalysis(struct FPropertyChangedEvent & PropertyChangeEvent);
 
 		/** Performs the analaysis of the audio */
 		UFUNCTION()
-		void AnalyzeAudio();
+		AUDIOANALYZER_API void AnalyzeAudio();
 #endif
 
 
@@ -166,19 +166,19 @@ class AUDIOANALYZER_API UAudioAnalyzerNRT : public UAudioAnalyzerAssetBase
 #if WITH_EDITOR
 
 		/** This sets the AnalyzeAudio callback to the delegate in settings. */
-		void SetSettingsDelegate(UAudioAnalyzerNRTSettings* InSettings);
+		AUDIOANALYZER_API void SetSettingsDelegate(UAudioAnalyzerNRTSettings* InSettings);
 
 		/** This removes the existing delegate handle from the current settings object */
-		void RemoveSettingsDelegate(UAudioAnalyzerNRTSettings* InSettings);
+		AUDIOANALYZER_API void RemoveSettingsDelegate(UAudioAnalyzerNRTSettings* InSettings);
 #endif 
 
 		/* Subclasses must override this method in order to inform this object which AnalyzerNRTFactory to use for analysis */
-		virtual FName GetAnalyzerNRTFactoryName() const PURE_VIRTUAL(UAudioAnalyzerNRT::GetAnalyzerNRTFactoryName, return FName(););
+		AUDIOANALYZER_API virtual FName GetAnalyzerNRTFactoryName() const PURE_VIRTUAL(UAudioAnalyzerNRT::GetAnalyzerNRTFactoryName, return FName(););
 
 	private:
 
 		// Returns UAudioAnalyzerNRTSettings* if property points to a valid UAudioAnalyzerNRTSettings, otherwise returns nullptr.
-		UAudioAnalyzerNRTSettings* GetSettingsFromProperty(FProperty* Property);
+		AUDIOANALYZER_API UAudioAnalyzerNRTSettings* GetSettingsFromProperty(FProperty* Property);
 
 		TSharedPtr<Audio::IAnalyzerNRTResult, ESPMode::ThreadSafe> Result;
 
diff --git a/Engine/Source/Runtime/AudioAnalyzer/Classes/AudioAnalyzerSubsystem.h b/Engine/Source/Runtime/AudioAnalyzer/Classes/AudioAnalyzerSubsystem.h
--- a/Engine/Source/Runtime/AudioAnalyzer/Classes/AudioAnalyzerSubsystem.h
+++ b/Engine/Source/Runtime/AudioAnalyzer/Classes/AudioAnalyzerSubsystem.h
@@ -13,26 +13,26 @@ class UWorld;
 /** 
 * Class manages querying analysis results from various audio analyzers.
 */
-UCLASS()
-class AUDIOANALYZER_API UAudioAnalyzerSubsystem : public UEngineSubsystem
+UCLASS(MinimalAPI)
+class UAudioAnalyzerSubsystem : public UEngineSubsystem
 {
 	GENERATED_BODY()
 
 public:
-	UAudioAnalyzerSubsystem();
-	~UAudioAnalyzerSubsystem();
+	AUDIOANALYZER_API UAudioAnalyzerSubsystem();
+	AUDIOANALYZER_API ~UAudioAnalyzerSubsystem();
 
-	bool Tick(float DeltaTime);
+	AUDIOANALYZER_API bool Tick(float DeltaTime);
 
 	//~ Begin USubsystem
-	virtual void Initialize(FSubsystemCollectionBase& Collection) override;
-	virtual void Deinitialize() override;
+	AUDIOANALYZER_API virtual void Initialize(FSubsystemCollectionBase& Collection) override;
+	AUDIOANALYZER_API virtual void Deinitialize() override;
 	//~ End USubsystem
 
-	void RegisterAudioAnalyzer(UAudioAnalyzer* InAnalyzer);
-	void UnregisterAudioAnalyzer(UAudioAnalyzer* InAnalyzer);
+	AUDIOANALYZER_API void RegisterAudioAnalyzer(UAudioAnalyzer* InAnalyzer);
+	AUDIOANALYZER_API void UnregisterAudioAnalyzer(UAudioAnalyzer* InAnalyzer);
 
-	static UAudioAnalyzerSubsystem* Get();
+	static AUDIOANALYZER_API UAudioAnalyzerSubsystem* Get();
 
 private:
 
diff --git a/Engine/Source/Runtime/AudioAnalyzer/Public/AudioAnalyzerFacade.h b/Engine/Source/Runtime/AudioAnalyzer/Public/AudioAnalyzerFacade.h
--- a/Engine/Source/Runtime/AudioAnalyzer/Public/AudioAnalyzerFacade.h
+++ b/Engine/Source/Runtime/AudioAnalyzer/Public/AudioAnalyzerFacade.h
@@ -18,18 +18,18 @@ namespace Audio
 	 *
 	 * FAnalyzer provides a simplified interface for performing analysis on audio buffers from the same AudioAnalyzer.
 	 */
-	class AUDIOANALYZER_API FAnalyzerFacade
+	class FAnalyzerFacade
 	{
 		public:
 			/**
 			 * Create an FAnalyzerBatch with the analyzer settings and factory name.
 			 */
-			FAnalyzerFacade(TUniquePtr<IAnalyzerSettings> InSettings, IAnalyzerFactory* InFactory);
+			AUDIOANALYZER_API FAnalyzerFacade(TUniquePtr<IAnalyzerSettings> InSettings, IAnalyzerFactory* InFactory);
 
 			/**
 			 * Analyze the audio buffer.
 			 */
-			TUniquePtr<IAnalyzerResult> AnalyzeAudioBuffer(const TArray<float>& InAudioBuffer, int32 InNumChannels, float InSampleRate);
+			AUDIOANALYZER_API TUniquePtr<IAnalyzerResult> AnalyzeAudioBuffer(const TArray<float>& InAudioBuffer, int32 InNumChannels, float InSampleRate);
 
 		private:
 
diff --git a/Engine/Source/Runtime/AudioAnalyzer/Public/AudioAnalyzerModule.h b/Engine/Source/Runtime/AudioAnalyzer/Public/AudioAnalyzerModule.h
--- a/Engine/Source/Runtime/AudioAnalyzer/Public/AudioAnalyzerModule.h
+++ b/Engine/Source/Runtime/AudioAnalyzer/Public/AudioAnalyzerModule.h
@@ -17,10 +17,10 @@ LLM_DECLARE_TAG_API(Audio_Analysis, AUDIOANALYZER_API);
 // Convenience macro for Audio_Analysis LLM scope to avoid misspells.
 #define AUDIO_ANALYSIS_LLM_SCOPE LLM_SCOPE_BYTAG(Audio_Analysis);
 
-class AUDIOANALYZER_API FAudioAnalyzerModule : public IModuleInterface
+class FAudioAnalyzerModule : public IModuleInterface
 {
 	public:
 	/** IModuleInterface implementation */
-	void StartupModule();
-	void ShutdownModule();
+	AUDIOANALYZER_API void StartupModule();
+	AUDIOANALYZER_API void ShutdownModule();
 };
diff --git a/Engine/Source/Runtime/AudioAnalyzer/Public/AudioAnalyzerNRTFacade.h b/Engine/Source/Runtime/AudioAnalyzer/Public/AudioAnalyzerNRTFacade.h
--- a/Engine/Source/Runtime/AudioAnalyzer/Public/AudioAnalyzerNRTFacade.h
+++ b/Engine/Source/Runtime/AudioAnalyzer/Public/AudioAnalyzerNRTFacade.h
@@ -19,18 +19,18 @@ namespace Audio
 	 * FAnalyzerNRTFacade provides a simplified interface for running 
 	 * analyzer factories over complete audio resources.
 	 */
-	class AUDIOANALYZER_API FAnalyzerNRTFacade
+	class FAnalyzerNRTFacade
 	{
 		public:
 			/**
 			 * Create an FAnalyzerNRTFacade with the analyzer settings and factory name.
 			 */
-			FAnalyzerNRTFacade(TUniquePtr<IAnalyzerNRTSettings> InSettings, const FName& InFactoryName);
+			AUDIOANALYZER_API FAnalyzerNRTFacade(TUniquePtr<IAnalyzerNRTSettings> InSettings, const FName& InFactoryName);
 
 			/**
 			 * Analyze an entire PCM16 encoded audio object.  Audio for the entire sound should be contained within InRawWaveData.
 			 */
-			TUniquePtr<IAnalyzerNRTResult> AnalyzePCM16Audio(const TArray<uint8>& InRawWaveData, int32 InNumChannels, float InSampleRate);
+			AUDIOANALYZER_API TUniquePtr<IAnalyzerNRTResult> AnalyzePCM16Audio(const TArray<uint8>& InRawWaveData, int32 InNumChannels, float InSampleRate);
 
 		private:
 
diff --git a/Engine/Source/Runtime/AudioExtensions/Public/AudioParameter.h b/Engine/Source/Runtime/AudioExtensions/Public/AudioParameter.h
--- a/Engine/Source/Runtime/AudioExtensions/Public/AudioParameter.h
+++ b/Engine/Source/Runtime/AudioExtensions/Public/AudioParameter.h
@@ -25,15 +25,15 @@ class UObject;
 
 namespace Audio
 {
-	struct AUDIOEXTENSIONS_API FParameterPath
+	struct FParameterPath
 	{
-		static const FString NamespaceDelimiter;
+		static AUDIOEXTENSIONS_API const FString NamespaceDelimiter;
 
 		// Combines names using the namespace delimiter
-		static FName CombineNames(FName InLeft, FName InRight);
+		static AUDIOEXTENSIONS_API FName CombineNames(FName InLeft, FName InRight);
 
 		// Splits name into namespace & parameter name
-		static void SplitName(FName InFullName, FName& OutNamespace, FName& OutParameterName);
+		static AUDIOEXTENSIONS_API void SplitName(FName InFullName, FName& OutNamespace, FName& OutParameterName);
 	};
 }
 
@@ -89,7 +89,7 @@ enum class EAudioParameterType : uint8
 
 
 USTRUCT(BlueprintType)
-struct AUDIOEXTENSIONS_API FAudioParameter
+struct FAudioParameter
 {
 	GENERATED_USTRUCT_BODY()
 
@@ -227,13 +227,13 @@ struct AUDIOEXTENSIONS_API FAudioParameter
 	// bInTakeName = Take the name of the provided parameter.
 	// bInTakeType = Take the type of the provided parameter.
 	// bInMergeArrayTypes - Appends array(s) for specified type if true, else swaps the local value with that of the provided parameter if false.
-	void Merge(const FAudioParameter& InParameter, bool bInTakeName = true, bool bInTakeType = true, bool bInMergeArrayTypes = false);
+	AUDIOEXTENSIONS_API void Merge(const FAudioParameter& InParameter, bool bInTakeName = true, bool bInTakeType = true, bool bInMergeArrayTypes = false);
 
 	// Moves InParams to OutParams that are not already included. For backward compatibility (i.e. SoundCues),
 	// if a param is already in OutParams, attempts to merge param values together, but assigns the param the
 	// incoming param's type. Currently existing OutParam values left if new value of the same type is provided
 	// by InParam.
-	static void Merge(TArray<FAudioParameter>&& InParams, TArray<FAudioParameter>& OutParams);
+	static AUDIOEXTENSIONS_API void Merge(TArray<FAudioParameter>&& InParams, TArray<FAudioParameter>& OutParams);
 
 	FAudioParameter(FAudioParameter&& InParameter) = default;
 	FAudioParameter(const FAudioParameter& InParameter) = default;
diff --git a/Engine/Source/Runtime/AudioExtensions/Public/AudioParameterControllerInterface.h b/Engine/Source/Runtime/AudioExtensions/Public/AudioParameterControllerInterface.h
--- a/Engine/Source/Runtime/AudioExtensions/Public/AudioParameterControllerInterface.h
+++ b/Engine/Source/Runtime/AudioExtensions/Public/AudioParameterControllerInterface.h
@@ -21,14 +21,14 @@ struct FAudioParameter;
 struct FFrame;
 
 
-UINTERFACE(BlueprintType, meta = (CannotImplementInterfaceInBlueprint))
-class AUDIOEXTENSIONS_API UAudioParameterControllerInterface : public UInterface
+UINTERFACE(BlueprintType, meta = (CannotImplementInterfaceInBlueprint), MinimalAPI)
+class UAudioParameterControllerInterface : public UInterface
 {
 	GENERATED_UINTERFACE_BODY()
 };
 
 // Base interface for any object implementing parameter control for a given sound instance controller.
-class AUDIOEXTENSIONS_API IAudioParameterControllerInterface : public IInterface
+class IAudioParameterControllerInterface : public IInterface
 {
 	GENERATED_IINTERFACE_BODY()
 
diff --git a/Engine/Source/Runtime/AudioExtensions/Public/IAudioEndpoint.h b/Engine/Source/Runtime/AudioExtensions/Public/IAudioEndpoint.h
--- a/Engine/Source/Runtime/AudioExtensions/Public/IAudioEndpoint.h
+++ b/Engine/Source/Runtime/AudioExtensions/Public/IAudioEndpoint.h
@@ -39,7 +39,7 @@ DECLARE_LOG_CATEGORY_EXTERN(LogAudioEndpoints, Display, All);
  * This interface should be used to provide a non-uclass version of the data described in
  * your implementation of UAudioEndpointSettingsBase.
  */
-class AUDIOEXTENSIONS_API IAudioEndpointSettingsProxy
+class IAudioEndpointSettingsProxy
 {
 public:
 	virtual ~IAudioEndpointSettingsProxy() {}
@@ -49,13 +49,13 @@ public:
  * This opaque class should be used for specifying settings for how audio should be
  * send to an external endpoint.
  */
-UCLASS(config = Engine, abstract, editinlinenew, BlueprintType)
-class AUDIOEXTENSIONS_API UAudioEndpointSettingsBase : public UObject
+UCLASS(config = Engine, abstract, editinlinenew, BlueprintType, MinimalAPI)
+class UAudioEndpointSettingsBase : public UObject
 {
 	GENERATED_BODY()
 
 public:
-	virtual TUniquePtr<IAudioEndpointSettingsProxy> GetProxy() const PURE_VIRTUAL(UAudioEndpointSettingsBase::GetProxy, return nullptr;);
+	AUDIOEXTENSIONS_API virtual TUniquePtr<IAudioEndpointSettingsProxy> GetProxy() const PURE_VIRTUAL(UAudioEndpointSettingsBase::GetProxy, return nullptr;);
 };
 
 //A blank class for when unimplemented endpoint types are returned
@@ -77,7 +77,7 @@ class UDummyEndpointSettings : public UAudioEndpointSettingsBase
  * Note that this only for interleaved audio buffers with no metadata for object-based or soundfield-based rendering.
  * For those, see 
  */
-class AUDIOEXTENSIONS_API IAudioEndpoint
+class IAudioEndpoint
 {
 public:
 	virtual ~IAudioEndpoint() {};
@@ -94,34 +94,34 @@ public:
 	 *                             	Audio::DownmixBuffer(NumInputChannels, NumOutputChannels, InputAudio, OutputAudio, MixdownGainsMap.GetData());
 	 * @returns a new Audio::FPatchInput. The FPatchInput may be disconnected if this endpoint's sample rate or channel count changes, in which case you will need to reconnect by calling this again.
 	 */
-	Audio::FPatchInput PatchNewInput(float ExpectedDurationPerRender, float& OutSampleRate, int32& OutNumChannels);
+	AUDIOEXTENSIONS_API Audio::FPatchInput PatchNewInput(float ExpectedDurationPerRender, float& OutSampleRate, int32& OutNumChannels);
 
 	/**
 	 * Post new settings for this endpoint.
 	 * There is no type safety on this call, so make sure that you are using the correct implementation
 	 * of IAudioEndpointSettingsProxy for this implementation of IAudioEndpoint.
 	 */
-	void SetNewSettings(TUniquePtr<IAudioEndpointSettingsProxy>&& InNewSettings);
+	AUDIOEXTENSIONS_API void SetNewSettings(TUniquePtr<IAudioEndpointSettingsProxy>&& InNewSettings);
 
 	/**
 	 * If this audio endpoint hasn't spawned a seperate callback thread but requires a callback, this should be executed somewhere.
 	 */
-	void ProcessAudioIfNeccessary();
+	AUDIOEXTENSIONS_API void ProcessAudioIfNeccessary();
 
 	/**
 	* Whether this endpoint is of an implemented type
 	*/
-	virtual bool IsImplemented();
+	AUDIOEXTENSIONS_API virtual bool IsImplemented();
 
 protected:
 
 	/** REQUIRED OVERRIDES: */
 
 	/** This should return the sample rate we should be sending to this endpoint. If the sample rate changes, please call DisconnectAllInputs(). */
-	virtual float GetSampleRate() const;
+	AUDIOEXTENSIONS_API virtual float GetSampleRate() const;
 
 	/** This should return the number of channels we should be sending to this endpoint. If the number of channels changes, please call DisconnectAllInputs. */
-	virtual int32 GetNumChannels() const;
+	AUDIOEXTENSIONS_API virtual int32 GetNumChannels() const;
 
 	/** OPTIONAL OVERRIDES: */
 
@@ -149,7 +149,7 @@ protected:
 	 * @param [in] NumSamples: The number of samples to fill OutAudio with.
 	 * @returns [out] the number of samples polled from this thing.
 	 */
-	int32 PopAudio(float* OutAudio, int32 NumSamples);
+	AUDIOEXTENSIONS_API int32 PopAudio(float* OutAudio, int32 NumSamples);
 
 	/**
 	 * Use this as a thread safe way to use the current settings posted to this IAudioEndpoint. Locks with IAudioEndpoint::SetSettings.
@@ -157,24 +157,24 @@ protected:
 	 *                                 This lambda is called immediately and synchronously, but is used
 	 *                                 to safely scope usage of the IAudioEndpointSettingsProxy pointer.
 	 */
-	void PollSettings(TFunctionRef<void(const IAudioEndpointSettingsProxy*)> NewSettingsRetrieved);
+	AUDIOEXTENSIONS_API void PollSettings(TFunctionRef<void(const IAudioEndpointSettingsProxy*)> NewSettingsRetrieved);
 
 	/**
 	 * Thread safe function to disconnect everything from this endpoint. 
 	 * Anything that owns an Audio::FPatchInput will be notified and will have to call PatchNewInput() again to reconnect.
 	 */
-	void DisconnectAllInputs();
+	AUDIOEXTENSIONS_API void DisconnectAllInputs();
 
 	/**
 	 * If EndpointRequiresCallback() returns true, this can be used to spawn an async thread and begin calling OnAudioCallback.
 	 */
-	void StartRunningAsyncCallback();
-	void StopRunningAsyncCallback();
+	AUDIOEXTENSIONS_API void StartRunningAsyncCallback();
+	AUDIOEXTENSIONS_API void StopRunningAsyncCallback();
 
 	/**
 	 * If EndpointRequiresCallback() returns true, this can be used to manually run the callback.
 	 */
-	void RunCallbackSynchronously();
+	AUDIOEXTENSIONS_API void RunCallbackSynchronously();
 
 
 private:
@@ -198,7 +198,7 @@ private:
  * Once a factory is constructed and RegisterEndpointType is called, it will be exposed as a type of endpoint
  * That a submix in the submix graph could be constructed with.
  */
-class AUDIOEXTENSIONS_API IAudioEndpointFactory : public IModularFeature
+class IAudioEndpointFactory : public IModularFeature
 {
 public:
 	/** Virtual destructor */
@@ -207,50 +207,50 @@ public:
 	}
 
 	/** Get the name for the endpoint type that this factory produces.  */
-	virtual FName GetEndpointTypeName();
+	AUDIOEXTENSIONS_API virtual FName GetEndpointTypeName();
 
 	/** This is a special cased name for endpoint submixes that render directly to the default audio device in Audio::FMixerDevice::OnProcessAudioStream. */
-	static FName GetTypeNameForDefaultEndpoint();
+	static AUDIOEXTENSIONS_API FName GetTypeNameForDefaultEndpoint();
 
 	/** 
 	 * This is used when calling IModularFeatures::Get().RegisterModularFeature for IAudioEndpointFactory implementations. 
 	 * It's not needed if one uses RegisterEndpointType() to register IAudioEndpointFactory implementations. 
 	 */
-	static FName GetModularFeatureName();
+	static AUDIOEXTENSIONS_API FName GetModularFeatureName();
 
 	/** 
 	 * This needs to be called to make a soundfield format usable by the engine.
 	 * It can be called from a ISoundfieldFactory subclass' constructor
 	*/
-	static void RegisterEndpointType(IAudioEndpointFactory* InFactory);
+	static AUDIOEXTENSIONS_API void RegisterEndpointType(IAudioEndpointFactory* InFactory);
 
 	/**
 	 * This needs to be called it an implementation of ISoundfieldFactory is about to be destroyed.
 	 * It can be called from the destructor of an implementation of ISoundfieldFactory.
 	 */
-	static void UnregisterEndpointType(IAudioEndpointFactory* InFactory);
+	static AUDIOEXTENSIONS_API void UnregisterEndpointType(IAudioEndpointFactory* InFactory);
 
 	/**
 	 * Get a registered endpoint factory by name.
 	 */
-	static IAudioEndpointFactory* Get(const FName& InName);
+	static AUDIOEXTENSIONS_API IAudioEndpointFactory* Get(const FName& InName);
 
-	static TArray<FName> GetAvailableEndpointTypes();
+	static AUDIOEXTENSIONS_API TArray<FName> GetAvailableEndpointTypes();
 
 	/** Called for every new endpoint submix created with this factory's endpoint type. */
-	virtual TUniquePtr<IAudioEndpoint> CreateNewEndpointInstance(const FAudioPluginInitializationParams& InitInfo, const IAudioEndpointSettingsProxy& InitialSettings);
+	AUDIOEXTENSIONS_API virtual TUniquePtr<IAudioEndpoint> CreateNewEndpointInstance(const FAudioPluginInitializationParams& InitInfo, const IAudioEndpointSettingsProxy& InitialSettings);
 
 	/**
 	 * Should return the StaticClass of this factory's implementation of UAudioEndpointSettingsBase.
 	 */
-	virtual UClass* GetCustomSettingsClass() const;
+	AUDIOEXTENSIONS_API virtual UClass* GetCustomSettingsClass() const;
 
 	/**
 	 * return the settings an endpoint should use 
 	 */
-	virtual const UAudioEndpointSettingsBase* GetDefaultSettings() const;
+	AUDIOEXTENSIONS_API virtual const UAudioEndpointSettingsBase* GetDefaultSettings() const;
 
 	bool bIsImplemented = false;
 
-	static IAudioEndpointFactory* GetDummyFactory();
-};
\ No newline at end of file
+	static AUDIOEXTENSIONS_API IAudioEndpointFactory* GetDummyFactory();
+};
diff --git a/Engine/Source/Runtime/AudioExtensions/Public/IAudioExtensionPlugin.h b/Engine/Source/Runtime/AudioExtensions/Public/IAudioExtensionPlugin.h
--- a/Engine/Source/Runtime/AudioExtensions/Public/IAudioExtensionPlugin.h
+++ b/Engine/Source/Runtime/AudioExtensions/Public/IAudioExtensionPlugin.h
@@ -180,8 +180,8 @@ struct FAudioPluginSourceOutputData
 };
 
 /** This is a class which should be overridden to provide users with settings to use for individual sounds */
-UCLASS(config = Engine, abstract, editinlinenew, BlueprintType)
-class AUDIOEXTENSIONS_API USpatializationPluginSourceSettingsBase : public UObject
+UCLASS(config = Engine, abstract, editinlinenew, BlueprintType, MinimalAPI)
+class USpatializationPluginSourceSettingsBase : public UObject
 {
 	GENERATED_BODY()
 };
@@ -380,8 +380,8 @@ public:
 
 
 /** This is a class which should be overridden to provide users with settings to use for individual sounds */
-UCLASS(config = Engine, abstract, editinlinenew, BlueprintType)
-class AUDIOEXTENSIONS_API USourceDataOverridePluginSourceSettingsBase : public UObject
+UCLASS(config = Engine, abstract, editinlinenew, BlueprintType, MinimalAPI)
+class USourceDataOverridePluginSourceSettingsBase : public UObject
 {
 	GENERATED_BODY()
 };
@@ -465,8 +465,8 @@ public:
 };
 
 /** This is a class which should be overridden to provide users with settings to use for individual sounds */
-UCLASS(config = Engine, abstract, editinlinenew, BlueprintType)
-class AUDIOEXTENSIONS_API UOcclusionPluginSourceSettingsBase : public UObject
+UCLASS(config = Engine, abstract, editinlinenew, BlueprintType, MinimalAPI)
+class UOcclusionPluginSourceSettingsBase : public UObject
 {
 	GENERATED_BODY()
 };
@@ -580,8 +580,8 @@ public:
 
 
 /** This is a class which should be overridden to provide users with settings to use for individual sounds */
-UCLASS(config = Engine, abstract, editinlinenew, BlueprintType)
-class AUDIOEXTENSIONS_API UReverbPluginSourceSettingsBase : public UObject
+UCLASS(config = Engine, abstract, editinlinenew, BlueprintType, MinimalAPI)
+class UReverbPluginSourceSettingsBase : public UObject
 {
 	GENERATED_BODY()
 };
diff --git a/Engine/Source/Runtime/AudioExtensions/Public/IAudioModulation.h b/Engine/Source/Runtime/AudioExtensions/Public/IAudioModulation.h
--- a/Engine/Source/Runtime/AudioExtensions/Public/IAudioModulation.h
+++ b/Engine/Source/Runtime/AudioExtensions/Public/IAudioModulation.h
@@ -48,14 +48,14 @@ namespace Audio
 	using FModulationMixFunction = TFunction<void(float& /* OutNormalizedA */, float /* InNormalizedB */)>;
 
 
-	struct AUDIOEXTENSIONS_API FModulationParameter
+	struct FModulationParameter
 	{
-		FModulationParameter();
-		FModulationParameter(const FModulationParameter& InParam);
-		FModulationParameter(FModulationParameter&& InParam);
+		AUDIOEXTENSIONS_API FModulationParameter();
+		AUDIOEXTENSIONS_API FModulationParameter(const FModulationParameter& InParam);
+		AUDIOEXTENSIONS_API FModulationParameter(FModulationParameter&& InParam);
 
-		FModulationParameter& operator=(FModulationParameter&& InParam);
-		FModulationParameter& operator=(const FModulationParameter& InParam);
+		AUDIOEXTENSIONS_API FModulationParameter& operator=(FModulationParameter&& InParam);
+		AUDIOEXTENSIONS_API FModulationParameter& operator=(const FModulationParameter& InParam);
 
 		FName ParameterName;
 
@@ -86,9 +86,9 @@ namespace Audio
 		// Function used to convert value buffer from unit space to normalized, unitless [0.0f, 1.0f] space.
 		FModulationNormalizedConversionFunction NormalizedFunction;
 
-		static const FModulationMixFunction& GetDefaultMixFunction();
-		static const FModulationUnitConversionFunction& GetDefaultUnitConversionFunction();
-		static const FModulationNormalizedConversionFunction& GetDefaultNormalizedConversionFunction();
+		static AUDIOEXTENSIONS_API const FModulationMixFunction& GetDefaultMixFunction();
+		static AUDIOEXTENSIONS_API const FModulationUnitConversionFunction& GetDefaultUnitConversionFunction();
+		static AUDIOEXTENSIONS_API const FModulationNormalizedConversionFunction& GetDefaultNormalizedConversionFunction();
 	};
 
 	AUDIOEXTENSIONS_API bool IsModulationParameterRegistered(FName InName);
@@ -102,7 +102,7 @@ namespace Audio
 	  * If proxy is already active, implementation is expected to ignore register call
 	  * and return existing modulator proxy's type Id & set parameter accordingly.
 	  */
-	class AUDIOEXTENSIONS_API IModulatorSettings
+	class IModulatorSettings
 	{
 	public:
 		virtual ~IModulatorSettings() = default;
@@ -117,26 +117,26 @@ namespace Audio
 	/** Handle to a modulator which interacts with the modulation API to manage lifetime
 	  * of modulator proxy objects internal to modulation plugin implementation.
 	  */
-	struct AUDIOEXTENSIONS_API FModulatorHandle
+	struct FModulatorHandle
 	{
 		FModulatorHandle() = default;
-		FModulatorHandle(Audio::FModulationParameter&& InParameter);
-		FModulatorHandle(IAudioModulationManager& InModulation, const Audio::IModulatorSettings& InModulatorSettings, Audio::FModulationParameter&& InParameter);
-		FModulatorHandle(const FModulatorHandle& InOther);
-		FModulatorHandle(FModulatorHandle&& InOther);
+		AUDIOEXTENSIONS_API FModulatorHandle(Audio::FModulationParameter&& InParameter);
+		AUDIOEXTENSIONS_API FModulatorHandle(IAudioModulationManager& InModulation, const Audio::IModulatorSettings& InModulatorSettings, Audio::FModulationParameter&& InParameter);
+		AUDIOEXTENSIONS_API FModulatorHandle(const FModulatorHandle& InOther);
+		AUDIOEXTENSIONS_API FModulatorHandle(FModulatorHandle&& InOther);
 
-		~FModulatorHandle();
+		AUDIOEXTENSIONS_API ~FModulatorHandle();
 
-		FModulatorHandle& operator=(const FModulatorHandle& InOther);
-		FModulatorHandle& operator=(FModulatorHandle&& InOther);
+		AUDIOEXTENSIONS_API FModulatorHandle& operator=(const FModulatorHandle& InOther);
+		AUDIOEXTENSIONS_API FModulatorHandle& operator=(FModulatorHandle&& InOther);
 
-		FModulatorId GetModulatorId() const;
-		const FModulationParameter& GetParameter() const;
-		FModulatorTypeId GetTypeId() const;
-		FModulatorHandleId GetHandleId() const;
-		bool GetValue(float& OutValue) const;
-		bool GetValueThreadSafe(float& OutValue) const;
-		bool IsValid() const;
+		AUDIOEXTENSIONS_API FModulatorId GetModulatorId() const;
+		AUDIOEXTENSIONS_API const FModulationParameter& GetParameter() const;
+		AUDIOEXTENSIONS_API FModulatorTypeId GetTypeId() const;
+		AUDIOEXTENSIONS_API FModulatorHandleId GetHandleId() const;
+		AUDIOEXTENSIONS_API bool GetValue(float& OutValue) const;
+		AUDIOEXTENSIONS_API bool GetValueThreadSafe(float& OutValue) const;
+		AUDIOEXTENSIONS_API bool IsValid() const;
 
 		friend FORCEINLINE uint32 GetTypeHash(const FModulatorHandle& InModulatorHandle)
 		{
@@ -162,7 +162,7 @@ namespace Audio
 	};
 } // namespace Audio
 
-class AUDIOEXTENSIONS_API IAudioModulationManager : public TSharedFromThis<IAudioModulationManager>
+class IAudioModulationManager : public TSharedFromThis<IAudioModulationManager>
 {
 public:
 	/** Virtual destructor */
@@ -207,23 +207,23 @@ protected:
 /**
  * Base class for all modulators
  */
-UCLASS(config = Engine, abstract, editinlinenew, BlueprintType)
-class AUDIOEXTENSIONS_API USoundModulatorBase : public UObject, public IAudioProxyDataFactory
+UCLASS(config = Engine, abstract, editinlinenew, BlueprintType, MinimalAPI)
+class USoundModulatorBase : public UObject, public IAudioProxyDataFactory
 {
 	GENERATED_BODY()
 
 public:
-	virtual const Audio::FModulationParameter& GetOutputParameter() const;
+	AUDIOEXTENSIONS_API virtual const Audio::FModulationParameter& GetOutputParameter() const;
 
-	virtual TSharedPtr<Audio::IProxyData> CreateProxyData(const Audio::FProxyDataInitParams& InitParams) override;
+	AUDIOEXTENSIONS_API virtual TSharedPtr<Audio::IProxyData> CreateProxyData(const Audio::FProxyDataInitParams& InitParams) override;
 
-	virtual TUniquePtr<Audio::IModulatorSettings> CreateProxySettings() const;
+	AUDIOEXTENSIONS_API virtual TUniquePtr<Audio::IModulatorSettings> CreateProxySettings() const;
 };
 
 /** Proxy to modulator, allowing for modulator to be referenced by the Audio Render Thread independently
   * from the implementing modulation plugin (ex. for MetaSound implementation).
   */
-class AUDIOEXTENSIONS_API FSoundModulatorAssetProxy : public Audio::TProxyData<FSoundModulatorAssetProxy>, public TSharedFromThis<FSoundModulatorAssetProxy, ESPMode::ThreadSafe>
+class FSoundModulatorAssetProxy : public Audio::TProxyData<FSoundModulatorAssetProxy>, public TSharedFromThis<FSoundModulatorAssetProxy, ESPMode::ThreadSafe>
 {
 public:
 	IMPL_AUDIOPROXY_CLASS(FSoundModulatorAssetProxy);
@@ -263,7 +263,7 @@ using FSoundModulatorAssetProxyPtr = TSharedPtr<FSoundModulatorAssetProxy, ESPMo
 /** Proxy to modulator, allowing for modulator to be referenced by the Audio Render Thread independently
   * from the implementing modulation plugin (ex. for MetaSound implementation).
   */
-class AUDIOEXTENSIONS_API FSoundModulationParameterAssetProxy : public Audio::TProxyData<FSoundModulationParameterAssetProxy>, public TSharedFromThis<FSoundModulationParameterAssetProxy, ESPMode::ThreadSafe>
+class FSoundModulationParameterAssetProxy : public Audio::TProxyData<FSoundModulationParameterAssetProxy>, public TSharedFromThis<FSoundModulationParameterAssetProxy, ESPMode::ThreadSafe>
 {
 public:
 	IMPL_AUDIOPROXY_CLASS(FSoundModulationParameterAssetProxy);
@@ -281,7 +281,7 @@ using FSoundModulationParameterAssetProxyPtr = TSharedPtr<FSoundModulationParame
 /** Interface to sound that is modulatable, allowing for certain specific
   * behaviors to be controlled on the sound level by the modulation system.
   */
-class AUDIOEXTENSIONS_API ISoundModulatable
+class ISoundModulatable
 {
 public:
 	virtual ~ISoundModulatable() = default;
diff --git a/Engine/Source/Runtime/AudioExtensions/Public/IAudioParameterInterfaceRegistry.h b/Engine/Source/Runtime/AudioExtensions/Public/IAudioParameterInterfaceRegistry.h
--- a/Engine/Source/Runtime/AudioExtensions/Public/IAudioParameterInterfaceRegistry.h
+++ b/Engine/Source/Runtime/AudioExtensions/Public/IAudioParameterInterfaceRegistry.h
@@ -26,17 +26,17 @@ namespace Audio
 	// Can be used generically for processing either logically at a control or Game Thread
 	// tick rate (ex. SoundCues), or by an underlying DSP operation via the Audio Render
 	// Thread or delegated task (ex. MetaSounds, custom SoundGenerator, etc.)
-	struct AUDIOEXTENSIONS_API FParameterInterface
+	struct FParameterInterface
 	{
 		// Version of interface (higher numbers are more recent)
-		struct AUDIOEXTENSIONS_API FVersion
+		struct FVersion
 		{
 			const int32 Major = 1;
 			const int32 Minor = 0;
 		};
 
 		// Input of interface
-		struct AUDIOEXTENSIONS_API FInput
+		struct FInput
 		{
 			// Name to be displayed in editor or tools
 			const FText DisplayName;
@@ -66,7 +66,7 @@ namespace Audio
 		};
 
 		// Output of interface
-		struct AUDIOEXTENSIONS_API FOutput
+		struct FOutput
 		{
 			// Name to be displayed in editor or tools
 			const FText DisplayName;
@@ -98,7 +98,7 @@ namespace Audio
 
 		// Read-only variable that cannot be modified by the sound instance,
 		// and maybe shared amongst instances.
-		struct AUDIOEXTENSIONS_API FEnvironmentVariable
+		struct FEnvironmentVariable
 		{
 			// Name to be displayed in editor or tools
 			const FText DisplayName;
@@ -120,7 +120,7 @@ namespace Audio
 		};
 
 		// Options used to restrict a corresponding UClass that interface may be applied to.
-		struct AUDIOEXTENSIONS_API FClassOptions
+		struct FClassOptions
 		{
 			// Path to restricted UClass
 			const FTopLevelAssetPath ClassPath;
@@ -137,34 +137,34 @@ namespace Audio
 		FParameterInterface() = default;
 
 		// Constructor used for parameter interface not limited to any particular UClass types
-		FParameterInterface(FName InName, const FVersion& InVersion);
+		AUDIOEXTENSIONS_API FParameterInterface(FName InName, const FVersion& InVersion);
 
 		UE_DEPRECATED(5.3, "Set UClassOptions to determine what options apply for a given UClass (if any).")
-		FParameterInterface(FName InName, const FVersion& InVersion, const UClass& InClass);
+		AUDIOEXTENSIONS_API FParameterInterface(FName InName, const FVersion& InVersion, const UClass& InClass);
 
 		// Returns name of interface
-		FName GetName() const;
+		AUDIOEXTENSIONS_API FName GetName() const;
 
 		// Returns version of interface
-		const FVersion& GetVersion() const;
+		AUDIOEXTENSIONS_API const FVersion& GetVersion() const;
 
 		UE_DEPRECATED(5.3, "Use FParameterInterface::FindSupportedUClasses instead")
-		const UClass& GetType() const;
+		AUDIOEXTENSIONS_API const UClass& GetType() const;
 
-		TArray<const UClass*> FindSupportedUClasses() const;
+		AUDIOEXTENSIONS_API TArray<const UClass*> FindSupportedUClasses() const;
 
 		// If specified, options used to restrict a corresponding UClass that interface may be
 		// applied to.  If unspecified, interface is assumed to be applicable to any arbitrary UClass.
-		const TArray<FClassOptions>& GetUClassOptions() const;
+		AUDIOEXTENSIONS_API const TArray<FClassOptions>& GetUClassOptions() const;
 
 		// Returns read-only array of inputs
-		const TArray<FInput>& GetInputs() const;
+		AUDIOEXTENSIONS_API const TArray<FInput>& GetInputs() const;
 
 		// Returns read-only array of outputs
-		const TArray<FOutput>& GetOutputs() const;
+		AUDIOEXTENSIONS_API const TArray<FOutput>& GetOutputs() const;
 
 		// Returns read-only array of environment variables
-		const TArray<FEnvironmentVariable>& GetEnvironment() const;
+		AUDIOEXTENSIONS_API const TArray<FEnvironmentVariable>& GetEnvironment() const;
 
 	private:
 		FName NamePrivate;
@@ -181,12 +181,12 @@ namespace Audio
 	// Registry of engine-defined audio parameter interfaces, used to parameterize data provided
 	// to or coming from an executable audio unit (ex. Sound like a MetaSoundSource, arbitrary DSP
 	// graph like a MetaSoundPatch, etc.).
-	class AUDIOEXTENSIONS_API IAudioParameterInterfaceRegistry
+	class IAudioParameterInterfaceRegistry
 	{
-		static TUniquePtr<IAudioParameterInterfaceRegistry> Instance;
+		static AUDIOEXTENSIONS_API TUniquePtr<IAudioParameterInterfaceRegistry> Instance;
 
 	public:
-		static IAudioParameterInterfaceRegistry& Get();
+		static AUDIOEXTENSIONS_API IAudioParameterInterfaceRegistry& Get();
 
 		virtual ~IAudioParameterInterfaceRegistry() = default;
 
diff --git a/Engine/Source/Runtime/AudioExtensions/Public/IAudioParameterTransmitter.h b/Engine/Source/Runtime/AudioExtensions/Public/IAudioParameterTransmitter.h
--- a/Engine/Source/Runtime/AudioExtensions/Public/IAudioParameterTransmitter.h
+++ b/Engine/Source/Runtime/AudioExtensions/Public/IAudioParameterTransmitter.h
@@ -17,7 +17,7 @@ namespace Audio
 	using DeviceID = uint32;
 
 	/** Data passed to CreateParameterTransmitter. */
-	struct AUDIOEXTENSIONS_API FParameterTransmitterInitParams
+	struct FParameterTransmitterInitParams
 	{
 		// Unique ID for this audio instance.
 		uint64 InstanceID = INDEX_NONE;
@@ -38,22 +38,22 @@ namespace Audio
 	// as object pointers within parameters should NOT be cached on threads
 	// other than the GameThread, utilizing a proxy methodology like MetaSounds
 	// that copies UObject data when and where necessary.
-	class AUDIOEXTENSIONS_API ILegacyParameterTransmitter
+	class ILegacyParameterTransmitter
 	{
 		public:
 			virtual ~ILegacyParameterTransmitter() = default;
 
-			virtual TArray<const TObjectPtr<UObject>*> GetReferencedObjects() const;
+			AUDIOEXTENSIONS_API virtual TArray<const TObjectPtr<UObject>*> GetReferencedObjects() const;
 	};
 
 	/** Interface for a audio instance transmitter.
 	 *
 	 * An audio instance transmitter ushers control parameters to a single audio object instance.
 	 */
-	class AUDIOEXTENSIONS_API IParameterTransmitter : public ILegacyParameterTransmitter
+	class IParameterTransmitter : public ILegacyParameterTransmitter
 	{
 		public:
-			static const FName RouterName;
+			static AUDIOEXTENSIONS_API const FName RouterName;
 
 			virtual ~IParameterTransmitter() = default;
 
@@ -81,18 +81,18 @@ namespace Audio
 	/** Base implementation for the parameter transmitter, which caches parameters
 	  * and provides implementer to add additional logic to route parameter data accordingly.
 	  */
-	class AUDIOEXTENSIONS_API FParameterTransmitterBase : public IParameterTransmitter
+	class FParameterTransmitterBase : public IParameterTransmitter
 	{
 	public:
-		FParameterTransmitterBase(TArray<FAudioParameter>&& InDefaultParams);
+		AUDIOEXTENSIONS_API FParameterTransmitterBase(TArray<FAudioParameter>&& InDefaultParams);
 		virtual ~FParameterTransmitterBase() = default;
 
-		virtual bool GetParameter(FName InName, FAudioParameter& OutParam) const override;
-		virtual void ResetParameters() override;
-		virtual const TArray<FAudioParameter>& GetParameters() const override;
+		AUDIOEXTENSIONS_API virtual bool GetParameter(FName InName, FAudioParameter& OutParam) const override;
+		AUDIOEXTENSIONS_API virtual void ResetParameters() override;
+		AUDIOEXTENSIONS_API virtual const TArray<FAudioParameter>& GetParameters() const override;
 		UE_DEPRECATED(5.2, "Use ResetParameters() or OnDeleteActiveSound() instead depending on use case.")
-		virtual bool Reset() override;
-		virtual bool SetParameters(TArray<FAudioParameter>&& InParameters) override;
+		AUDIOEXTENSIONS_API virtual bool Reset() override;
+		AUDIOEXTENSIONS_API virtual bool SetParameters(TArray<FAudioParameter>&& InParameters) override;
 
 	protected:
 		TArray<FAudioParameter> AudioParameters;
diff --git a/Engine/Source/Runtime/AudioExtensions/Public/IAudioProxyInitializer.h b/Engine/Source/Runtime/AudioExtensions/Public/IAudioProxyInitializer.h
--- a/Engine/Source/Runtime/AudioExtensions/Public/IAudioProxyInitializer.h
+++ b/Engine/Source/Runtime/AudioExtensions/Public/IAudioProxyInitializer.h
@@ -124,13 +124,13 @@ namespace Audio
 /*
 * This can be subclassed to make a UClass an audio proxy factory.
 */
-class AUDIOEXTENSIONS_API IAudioProxyDataFactory
+class IAudioProxyDataFactory
 {
 public:
 	UE_DEPRECATED(5.2, "Call TSharedPtr<Audio::IProxyData> CreateProxyData(...) instead of a TUniquePtr<Audio::IProxyData> CreateNewProxyData(...).")
-	virtual TUniquePtr<Audio::IProxyData> CreateNewProxyData(const Audio::FProxyDataInitParams& InitParams);
+	AUDIOEXTENSIONS_API virtual TUniquePtr<Audio::IProxyData> CreateNewProxyData(const Audio::FProxyDataInitParams& InitParams);
 
-	virtual TSharedPtr<Audio::IProxyData> CreateProxyData(const Audio::FProxyDataInitParams& InitParams);
+	AUDIOEXTENSIONS_API virtual TSharedPtr<Audio::IProxyData> CreateProxyData(const Audio::FProxyDataInitParams& InitParams);
 };
 
 namespace Audio
