Downgrade error to warning so it doesn't prevent builds being packaged.

diff --git a/Engine/Source/Runtime/AssetRegistry/Private/AssetRegistry.cpp b/Engine/Source/Runtime/AssetRegistry/Private/AssetRegistry.cpp
--- a/Engine/Source/Runtime/AssetRegistry/Private/AssetRegistry.cpp
+++ b/Engine/Source/Runtime/AssetRegistry/Private/AssetRegistry.cpp
@@ -6693,7 +6693,10 @@ void FAssetRegistryImpl::SetManageReferences(const TMultiMap<FAssetIdentifier, F
 
 		if (!ManagedNode)
 		{
-			UE_LOG(LogAssetRegistry, Error, TEXT("Cannot set %s to manage asset %s because %s does not exist!"), *Pair.Key.ToString(), *Pair.Value.ToString(), *Pair.Value.ToString());
+			UE_LOG(LogAssetRegistry, Warning, TEXT("Cannot set %s to manage asset %s because %s does not exist!"), 
+				*Pair.Key.ToString(),
+				*Pair.Value.ToString(),
+				*Pair.Value.ToString());
 			continue;
 		}
 
