Allow AssetGatherAll and AssetGatherSync to be used by a shipping editor

diff --git a/Engine/Source/Runtime/AssetRegistry/Private/AssetDataGatherer.cpp b/Engine/Source/Runtime/AssetRegistry/Private/AssetDataGatherer.cpp
--- a/Engine/Source/Runtime/AssetRegistry/Private/AssetDataGatherer.cpp
+++ b/Engine/Source/Runtime/AssetRegistry/Private/AssetDataGatherer.cpp
@@ -3204,13 +3204,13 @@ FAssetDataGatherer::FAssetDataGatherer(const TArray<FString>& InLongPackageNames
 	bWriteMonolithicCache = bCacheWriteEnabled && GPreloadSettings.IsMonolithicCacheActivatedDuringPreload();
 	LastCacheWriteTime = FPlatformTime::Seconds();
 
-#if !UE_BUILD_SHIPPING
+#if WITH_EDITOR || !UE_BUILD_SHIPPING
 	bool bCommandlineSynchronous;
 	if (FParse::Bool(FCommandLine::Get(), TEXT("AssetGatherSync="), bCommandlineSynchronous))
 	{
 		bAsyncEnabled = !bCommandlineSynchronous;
 	}
-#endif
+#endif // WITH_EDITOR || !UE_BUILD_SHIPPING
 	if (bAsyncEnabled && !FPlatformProcess::SupportsMultithreading())
 	{
 		bAsyncEnabled = false;
diff --git a/Engine/Source/Runtime/AssetRegistry/Private/AssetRegistry.cpp b/Engine/Source/Runtime/AssetRegistry/Private/AssetRegistry.cpp
--- a/Engine/Source/Runtime/AssetRegistry/Private/AssetRegistry.cpp
+++ b/Engine/Source/Runtime/AssetRegistry/Private/AssetRegistry.cpp
@@ -7439,13 +7439,13 @@ bool ShouldSearchAllAssetsAtStart()
 			}
 		}
 	}
-#if !UE_BUILD_SHIPPING
+#if WITH_EDITOR || !UE_BUILD_SHIPPING
 	bool bCommandlineAllAssetsAtStart;
 	if (FParse::Bool(FCommandLine::Get(), TEXT("AssetGatherAll="), bCommandlineAllAssetsAtStart))
 	{
 		bSearchAllAssetsAtStart = bCommandlineAllAssetsAtStart;
 	}
-#endif // !UE_BUILD_SHIPPING
+#endif // WITH_EDITOR || !UE_BUILD_SHIPPING
 	return bSearchAllAssetsAtStart;
 }
 
