Telemetry for asset registry scan. Not sent to any telemetry endpoint, provided to FTelemetryRouter for internal analytics or licensees' own analytics.

diff --git a/Engine/Source/Runtime/AssetRegistry/AssetRegistry.Build.cs b/Engine/Source/Runtime/AssetRegistry/AssetRegistry.Build.cs
--- a/Engine/Source/Runtime/AssetRegistry/AssetRegistry.Build.cs
+++ b/Engine/Source/Runtime/AssetRegistry/AssetRegistry.Build.cs
@@ -13,6 +13,7 @@ public class AssetRegistry : ModuleRules
 				"TraceLog",
 				"ApplicationCore",
 				"Projects",
+				"TelemetryUtils"
 			}
 			);
 
diff --git a/Engine/Source/Runtime/AssetRegistry/Private/AssetDataGatherer.cpp b/Engine/Source/Runtime/AssetRegistry/Private/AssetDataGatherer.cpp
--- a/Engine/Source/Runtime/AssetRegistry/Private/AssetDataGatherer.cpp
+++ b/Engine/Source/Runtime/AssetRegistry/Private/AssetDataGatherer.cpp
@@ -3184,8 +3184,6 @@ FAssetDataGatherer::FAssetDataGatherer(const TArray<FString>& InLongPackageNames
 	, bFinishedInitialDiscovery(false)
 	, bAllModuleLoadingPhasesComplete(false)
 	, WaitBatchCount(0)
-	, NumCachedAssetFiles(0)
-	, NumUncachedAssetFiles(0)
 	, LastMonolithicCacheSaveUncachedAssetFiles(0)
 	, CacheInUseCount(0)
 	, bIsSavingAsyncCache(false)
@@ -3613,6 +3611,8 @@ void FAssetDataGatherer::TickInternal(bool& bOutIsTickInterrupt, double& TickSta
 		}
 	};
 
+	int32 NewCachedAssetFiles = 0;
+	int32 NewUncachedAssetFiles = 0;
 	// Try to read each file in the batch out of the cache, and accumulate a list for more expensive reading of all of the files that are not in the cache 
 	TArray<FReadContext> ReadContexts;
 	for (FGatheredPathData& AssetFileData : LocalFilesToSearch)
@@ -3662,7 +3662,7 @@ void FAssetDataGatherer::TickInternal(bool& bOutIsTickInterrupt, double& TickSta
 		if (DiskCachedAssetData)
 		{
 			// Add the valid cached data to our results, and to the map of data we keep to write out the new version of the cache file
-			++NumCachedAssetFiles;
+			++NewCachedAssetFiles;
 
 			// Set the transient flags based on whether our current cache has dependency data.
 			// Note that in editor, bGatherAssetPackageData is always true, no way to turn it off,
@@ -3713,7 +3713,7 @@ void FAssetDataGatherer::TickInternal(bool& bOutIsTickInterrupt, double& TickSta
 		}
 		else if (ReadContext.bResult)
 		{
-			++NumUncachedAssetFiles;
+			++NewUncachedAssetFiles;
 
 			// Add the results from a cooked package into our results on cooked package
 			LocalCookedPackageNamesWithoutAssetDataResults.Append(MoveTemp(ReadContext.CookedPackageNamesWithoutAssetData));
@@ -3773,6 +3773,9 @@ void FAssetDataGatherer::TickInternal(bool& bOutIsTickInterrupt, double& TickSta
 		VerseResults.Append(MoveTemp(LocalVerseResults));
 		BlockedResults.Append(MoveTemp(LocalBlockedResults));
 
+		NumUncachedAssetFiles += NewUncachedAssetFiles;
+		NumCachedAssetFiles += NewCachedAssetFiles;
+
 		if (bHasCancelation)
 		{
 			// If we skipped reading files due to a pause request, push the canceled files back onto the FilesToSearch
@@ -3955,11 +3958,15 @@ void FAssetDataGatherer::GetAndTrimSearchResults(FResults& InOutResults, FResult
 	OutContext.bAbleToProgress = !bIsIdle;
 }
 
-void FAssetDataGatherer::GetDiagnostics(float& OutGatherTimeSeconds, float& OutDiscoverTimeSeconds)
+FAssetGatherDiagnostics FAssetDataGatherer::GetDiagnostics()
 {
-	Discovery->GetDiagnostics(OutDiscoverTimeSeconds);
+	FAssetGatherDiagnostics Diag;
+	Discovery->GetDiagnostics(Diag.DiscoveryTimeSeconds);
 	FGathererScopeLock ResultsScopeLock(&ResultsLock);
-	OutGatherTimeSeconds = CumulativeGatherTime;
+	Diag.GatherTimeSeconds = CumulativeGatherTime;
+	Diag.NumCachedAssetFiles = NumCachedAssetFiles;
+	Diag.NumUncachedAssetFiles = NumUncachedAssetFiles;
+	return Diag;
 }
 
 void FAssetDataGatherer::GetPackageResults(TMultiMap<FName, FAssetData*>& OutAssetResults, TMultiMap<FName, FPackageDependencyData>& OutDependencyResults)
@@ -4265,15 +4272,17 @@ void FAssetDataGatherer::TryReserveSaveMonolithicCache(bool& bOutShouldSave, TAr
 		return;
 	}
 	CHECK_IS_LOCKED_CURRENT_THREAD(TickLock);
+	int32 LocalNumUncachedAssetFiles;
 	{
 		FGathererScopeLock ResultsScopeLock(&ResultsLock);
 		bOutShouldSave = bWriteMonolithicCache;
+		LocalNumUncachedAssetFiles = NumUncachedAssetFiles;
 	}
 	if (bOutShouldSave)
 	{
 		GetMonolithicCacheAssetsToSave(AssetsToSave);
 		bIsSavingAsyncCache = true;
-		LastMonolithicCacheSaveUncachedAssetFiles = NumUncachedAssetFiles;
+		LastMonolithicCacheSaveUncachedAssetFiles = LocalNumUncachedAssetFiles ;
 	}
 	bSaveAsyncCacheTriggered = false;
 }
diff --git a/Engine/Source/Runtime/AssetRegistry/Private/AssetDataGatherer.h b/Engine/Source/Runtime/AssetRegistry/Private/AssetDataGatherer.h
--- a/Engine/Source/Runtime/AssetRegistry/Private/AssetDataGatherer.h
+++ b/Engine/Source/Runtime/AssetRegistry/Private/AssetDataGatherer.h
@@ -52,6 +52,18 @@ typedef FScopeLock FGathererScopeLock;
 #define CHECK_IS_NOT_LOCKED_CURRENT_THREAD(CritSec) do {} while (false)
 #endif
 
+struct FAssetGatherDiagnostics
+{
+	/** Time spent identifying asset files on disk */
+	float DiscoveryTimeSeconds;
+	/** Time spent reading asset files on disk / from cache */
+	float GatherTimeSeconds;
+	/** How many files in the search results were read from the cache. */
+	int32 NumCachedAssetFiles;
+	/** How many files in the search results were not in the cache and were read by parsing the file. */
+	int32 NumUncachedAssetFiles;
+};
+
 /**
  * Async task for gathering asset data from from the file list in FAssetRegistry
  */
@@ -124,7 +136,7 @@ public:
 	/** Gets search results from the data gatherer. */
 	void GetAndTrimSearchResults(FResults& InOutResults, FResultContext& OutContext);
 	/** Get diagnostics for telemetry or logging. */
-	void GetDiagnostics(float& OutGatherTimeSeconds, float& OutDiscoverTimeSeconds);
+	FAssetGatherDiagnostics GetDiagnostics();
 	/** Gets just the AssetResults and DependencyResults from the data gatherer. */
 	void GetPackageResults(TMultiMap<FName, FAssetData*>& OutAssetResults,
 		TMultiMap<FName, FPackageDependencyData>& OutDependencyResults);
@@ -424,6 +436,10 @@ private:
 	double LastCacheWriteTime;
 	/** The cached value of the NumPathsToSearch returned by Discovery the last time we synchronized with it. */
 	int32 NumPathsToSearchAtLastSyncPoint;
+	/** The total number of files in the search results that were read from the cache. */
+	int32 NumCachedAssetFiles = 0;
+	/** The total number of files in the search results that were not in the cache and were read by parsing the file. */
+	int32 NumUncachedAssetFiles = 0;
 	/**
 	 * Track whether we are allowed to read from a monolithic cache that should be loaded during tick.
 	 * Even if we are or not, if bCacheReadEnabled the AssetRegistry can also call LoadCacheFile/ScanPathsSynchronous to
@@ -469,10 +485,6 @@ private:
 	TMap<FName, FDiskCachedAssetData*> NewCachedAssetDataMap;
 	/** Used to block on gather results. If non-zero, tick should end when WaitBatchCount files have been processed. */
 	int32 WaitBatchCount;
-	/** How many files in the search results were read from the cache. */
-	int32 NumCachedAssetFiles;
-	/** How many files in the search results were not in the cache and were read by parsing the file. */
-	int32 NumUncachedAssetFiles;
 	/** How many uncached asset files had been discovered at the last async cache save */
 	int32 LastMonolithicCacheSaveUncachedAssetFiles;
 	/**
diff --git a/Engine/Source/Runtime/AssetRegistry/Private/AssetRegistry.cpp b/Engine/Source/Runtime/AssetRegistry/Private/AssetRegistry.cpp
--- a/Engine/Source/Runtime/AssetRegistry/Private/AssetRegistry.cpp
+++ b/Engine/Source/Runtime/AssetRegistry/Private/AssetRegistry.cpp
@@ -6,6 +6,7 @@
 #include "AssetDataGatherer.h"
 #include "AssetRegistry/ARFilter.h"
 #include "AssetRegistry/AssetDependencyGatherer.h"
+#include "AssetRegistry/AssetRegistryTelemetry.h"
 #include "AssetRegistryConsoleCommands.h"
 #include "AssetRegistryPrivate.h"
 #include "Async/Async.h"
@@ -35,6 +36,7 @@
 #include "Serialization/CompactBinarySerialization.h"
 #include "Serialization/CompactBinaryWriter.h"
 #include "Templates/UnrealTemplate.h"
+#include "TelemetryRouter.h"
 #include "UObject/ConstructorHelpers.h"
 #include "UObject/CoreRedirects.h"
 #include "UObject/MetaData.h"
@@ -809,12 +811,14 @@ void FAssetRegistryImpl::Initialize(Impl::FInitializeContext& Context)
 	Utils::InitializeSerializationOptionsFromIni(SerializationOptions, FString());
 	Utils::InitializeSerializationOptionsFromIni(DevelopmentSerializationOptions, FString(), UE::AssetRegistry::ESerializationTarget::ForDevelopment);
 
+	bool bStartedAsyncGather = false;
 	if (ShouldSearchAllAssetsAtStart())
 	{
 		ConstructGatherer();
 		if (GlobalGatherer->IsAsyncEnabled())
 		{
 			SearchAllAssetsInitialAsync(Context.Events, Context.InheritanceContext);
+			bStartedAsyncGather = true;
 		}
 		else
 		{
@@ -826,7 +830,13 @@ void FAssetRegistryImpl::Initialize(Impl::FInitializeContext& Context)
 	ConsumeOrDeferPreloadedPremade(Context.UARI, Context.Events);
 
 	// Report startup time. This does not include DirectoryWatcher startup time.
-	UE_LOG(LogAssetRegistry, Log, TEXT("FAssetRegistry took %0.4f seconds to start up"), FPlatformTime::Seconds() - StartupStartTime);
+	double StartupDuration = FPlatformTime::Seconds() - StartupStartTime;
+	UE_LOG(LogAssetRegistry, Log, TEXT("FAssetRegistry took %0.4f seconds to start up"), StartupDuration);
+	
+	FTelemetryRouter::Get().ProvideTelemetry<UE::Telemetry::AssetRegistry::FStartupTelemetry>({
+		StartupDuration,
+		bStartedAsyncGather
+	});
 
 #if WITH_EDITOR
 	if (GConfig)
@@ -1554,6 +1564,7 @@ void FAssetRegistryImpl::ConstructGatherer()
 void FAssetRegistryImpl::SearchAllAssetsInitialAsync(Impl::FEventContext& EventContext,
 	Impl::FClassInheritanceContext& InheritanceContext)
 {
+	InitialSearchStartTime = FPlatformTime::Seconds();
 	bInitialSearchStarted = true;
 	bInitialSearchCompleted = false;
 	SearchAllAssets(EventContext, InheritanceContext, false /* bSynchronousSearch */);
@@ -1565,6 +1576,7 @@ void UAssetRegistryImpl::SearchAllAssets(bool bSynchronousSearch)
 {
 	using namespace UE::AssetRegistry::Impl;
 
+	double StartTime = FPlatformTime::Seconds();
 	FEventContext EventContext;
 	{
 		LLM_SCOPE(ELLMTag::AssetRegistry);
@@ -1580,7 +1592,7 @@ void UAssetRegistryImpl::SearchAllAssets(bool bSynchronousSearch)
 		GuardedData.SearchAllAssets(EventContext, InheritanceContext, bSynchronousSearch);
 		if (bSynchronousSearch)
 		{
-			GuardedData.LogSearchDiagnostics();
+			GuardedData.LogSearchDiagnostics(StartTime);
 		}
 	}
 #if WITH_EDITOR
@@ -3511,11 +3523,15 @@ void UAssetRegistryImpl::ScanPathsSynchronousInternal(const TArray<FString>& InD
 	UE::AssetRegistry::Impl::FScanPathContext Context(EventContext, InheritanceContext, InDirs, InFiles,
 		bForceRescan, bIgnoreDenyListScanFilters, nullptr /* OutFindAssets */);
 
+	bool bInitialSearchStarted;
+	bool bInitialSearchCompleted;
 	{
 		LLM_SCOPE(ELLMTag::AssetRegistry);
 		FWriteScopeLock InterfaceScopeLock(InterfaceLock);
 		GetInheritanceContextWithRequiredLock(InterfaceScopeLock, InheritanceContext, InheritanceBuffer);
 
+		bInitialSearchStarted = GuardedData.IsInitialSearchStarted();
+		bInitialSearchCompleted = GuardedData.IsInitialSearchCompleted();
 		// make sure any outstanding async preload is complete
 		GuardedData.ConditionalLoadPremadeAssetRegistry(*this, EventContext, InterfaceScopeLock);
 		GuardedData.ScanPathsSynchronous(Context);
@@ -3544,8 +3560,19 @@ void UAssetRegistryImpl::ScanPathsSynchronousInternal(const TArray<FString>& InD
 		PathsString = FString::Printf(TEXT("'%s'"), *Context.LocalPaths[0]);
 	}
 
+
+	double Duration = FPlatformTime::Seconds() - SearchStartTime;
+	UE::Telemetry::AssetRegistry::FSynchronousScanTelemetry Telemetry; 
+	Telemetry.Directories = MakeArrayView(InDirs);
+	Telemetry.Files = MakeArrayView(InFiles);
+	Telemetry.Flags = InScanFlags;
+	Telemetry.NumFoundAssets = Context.NumFoundAssets;
+	Telemetry.Duration = Duration;
+	Telemetry.bInitialSearchStarted = bInitialSearchStarted;
+	Telemetry.bInitialSearchCompleted = bInitialSearchCompleted;
+	FTelemetryRouter::Get().ProvideTelemetry(Telemetry);
 	UE_LOG(LogAssetRegistry, Verbose, TEXT("ScanPathsSynchronous completed scanning %s to find %d assets in %0.4f seconds"), *PathsString,
-		Context.NumFoundAssets, FPlatformTime::Seconds() - SearchStartTime);
+		Context.NumFoundAssets, Duration);
 }
 
 void UAssetRegistryImpl::PrioritizeSearchPath(const FString& PathToPrioritize)
@@ -4050,7 +4077,7 @@ Impl::EGatherStatus FAssetRegistryImpl::TickGatherer(Impl::FEventContext& EventC
 			UpdateRedirectCollector();
 #endif
 			RecordTimer();
-			LogSearchDiagnostics();
+			LogSearchDiagnostics(InitialSearchStartTime);
 			TRACE_END_REGION(TEXT("Asset Registry Scan"));
 
 			bInitialSearchCompleted = true;
@@ -4062,14 +4089,21 @@ Impl::EGatherStatus FAssetRegistryImpl::TickGatherer(Impl::FEventContext& EventC
 	return OutStatus;
 }
 
-void FAssetRegistryImpl::LogSearchDiagnostics() const
+void FAssetRegistryImpl::LogSearchDiagnostics(double StartTime) const
 {
-	float GatherTimeSeconds;
-	float DiscoveryTimeSeconds;
-	GlobalGatherer->GetDiagnostics(GatherTimeSeconds, DiscoveryTimeSeconds);
-	float Total = DiscoveryTimeSeconds + GatherTimeSeconds + StoreGatherResultsTimeSeconds;
+	FAssetGatherDiagnostics Diagnostics = GlobalGatherer->GetDiagnostics();
+	float Total = Diagnostics.DiscoveryTimeSeconds + Diagnostics.GatherTimeSeconds + StoreGatherResultsTimeSeconds;
+	UE::Telemetry::AssetRegistry::FGatherTelemetry Telemetry;
+	Telemetry.TotalSearchDurationSeconds = FPlatformTime::Seconds() - StartTime;
+	Telemetry.TotalWorkTimeSeconds = Total;
+	Telemetry.DiscoveryTimeSeconds = Diagnostics.DiscoveryTimeSeconds;
+	Telemetry.GatherTimeSeconds = Diagnostics.GatherTimeSeconds;
+	Telemetry.StoreTimeSeconds = StoreGatherResultsTimeSeconds;
+	Telemetry.NumCachedAssetFiles = Diagnostics.NumCachedAssetFiles;
+	Telemetry.NumUncachedAssetFiles = Diagnostics.NumUncachedAssetFiles;
+	FTelemetryRouter::Get().ProvideTelemetry(Telemetry);
 	UE_LOG(LogAssetRegistry, Log, TEXT("AssetRegistryGather time %.4fs: AssetDataDiscovery %0.4fs, AssetDataGather %0.4fs, StoreResults %0.4fs."),
-		Total, DiscoveryTimeSeconds, GatherTimeSeconds, StoreGatherResultsTimeSeconds);
+		Total, Diagnostics.DiscoveryTimeSeconds, Diagnostics.GatherTimeSeconds, StoreGatherResultsTimeSeconds);
 }
 
 void FAssetRegistryImpl::TickGatherPackage(Impl::FEventContext& EventContext, const FString& PackageName, const FString& LocalPath)
@@ -5516,6 +5550,8 @@ void UAssetRegistryImpl::OnDirectoryChanged(const TArray<FFileChangeData>& FileC
 {
 	TRACE_CPUPROFILER_EVENT_SCOPE(UAssetRegistryImpl::OnDirectoryChanged);
 
+	double StartTime = FPlatformTime::Seconds();
+	
 	// Take local copy of FileChanges array as we wish to collapse pairs of 'Removed then Added' FileChangeData
 	// entries into a single 'Modified' entry.
 	TArray<FFileChangeData> FileChangesProcessed(FileChanges);
@@ -5544,15 +5580,26 @@ void UAssetRegistryImpl::OnDirectoryChanged(const TArray<FFileChangeData>& FileC
 	}
 
 	UE::AssetRegistry::Impl::FEventContext EventContext;
+	bool bInitialSearchStarted;
+	bool bInitialSearchCompleted;
 	{
 		LLM_SCOPE(ELLMTag::AssetRegistry);
 		FWriteScopeLock InterfaceScopeLock(InterfaceLock);
+		bInitialSearchStarted = GuardedData.IsInitialSearchStarted();
+		bInitialSearchCompleted = GuardedData.IsInitialSearchCompleted();
 		UE::AssetRegistry::Impl::FClassInheritanceContext InheritanceContext;
 		UE::AssetRegistry::Impl::FClassInheritanceBuffer InheritanceBuffer;
 		GetInheritanceContextWithRequiredLock(InterfaceScopeLock, InheritanceContext, InheritanceBuffer);
 		GuardedData.OnDirectoryChanged(EventContext, InheritanceContext, FileChangesProcessed);
 	}
 	Broadcast(EventContext);
+	
+	FTelemetryRouter::Get().ProvideTelemetry<UE::Telemetry::AssetRegistry::FDirectoryWatcherUpdateTelemetry>({
+		FileChanges,
+		FPlatformTime::Seconds() - StartTime,
+		bInitialSearchStarted,
+		bInitialSearchCompleted,
+	});
 }
 
 namespace UE::AssetRegistry
diff --git a/Engine/Source/Runtime/AssetRegistry/Private/AssetRegistryImpl.h b/Engine/Source/Runtime/AssetRegistry/Private/AssetRegistryImpl.h
--- a/Engine/Source/Runtime/AssetRegistry/Private/AssetRegistryImpl.h
+++ b/Engine/Source/Runtime/AssetRegistry/Private/AssetRegistryImpl.h
@@ -186,8 +186,10 @@ public:
 	Impl::EGatherStatus TickGatherer(Impl::FEventContext& EventContext,
 		Impl::FClassInheritanceContext& InheritanceContext, const double TickStartTime, bool& bOutInterrupted,
 		TOptional<FAssetsFoundCallback> AssetsFoundCallback = TOptional<FAssetsFoundCallback>());
-	/** Send a log message with the search statistics. */
-	void LogSearchDiagnostics() const;
+	/** Send a log message with the search statistics. 
+	 *  StartTime is used to report wall clock search time in the case of background scan
+	 */
+	void LogSearchDiagnostics(double StartTime) const;
 	/** Look for and load a single AssetData result from the gatherer. */
 	void TickGatherPackage(Impl::FEventContext& EventContext, const FString& PackageName, const FString& LocalPath);
 	void ClearGathererCache();
@@ -392,6 +394,8 @@ private:
 	/** The highest number of pending results observed during initial gathering */
 	int32 HighestPending = 0;
 
+	/** Time the initial async search was started */
+	double InitialSearchStartTime = 0.0f;
 	/** Flag to indicate if we used an initial async search */
 	bool bInitialSearchStarted;
 	/** Flag to indicate if the initial background search has completed */
diff --git a/Engine/Source/Runtime/AssetRegistry/Public/AssetRegistry/AssetRegistryTelemetry.h b/Engine/Source/Runtime/AssetRegistry/Public/AssetRegistry/AssetRegistryTelemetry.h
new file mode 100644
--- /dev/null
+++ b/Engine/Source/Runtime/AssetRegistry/Public/AssetRegistry/AssetRegistryTelemetry.h
@@ -0,0 +1,81 @@
+// Copyright Epic Games, Inc. All Rights Reserved.
+
+#pragma once
+
+#include "AssetRegistry/IAssetRegistry.h"
+#include "Containers/ArrayView.h"
+#include "Misc/Guid.h"
+
+struct FFileChangeData;
+
+namespace UE::Telemetry::AssetRegistry
+{
+    // Information about very early asset registry startup.
+    struct FStartupTelemetry
+    {
+        static constexpr inline FGuid TelemetryID = FGuid(0x903732ce, 0x8bd04eb9, 0x92170fde, 0x0a1c1562);
+        
+        // Time spent in synchronous initialization
+        double StartupDuration = 0.0;
+        // Whether async gather operation was started at this time
+        bool bStartedAsyncGather = false;
+    };
+
+    // Information about a synchronous scan that was executed, blocking the calling thread.
+    struct FSynchronousScanTelemetry
+    {
+        static constexpr inline FGuid TelemetryID = FGuid(0x2b4b9f1a, 0xdcfd4958, 0xbe43ba9d, 0xae309392);
+        
+        // List of directories that were scanned
+        TConstArrayView<FString> Directories;
+        // List of specific files that were scanned
+        TConstArrayView<FString> Files;
+        // Flags controlling scan behavior
+        UE::AssetRegistry::EScanFlags Flags = UE::AssetRegistry::EScanFlags::None;
+        // Number of assets found by this scan
+        int32 NumFoundAssets = 0;
+        // Duration of synchronous scan
+        double Duration = 0.0;
+        // Whether the main background async gather was started at this time
+        bool bInitialSearchStarted = false;
+        // Whether the main background async gather was completed at this time
+        bool bInitialSearchCompleted = false;
+    };
+
+    // Information about the initial asset registry scan triggered when launching the process. 
+    // This is also a good time to gather information about the contents of the asset registry.
+    struct FGatherTelemetry
+    {
+        static constexpr inline FGuid TelemetryID = FGuid(0x2b4b9f1a, 0xdcfd4958, 0xbe43ba9d, 0xae309392);
+        
+        // Total wall clock time between start of search and completion.
+        double TotalSearchDurationSeconds = 0.0;
+        // Total work time (includes work done in parallel stages)
+        double TotalWorkTimeSeconds = 0.0;
+        // Time spent discovering asset files on disk
+        double DiscoveryTimeSeconds = 0.0;
+        // Time spent gathering asset data from files on disk (or cache)
+        double GatherTimeSeconds = 0.0;
+        // Time spent storing asset data in the asset registry for searching
+        double StoreTimeSeconds = 0.0;
+        // Number of asset files read from cache
+        int32 NumCachedAssetFiles = 0;
+        // Number of asset files read loose from disk
+        int32 NumUncachedAssetFiles = 0;
+    };
+    
+    // Information about an asset registry update that was triggered by the directory watcher module. 
+    struct FDirectoryWatcherUpdateTelemetry
+    {
+        static constexpr inline FGuid TelemetryID = FGuid(0xa1da56e1, 0xe1314918, 0xba6fdb24, 0xb8911ce0);
+        
+        // File change data from watcher
+        TConstArrayView<FFileChangeData> Changes;
+        // Total time spent in update handler
+        double DurationSeconds = 0.0;
+        // Whether the main background async gather was started at this time
+        bool bInitialSearchStarted = false;
+        // Whether the main background async gather was completed at this time
+        bool bInitialSearchCompleted = false;
+    };
+}
\ No newline at end of file
