EngineStartup: Add IsEngineStartupModuleLoadingComplete so systems can take different behavior (for logging, user notification, other) when the engine is still initializing versus when initialization is complete.
Convert AssetRegistry to use this function instead of tracking the state itself.
#rxn

diff --git a/Engine/Source/Runtime/AssetRegistry/Private/AssetDataGatherer.cpp b/Engine/Source/Runtime/AssetRegistry/Private/AssetDataGatherer.cpp
--- a/Engine/Source/Runtime/AssetRegistry/Private/AssetDataGatherer.cpp
+++ b/Engine/Source/Runtime/AssetRegistry/Private/AssetDataGatherer.cpp
@@ -3182,7 +3182,6 @@ FAssetDataGatherer::FAssetDataGatherer(const TArray<FString>& InLongPackageNames
 	, bIsIdle(false)
 	, bFirstTickAfterIdle(true)
 	, bFinishedInitialDiscovery(false)
-	, bAllModuleLoadingPhasesComplete(false)
 	, WaitBatchCount(0)
 	, LastMonolithicCacheSaveUncachedAssetFiles(0)
 	, CacheInUseCount(0)
@@ -4794,7 +4793,10 @@ void FAssetDataGatherer::SetIsIdle(bool bInIsIdle, double& TickStartTime)
 			CurrentSearchTime += FPlatformTime::Seconds() - TickStartTime;
 			TickStartTime = -1.;
 		}
-		if (!bFinishedInitialDiscovery && bAllModuleLoadingPhasesComplete)
+		// Finishing the initial discovery is blocked until IsEngineStartupModuleLoadingComplete because plugins can
+		// be mounted during startup up until that point, and we need to wait for all the plugins that will load
+		// before declaring completion.
+		if (!bFinishedInitialDiscovery && IsEngineStartupModuleLoadingComplete())
 		{
 			bFinishedInitialDiscovery = true;
 
@@ -4834,7 +4836,6 @@ void FAssetDataGatherer::OnAllModuleLoadingPhasesComplete()
 {
 	CHECK_IS_NOT_LOCKED_CURRENT_THREAD(ResultsLock);
 	FGathererScopeLock ResultsScopeLock(&ResultsLock);
-	bAllModuleLoadingPhasesComplete = true;
 	SetIsIdle(false);
 }
 
diff --git a/Engine/Source/Runtime/AssetRegistry/Private/AssetDataGatherer.h b/Engine/Source/Runtime/AssetRegistry/Private/AssetDataGatherer.h
--- a/Engine/Source/Runtime/AssetRegistry/Private/AssetDataGatherer.h
+++ b/Engine/Source/Runtime/AssetRegistry/Private/AssetDataGatherer.h
@@ -460,13 +460,6 @@ private:
 	bool bFirstTickAfterIdle;
 	/** True if we have finished discovering our first wave of files, to report metrics for that most-important wave. */
 	bool bFinishedInitialDiscovery;
-	/**
-	 * Flag to indicate LaunchEngineLoop's AllModuleLoadingPhases is complete; finishing the initial discovery is
-	 * blocked until preloading complete because plugins can be mounted during startup up until that point, and
-	 * we need to wait for all the plugins that will load before declaring completion.
-	 */
-	bool bAllModuleLoadingPhasesComplete = false;
-
 
 	// Variable section for variables that are read/writable only within TickLock.
 
diff --git a/Engine/Source/Runtime/AssetRegistry/Private/AssetRegistry.cpp b/Engine/Source/Runtime/AssetRegistry/Private/AssetRegistry.cpp
--- a/Engine/Source/Runtime/AssetRegistry/Private/AssetRegistry.cpp
+++ b/Engine/Source/Runtime/AssetRegistry/Private/AssetRegistry.cpp
@@ -960,7 +960,6 @@ void UAssetRegistryImpl::InitializeEvents(UE::AssetRegistry::Impl::FInitializeCo
 #endif // WITH_EDITOR
 
 	FCoreDelegates::OnEnginePreExit.AddUObject(this, &UAssetRegistryImpl::OnEnginePreExit);
-	FCoreDelegates::OnAllModuleLoadingPhasesComplete.AddUObject(this, &UAssetRegistryImpl::OnAllModuleLoadingPhasesComplete);
 
 	// Listen for new content paths being added or removed at runtime.  These are usually plugin-specific asset paths that
 	// will be loaded a bit later on.
@@ -1425,13 +1424,6 @@ void UAssetRegistryImpl::OnEnginePreExit()
 	GuardedData.OnEnginePreExit();
 }
 
-void UAssetRegistryImpl::OnAllModuleLoadingPhasesComplete()
-{
-	LLM_SCOPE(ELLMTag::AssetRegistry);
-	FWriteScopeLock InterfaceScopeLock(InterfaceLock);
-	GuardedData.OnAllModuleLoadingPhasesComplete();
-}
-
 void UAssetRegistryImpl::FinishDestroy()
 {
 	LLM_SCOPE(ELLMTag::AssetRegistry);
@@ -1443,7 +1435,6 @@ void UAssetRegistryImpl::FinishDestroy()
 		FPackageName::OnContentPathDismounted().RemoveAll(this);
 		FCoreDelegates::OnPostEngineInit.RemoveAll(this);
 		FCoreDelegates::OnEnginePreExit.RemoveAll(this);
-		FCoreDelegates::OnAllModuleLoadingPhasesComplete.RemoveAll(this);
 		IPluginManager::Get().OnLoadingPhaseComplete().RemoveAll(this);
 
 #if WITH_EDITOR
@@ -1530,11 +1521,6 @@ void FAssetRegistryImpl::OnEnginePreExit()
 	GlobalGatherer.Reset();
 }
 
-void FAssetRegistryImpl::OnAllModuleLoadingPhasesComplete()
-{
-	bAllModuleLoadingPhasesComplete = true;
-}
-
 void FAssetRegistryImpl::ConstructGatherer()
 {
 	if (GlobalGatherer.IsValid())
@@ -4070,7 +4056,9 @@ Impl::EGatherStatus FAssetRegistryImpl::TickGatherer(Impl::FEventContext& EventC
 		HighestPending = 0;
 		BackgroundResults.Shrink();
 
-		if (!bInitialSearchCompleted && bPreloadingComplete && bAllModuleLoadingPhasesComplete)
+		// Finishing the background search is blocked until preloading complete because plugins can be mounted during
+		// startup up until that point, and we need to wait for all the plugins to load before declaring completion.
+		if (!bInitialSearchCompleted && bPreloadingComplete && IsEngineStartupModuleLoadingComplete())
 		{
 #if WITH_EDITOR
 			// update redirectors
diff --git a/Engine/Source/Runtime/AssetRegistry/Private/AssetRegistry.h b/Engine/Source/Runtime/AssetRegistry/Private/AssetRegistry.h
--- a/Engine/Source/Runtime/AssetRegistry/Private/AssetRegistry.h
+++ b/Engine/Source/Runtime/AssetRegistry/Private/AssetRegistry.h
@@ -199,7 +199,6 @@ protected:
 
 private:
 	void OnEnginePreExit();
-	void OnAllModuleLoadingPhasesComplete();
 #if WITH_EDITOR
 	void OnFEngineLoopInitCompleteSearchAllAssets();
 	/** Called when new gatherer is registered. Requires subsequent call to RebuildAssetDependencyGathererMapIfNeeded */
diff --git a/Engine/Source/Runtime/AssetRegistry/Private/AssetRegistryImpl.h b/Engine/Source/Runtime/AssetRegistry/Private/AssetRegistryImpl.h
--- a/Engine/Source/Runtime/AssetRegistry/Private/AssetRegistryImpl.h
+++ b/Engine/Source/Runtime/AssetRegistry/Private/AssetRegistryImpl.h
@@ -118,7 +118,6 @@ public:
 	/** Construct the AssetRegistryImpl, including initial scans if applicable. */
 	void Initialize(Impl::FInitializeContext& Context);
 	void OnEnginePreExit();
-	void OnAllModuleLoadingPhasesComplete();
 
 	// Helpers for functions of the same name from UAssetRegistryImpl
 
@@ -405,12 +404,6 @@ private:
 	 * because preloading can add assets.
 	 */
 	bool bPreloadingComplete = false;
-	/**
-	 * Flag to indicate LaunchEngineLoop's AllModuleLoadingPhases is complete; finishing the background search is
-	 * blocked until preloading complete because plugins can be mounted during startup up until that point, and
-	 * we need to wait for all the plugins that will load before declaring completion.
-	 */
-	bool bAllModuleLoadingPhasesComplete = false;
 	/** Status of the background search, so we can take actions when it changes to or from idle */
 	Impl::EGatherStatus GatherStatus;
 
diff --git a/Engine/Source/Runtime/Core/Private/Misc/CoreGlobals.cpp b/Engine/Source/Runtime/Core/Private/Misc/CoreGlobals.cpp
--- a/Engine/Source/Runtime/Core/Private/Misc/CoreGlobals.cpp
+++ b/Engine/Source/Runtime/Core/Private/Misc/CoreGlobals.cpp
@@ -9,6 +9,7 @@
 #include "Misc/CoreStats.h"
 #include "Misc/TrackedActivity.h"
 #include "Misc/Compression.h"
+#include "Misc/CoreDelegates.h"
 #include "Misc/LazySingleton.h"
 #include "Misc/PlayInEditorLoadingScope.h"
 #include "Misc/CommandLine.h"
@@ -285,6 +286,22 @@ static void appNoop()
 {
 }
 
+bool GEngineStartupModuleLoadingComplete = false;
+CORE_API bool IsEngineStartupModuleLoadingComplete()
+{
+	return GEngineStartupModuleLoadingComplete;
+}
+
+CORE_API void SetEngineStartupModuleLoadingComplete()
+{
+	if (ensure(!GEngineStartupModuleLoadingComplete))
+	{
+		GEngineStartupModuleLoadingComplete = true;
+		SCOPED_BOOT_TIMING("OnAllModuleLoadingPhasesComplete.Broadcast");
+		FCoreDelegates::OnAllModuleLoadingPhasesComplete.Broadcast();
+	}
+}
+
 // This should be left non static to allow *edge* cases only in Core to extern and set this.
 bool GShouldRequestExit = false;
 
diff --git a/Engine/Source/Runtime/Core/Public/CoreGlobals.h b/Engine/Source/Runtime/Core/Public/CoreGlobals.h
--- a/Engine/Source/Runtime/Core/Public/CoreGlobals.h
+++ b/Engine/Source/Runtime/Core/Public/CoreGlobals.h
@@ -352,6 +352,13 @@ extern CORE_API bool GIsGuarded;
 	#define UE_SET_REQUEST_EXIT_ON_TICK_ONLY 0
 #endif
 
+/**
+ * Returns true after all initial modules have finished all of their loading phases during Engine startup:
+ * Corresponds to the event FCoreDelegates::OnAllModuleLoadingPhasesComplete
+ */
+extern CORE_API bool IsEngineStartupModuleLoadingComplete();
+extern CORE_API void SetEngineStartupModuleLoadingComplete();
+
 UE_DEPRECATED(4.24, "Please use IsEngineExitRequested()/RequestEngineExit(const FString&)")
 extern CORE_API bool GIsRequestingExit;
 
diff --git a/Engine/Source/Runtime/Launch/Private/LaunchEngineLoop.cpp b/Engine/Source/Runtime/Launch/Private/LaunchEngineLoop.cpp
--- a/Engine/Source/Runtime/Launch/Private/LaunchEngineLoop.cpp
+++ b/Engine/Source/Runtime/Launch/Private/LaunchEngineLoop.cpp
@@ -4105,7 +4105,7 @@ int32 FEngineLoop::PreInitPostStartupScreen(const TCHAR* CmdLine)
 			ensure(IPluginManager::Get().LoadModulesForEnabledPlugins(ELoadingPhase::PostEngineInit));
 
 			// Call module loading phases completion callbacks
-			FCoreDelegates::OnAllModuleLoadingPhasesComplete.Broadcast();
+			SetEngineStartupModuleLoadingComplete();
 
 			//run automation smoke tests now that the commandlet has had a chance to override the above flags and GEngine is available
 			FAutomationTestFramework::Get().RunSmokeTests();
@@ -4843,10 +4843,7 @@ int32 FEngineLoop::Init()
 	}
 
 	// Call module loading phases completion callbacks
-	{
-		SCOPED_BOOT_TIMING("OnAllModuleLoadingPhasesComplete.Broadcast");
-		FCoreDelegates::OnAllModuleLoadingPhasesComplete.Broadcast();
-	}
+	SetEngineStartupModuleLoadingComplete();
 
 	{
 		SCOPED_BOOT_TIMING("GEngine->Start()");
