Fixed an issue where the asset registry state of an cooked plugin would remove some of the needed devellopment tags from an asset that already exist in the EditorClientAssetRegistry.

diff --git a/Engine/Source/Runtime/AssetRegistry/Private/AssetRegistry.cpp b/Engine/Source/Runtime/AssetRegistry/Private/AssetRegistry.cpp
--- a/Engine/Source/Runtime/AssetRegistry/Private/AssetRegistry.cpp
+++ b/Engine/Source/Runtime/AssetRegistry/Private/AssetRegistry.cpp
@@ -761,7 +761,16 @@ void FAssetRegistryImpl::LoadPremadeAssetRegistry(Impl::FEventContext& EventCont
 					FAssetRegistryState PluginState;
 					PluginState.Load(SerializedAssetData);
 
+#if ASSETREGISTRY_ENABLE_PREMADE_REGISTRY_IN_EDITOR
+					/*
+					 * Only update the new assets when using a premade asset registry in editor.
+					 * The main state will often already include the DLC/plugin assets and is often in a development mode where the plugin state will not be.
+					 * If we update the existing assets in those cases it will be causing a lost of tags and values that are needed for the editor systems.
+					 */
+					AppendState(EventContext, PluginState,  FAssetRegistryState::EInitializationMode::OnlyUpdateNew);
+#else
 					AppendState(EventContext, PluginState);
+#endif
 				}
 			}
 		}
