#UE Lower verbosity on error aboud missing assetregistry data for additional tags

diff --git a/Engine/Source/Runtime/AssetRegistry/Private/AssetRegistryState.cpp b/Engine/Source/Runtime/AssetRegistry/Private/AssetRegistryState.cpp
--- a/Engine/Source/Runtime/AssetRegistry/Private/AssetRegistryState.cpp
+++ b/Engine/Source/Runtime/AssetRegistry/Private/AssetRegistryState.cpp
@@ -1838,7 +1838,7 @@ void FAssetRegistryState::AddTagsToAssetData(const FSoftObjectPath& InObjectPath
 	FSetElementId Id = CachedAssets.FindId(FCachedAssetKey(InObjectPath));
 	if (!Id.IsValidId())
 	{
-		UE_LOG(LogAssetRegistry, Error, TEXT("AddTagsToAssetData called with asset data that doesn't exist! Tags not added. ObjectPath: %s"), *InObjectPath.ToString());
+		UE_LOG(LogAssetRegistry, Warning, TEXT("AddTagsToAssetData called with asset data that doesn't exist! Tags not added. ObjectPath: %s"), *InObjectPath.ToString());
 		return;
 	}
 	FAssetData* AssetData = CachedAssets[Id];
