Internal analytics for asset registry scan.

diff --git a/Engine/Source/Runtime/AssetRegistry/Private/AssetRegistry.cpp b/Engine/Source/Runtime/AssetRegistry/Private/AssetRegistry.cpp
--- a/Engine/Source/Runtime/AssetRegistry/Private/AssetRegistry.cpp
+++ b/Engine/Source/Runtime/AssetRegistry/Private/AssetRegistry.cpp
@@ -116,6 +116,38 @@ namespace UE::AssetRegistry
 	const FName Stage_ChunkInstalledSizeFName(TEXT("Stage_ChunkInstalledSize"));
 	const FName Stage_ChunkStreamingSizeFName(TEXT("Stage_ChunkStreamingSize"));
 	const FName Stage_ChunkOptionalSizeFName(TEXT("Stage_ChunkOptionalSize"));
+	
+	FString LexToString(EScanFlags Flags)
+	{
+		const TCHAR* Names[] = {
+			TEXT("ForceRescan"),
+			TEXT("IgnoreDenyListScanFilters"),
+			TEXT("WaitForInMemoryObjects")
+		};
+		
+		if (Flags == EScanFlags::None)
+		{
+			return TEXT("None");
+		}
+		
+		uint32 AllKnownFlags = (1 << (UE_ARRAY_COUNT(Names)+1)) - 1;
+		ensureMsgf(EnumHasAllFlags((EScanFlags)AllKnownFlags, Flags), TEXT("LexToString(UE::AssetRegistry::EScanFlags) is missing some cases"));
+
+		TStringBuilder<256> Builder;
+		for (uint32 i=0; i < UE_ARRAY_COUNT(Names); ++i)
+		{
+			if (EnumHasAllFlags(Flags, (EScanFlags)(1 << i)))
+			{
+				if (Builder.Len() != 0)
+				{
+					Builder << TEXT("|");
+				}
+				Builder << Names[i];	
+			}
+		}
+		
+		return Builder.ToString();
+	}
 }
 
 namespace UE::AssetRegistry::Impl
diff --git a/Engine/Source/Runtime/AssetRegistry/Public/AssetRegistry/AssetRegistryTelemetry.h b/Engine/Source/Runtime/AssetRegistry/Public/AssetRegistry/AssetRegistryTelemetry.h
--- a/Engine/Source/Runtime/AssetRegistry/Public/AssetRegistry/AssetRegistryTelemetry.h
+++ b/Engine/Source/Runtime/AssetRegistry/Public/AssetRegistry/AssetRegistryTelemetry.h
@@ -46,7 +46,7 @@ namespace UE::Telemetry::AssetRegistry
     // This is also a good time to gather information about the contents of the asset registry.
     struct FGatherTelemetry
     {
-        static constexpr inline FGuid TelemetryID = FGuid(0x2b4b9f1a, 0xdcfd4958, 0xbe43ba9d, 0xae309392);
+        static constexpr inline FGuid TelemetryID = FGuid(0xafcec052, 0x5d2c4850, 0xbfc6d11d, 0x3163ccd5);
         
         // Total wall clock time between start of search and completion.
         double TotalSearchDurationSeconds = 0.0;
diff --git a/Engine/Source/Runtime/AssetRegistry/Public/AssetRegistry/IAssetRegistry.h b/Engine/Source/Runtime/AssetRegistry/Public/AssetRegistry/IAssetRegistry.h
--- a/Engine/Source/Runtime/AssetRegistry/Public/AssetRegistry/IAssetRegistry.h
+++ b/Engine/Source/Runtime/AssetRegistry/Public/AssetRegistry/IAssetRegistry.h
@@ -58,6 +58,7 @@ namespace EAssetAvailabilityProgressReportingType
 namespace UE::AssetRegistry
 {
 
+	// Please update LexToString when modifying this enum
 enum class EScanFlags : uint32
 {
 	None = 0,
@@ -67,6 +68,8 @@ enum class EScanFlags : uint32
 };
 ENUM_CLASS_FLAGS(EScanFlags);
 
+ASSETREGISTRY_API FString LexToString(EScanFlags Flags);
+
 } // namespace UE::AssetRegistry
 
 USTRUCT(BlueprintType)
