Optimize editor startup time related to the asset registry:
Remove duplicates from package dependencies returned by discovery threads
Optimize PathTree and ObjectRedirector handling to avoid redundant operations
Optimize localization and component type caches which are updated for every scanned asset
Add AssetRegistry.MaxSecondsPerFrame cvar to change amount of main thread time used by the asset registry

diff --git a/Engine/Source/Editor/UnrealEd/Private/ComponentTypeRegistry.cpp b/Engine/Source/Editor/UnrealEd/Private/ComponentTypeRegistry.cpp
--- a/Engine/Source/Editor/UnrealEd/Private/ComponentTypeRegistry.cpp
+++ b/Engine/Source/Editor/UnrealEd/Private/ComponentTypeRegistry.cpp
@@ -297,6 +297,7 @@ public:
 static const FString CommonClassGroup(TEXT("Common"));
 // This has to stay in sync with logic in FKismetCompilerContext::FinishCompilingClass
 static const FString BlueprintComponents(TEXT("Custom"));
+static const FName BPParentClassName(GET_MEMBER_NAME_CHECKED(UBlueprint, ParentClass));
 
 template <typename ObjectType>
 static ObjectType* FindOrLoadObject( const FString& ObjectPath )
@@ -426,7 +427,11 @@ FComponentTypeRegistryData::FComponentTypeRegistryData()
 {
 	const auto HandleAdded = [](const FAssetData& Data, FComponentTypeRegistryData* Parent)
 	{
-		Parent->PendingAssetData.Push(Data);
+		// Only add to pending array if this might actually be a blueprint
+		if (Data.FindTag(BPParentClassName))
+		{
+			Parent->PendingAssetData.Push(Data);
+		}
 	};
 
 	IAssetRegistry& AssetRegistry = FModuleManager::LoadModuleChecked<FAssetRegistryModule>(AssetRegistryConstants::ModuleName).Get();
@@ -435,7 +440,10 @@ FComponentTypeRegistryData::FComponentTypeRegistryData()
 
 	const auto HandleRenamed = [](const FAssetData& Data, const FString&, FComponentTypeRegistryData* Parent)
 	{
-		Parent->PendingAssetData.Push(Data);
+		if (Data.FindTag(BPParentClassName))
+		{
+			Parent->PendingAssetData.Push(Data);
+		}
 	};
 	AssetRegistry.OnAssetRenamed().AddStatic(HandleRenamed, this);
 }
@@ -693,7 +701,6 @@ void FComponentTypeRegistryData::Tick(float)
 
 		for (const FAssetData& Asset : PendingAssetData)
 		{
-			const FName BPParentClassName(GET_MEMBER_NAME_CHECKED(UBlueprint, ParentClass));
 			const FString TagValue = Asset.GetTagValueRef<FString>(BPParentClassName);
 			const FTopLevelAssetPath ObjectPath(FPackageName::ExportTextPathToObjectPath(TagValue));
 			if (DerivedClassNames.Contains(ObjectPath))
diff --git a/Engine/Source/Runtime/AssetRegistry/Private/AssetRegistry.cpp b/Engine/Source/Runtime/AssetRegistry/Private/AssetRegistry.cpp
--- a/Engine/Source/Runtime/AssetRegistry/Private/AssetRegistry.cpp
+++ b/Engine/Source/Runtime/AssetRegistry/Private/AssetRegistry.cpp
@@ -153,9 +153,17 @@ namespace UE::AssetRegistry
 namespace UE::AssetRegistry::Impl
 {
 	/** The max time to spend in UAssetRegistryImpl::Tick */
-	const float MaxSecondsPerFrame = 0.04f;
+	float MaxSecondsPerFrame = 0.04f;
+
+	/** Name of UObjectRedirector property */
+	const FName DestinationObjectFName(TEXT("DestinationObject"));
 }
 
+static FAutoConsoleVariableRef CVarAssetRegistryMaxSecondsPerFrame(
+	TEXT("AssetRegistry.MaxSecondsPerFrame"),
+	UE::AssetRegistry::Impl::MaxSecondsPerFrame,
+	TEXT("Maximum amount of time allowed for Asset Registry processing, in seconds"));
+
 /**
  * Implementation of IAssetRegistryInterface; forwards calls from the CoreUObject-accessible IAssetRegistryInterface into the AssetRegistry-accessible IAssetRegistry
  */
@@ -2760,26 +2768,29 @@ FSoftObjectPath FAssetRegistryImpl::GetRedirectedObjectPath(const FSoftObjectPat
 	}
 	const FAssetData* AssetData = State.GetAssetByObjectPath(RedirectedPath);
 
-	TSet<FSoftObjectPath> SeenPaths;
-	SeenPaths.Add(RedirectedPath);
+	// Most of the time this will either not be a redirector or only have one redirect, so optimize for that case
+	TArray<FSoftObjectPath, TInlineAllocator<2>> SeenPaths = { RedirectedPath };
 
 	// Need to follow chain of redirectors
 	while (AssetData && AssetData->IsRedirector())
 	{
 		FString Dest;
-		if (!AssetData->GetTagValue("DestinationObject", Dest))
+
+		if (!AssetData->GetTagValue(UE::AssetRegistry::Impl::DestinationObjectFName, Dest))
 		{
 			break;
 		}
-		ConstructorHelpers::StripObjectClass(Dest);
+		
+		// The FSoftObjectPath functions handle stripping class name if necessary
 		RedirectedPath = Dest;
-		bool bAlreadyExists;
-		SeenPaths.Add(RedirectedPath, &bAlreadyExists);
-		if (bAlreadyExists)
+
+		if (SeenPaths.Contains(RedirectedPath))
 		{
 			// Recursive, bail
 			break;
 		}
+
+		SeenPaths.Add(RedirectedPath);
 		AssetData = State.GetAssetByObjectPath(RedirectedPath);
 	}
 
@@ -6048,11 +6059,12 @@ void FAssetRegistryImpl::UpdateRedirectCollector()
 
 	for (const FAssetData* AssetData : RedirectorAssets)
 	{
-		FSoftObjectPath Destination = GetRedirectedObjectPath(AssetData->GetSoftObjectPath());
+		FSoftObjectPath Source = AssetData->GetSoftObjectPath();
+		FSoftObjectPath Destination = GetRedirectedObjectPath(Source);
 
-		if (Destination != AssetData->GetSoftObjectPath())
+		if (Destination != Source)
 		{
-			GRedirectCollector.AddAssetPathRedirection(AssetData->GetSoftObjectPath(), Destination);
+			GRedirectCollector.AddAssetPathRedirection(Source, Destination);
 		}
 	}
 }
diff --git a/Engine/Source/Runtime/AssetRegistry/Private/PackageDependencyData.cpp b/Engine/Source/Runtime/AssetRegistry/Private/PackageDependencyData.cpp
--- a/Engine/Source/Runtime/AssetRegistry/Private/PackageDependencyData.cpp
+++ b/Engine/Source/Runtime/AssetRegistry/Private/PackageDependencyData.cpp
@@ -1,6 +1,7 @@
 // Copyright Epic Games, Inc. All Rights Reserved.
 
 #include "PackageDependencyData.h"
+#include "Algo/Unique.h"
 
 FName FPackageDependencyData::GetImportPackageName(TConstArrayView<FObjectImport> ImportMap, int32 ImportIndex)
 {
@@ -30,6 +31,18 @@ FName FPackageDependencyData::GetImportPackageName(TConstArrayView<FObjectImport
 	return NAME_None;
 }
 
+struct FSortPackageDependency
+{
+	FORCEINLINE bool operator()(const FPackageDependencyData::FPackageDependency& A, const FPackageDependencyData::FPackageDependency& B) const
+	{
+		if (int32 Comparison = A.PackageName.CompareIndexes(B.PackageName))
+		{
+			return Comparison < 0;
+		}
+		return static_cast<uint8>(A.Property) < static_cast<uint8>(B.Property);
+	}
+};
+
 void FPackageDependencyData::LoadDependenciesFromPackageHeader(FName SourcePackageName, TConstArrayView<FObjectImport> ImportMap,
 	TArray<FName>& SoftPackageReferenceList, TMap<FPackageIndex, TArray<FName>>& SearchableNames,
 	TBitArray<>& ImportUsedInGame, TBitArray<>& SoftPackageUsedInGame)
@@ -46,6 +59,12 @@ void FPackageDependencyData::LoadDependenciesFromPackageHeader(FName SourcePacka
 		PackageDependencies.Add({ DependencyPackageName, DependencyProperty });
 	}
 
+	// Sort and make unique to reduce data saved and processed
+	PackageDependencies.Sort(FSortPackageDependency());
+
+	int UniqueNum = Algo::Unique(PackageDependencies);
+	PackageDependencies.SetNum(UniqueNum, false);
+
 	check(SoftPackageReferenceList.Num() == SoftPackageUsedInGame.Num());
 	for (int32 SoftPackageIdx = 0; SoftPackageIdx < SoftPackageReferenceList.Num(); ++SoftPackageIdx)
 	{
@@ -53,6 +72,8 @@ void FPackageDependencyData::LoadDependenciesFromPackageHeader(FName SourcePacka
 		FAssetIdentifier AssetId(DependencyPackageName);
 		EDependencyProperty DependencyProperty = UE::AssetRegistry::EDependencyProperty::Build; // !EDependencyProperty::Hard
 		DependencyProperty |= (SoftPackageUsedInGame[SoftPackageIdx] ? EDependencyProperty::Game : EDependencyProperty::None);
+
+		// Don't need to remove duplicates here because SavePackage only writes unique elements into SoftPackageReferenceList
 		PackageDependencies.Add({ DependencyPackageName, DependencyProperty });
 	}
 
diff --git a/Engine/Source/Runtime/AssetRegistry/Private/PackageDependencyData.h b/Engine/Source/Runtime/AssetRegistry/Private/PackageDependencyData.h
--- a/Engine/Source/Runtime/AssetRegistry/Private/PackageDependencyData.h
+++ b/Engine/Source/Runtime/AssetRegistry/Private/PackageDependencyData.h
@@ -23,6 +23,11 @@ public:
 			Dependency.Property = static_cast<UE::AssetRegistry::EDependencyProperty>(PropertyAsInteger);
 			return Ar;
 		}
+
+		bool operator==(const FPackageDependency& Other) const
+		{
+			return PackageName == Other.PackageName && Property == Other.Property;
+		}
 	};
 	struct FSearchableNamesDependency
 	{
diff --git a/Engine/Source/Runtime/AssetRegistry/Private/PathTree.cpp b/Engine/Source/Runtime/AssetRegistry/Private/PathTree.cpp
--- a/Engine/Source/Runtime/AssetRegistry/Private/PathTree.cpp
+++ b/Engine/Source/Runtime/AssetRegistry/Private/PathTree.cpp
@@ -64,25 +64,21 @@ bool FPathTree::CachePath(FName Path, TFunctionRef<void(FName)> OnPathAdded)
 			check(!CurrentPath.IsNone());	// Path parts cannot be empty
 			check(*(PathCharPtr-1) != '/'); // The previous character cannot be a /, as that would suggest a malformed path such as "/Game//MyAsset"
 
-			bool bAddedPath = false;
 			if (!ParentPathToChildPaths.Contains(CurrentPath))
 			{
 				ParentPathToChildPaths.Add(CurrentPath);
-				bAddedPath = true;
-			}
 
-			if (!LastPath.IsNone())
-			{
-				// Add us as a known child of our parent path
-				TSet<FName>& ChildPaths = ParentPathToChildPaths.FindChecked(LastPath);
-				ChildPaths.Add(CurrentPath);
+				// Only update the parent maps when this child is first added
+				if (!LastPath.IsNone())
+				{
+					// Add us as a known child of our parent path
+					TSet<FName>& ChildPaths = ParentPathToChildPaths.FindChecked(LastPath);
+					ChildPaths.Add(CurrentPath);
 
-				// Make sure we know how to find our parent again later on
-				ChildPathToParentPath.Add(CurrentPath, LastPath);
-			}
+					// Make sure we know how to find our parent again later on
+					ChildPathToParentPath.Add(CurrentPath, LastPath);
+				}
 
-			if (bAddedPath)
-			{
 				OnPathAdded(CurrentPath);
 			}
 
diff --git a/Engine/Source/Runtime/Engine/Private/Internationalization/EnginePackageLocalizationCache.cpp b/Engine/Source/Runtime/Engine/Private/Internationalization/EnginePackageLocalizationCache.cpp
--- a/Engine/Source/Runtime/Engine/Private/Internationalization/EnginePackageLocalizationCache.cpp
+++ b/Engine/Source/Runtime/Engine/Private/Internationalization/EnginePackageLocalizationCache.cpp
@@ -135,26 +135,32 @@ void FEnginePackageLocalizationCache::HandleAssetAdded(const FAssetData& InAsset
 		return;
 	}
 
+	// Convert the string outside the lock and loop	as this is called often while loading
+	const FString PackagePath = InAssetData.PackageName.ToString();
+
 	FScopeLock Lock(&LocalizedCachesCS);
 
 	for (auto& CultureCachePair : AllCultureCaches)
 	{
-		bPackageNameToAssetGroupDirty |= CultureCachePair.Value->AddPackage(InAssetData.PackageName.ToString());
+		bPackageNameToAssetGroupDirty |= CultureCachePair.Value->AddPackage(PackagePath);
 	}
 }
 
 void FEnginePackageLocalizationCache::HandleAssetRemoved(const FAssetData& InAssetData)
 {
+	const FString PackagePath = InAssetData.PackageName.ToString();
+
 	FScopeLock Lock(&LocalizedCachesCS);
 
 	for (auto& CultureCachePair : AllCultureCaches)
 	{
-		bPackageNameToAssetGroupDirty |= CultureCachePair.Value->RemovePackage(InAssetData.PackageName.ToString());
+		bPackageNameToAssetGroupDirty |= CultureCachePair.Value->RemovePackage(PackagePath);
 	}
 }
 
 void FEnginePackageLocalizationCache::HandleAssetRenamed(const FAssetData& InAssetData, const FString& InOldObjectPath)
 {
+	const FString PackagePath = InAssetData.PackageName.ToString();
 	const FString OldPackagePath = FPackageName::ObjectPathToPackageName(InOldObjectPath);
 
 	FScopeLock Lock(&LocalizedCachesCS);
@@ -162,6 +168,6 @@ void FEnginePackageLocalizationCache::HandleAssetRenamed(const FAssetData& InAss
 	for (auto& CultureCachePair : AllCultureCaches)
 	{
 		bPackageNameToAssetGroupDirty |= CultureCachePair.Value->RemovePackage(OldPackagePath);
-		bPackageNameToAssetGroupDirty |= CultureCachePair.Value->AddPackage(InAssetData.PackageName.ToString());
+		bPackageNameToAssetGroupDirty |= CultureCachePair.Value->AddPackage(PackagePath);
 	}
 }
