AssetRegistry diagnostics: Log the amount of memory used by each class for its AssetDataTags.
#rnx

diff --git a/Engine/Source/Runtime/AssetRegistry/Private/AssetRegistry.cpp b/Engine/Source/Runtime/AssetRegistry/Private/AssetRegistry.cpp
--- a/Engine/Source/Runtime/AssetRegistry/Private/AssetRegistry.cpp
+++ b/Engine/Source/Runtime/AssetRegistry/Private/AssetRegistry.cpp
@@ -4129,7 +4129,7 @@ Impl::EGatherStatus FAssetRegistryImpl::TickGatherer(Impl::FEventContext& EventC
 	return OutStatus;
 }
 
-void FAssetRegistryImpl::LogSearchDiagnostics(double StartTime) const
+void FAssetRegistryImpl::LogSearchDiagnostics(double StartTime)
 {
 	FAssetGatherDiagnostics Diagnostics = GlobalGatherer->GetDiagnostics();
 	float Total = Diagnostics.DiscoveryTimeSeconds + Diagnostics.GatherTimeSeconds + StoreGatherResultsTimeSeconds;
@@ -4144,6 +4144,24 @@ void FAssetRegistryImpl::LogSearchDiagnostics(double StartTime) const
 	FTelemetryRouter::Get().ProvideTelemetry(Telemetry);
 	UE_LOG(LogAssetRegistry, Log, TEXT("AssetRegistryGather time %.4fs: AssetDataDiscovery %0.4fs, AssetDataGather %0.4fs, StoreResults %0.4fs."),
 		Total, Diagnostics.DiscoveryTimeSeconds, Diagnostics.GatherTimeSeconds, StoreGatherResultsTimeSeconds);
+
+#if !NO_LOGGING
+	if (LogAssetRegistry.GetVerbosity() >= ELogVerbosity::Verbose)
+	{
+		UE_LOG(LogAssetRegistry, Verbose, TEXT("TagMemoryUse:"));
+		TagSizeByClass.ValueSort([](int64 A, int64 B) { return A > B; });
+		constexpr int64 MinSizeToLog = 1000 * 1000;
+		for (const TPair<FTopLevelAssetPath, int64>& Pair : TagSizeByClass)
+		{
+			if (Pair.Value < MinSizeToLog)
+			{
+				break;
+			}
+			UE_LOG(LogAssetRegistry, Verbose, TEXT("%s: %.1f MB"),
+				*Pair.Key.ToString(), (float)Pair.Value / (1000.f * 1000.f));
+		}
+	}
+#endif
 }
 
 void FAssetRegistryImpl::TickGatherPackage(Impl::FEventContext& EventContext, const FString& PackageName, const FString& LocalPath)
@@ -5062,6 +5080,12 @@ void FAssetRegistryImpl::AssetSearchDataGathered(Impl::FEventContext& EventConte
 			// The asset isn't in the cache yet, add it and notify subscribers
 			if (bPathIsMounted)
 			{
+				int64& ClassTagSizes = TagSizeByClass.FindOrAdd(BackgroundResult->AssetClassPath);
+				BackgroundResult->TagsAndValues.ForEach([&ClassTagSizes](const TPair<FName, FAssetTagValueRef>& Pair)
+					{
+						ClassTagSizes += Pair.Value.GetResourceSize();
+					});
+
 #if WITH_EDITOR
 				PostLoadAssetRegistryTags(BackgroundResult.Get());
 #endif
diff --git a/Engine/Source/Runtime/AssetRegistry/Private/AssetRegistryImpl.h b/Engine/Source/Runtime/AssetRegistry/Private/AssetRegistryImpl.h
--- a/Engine/Source/Runtime/AssetRegistry/Private/AssetRegistryImpl.h
+++ b/Engine/Source/Runtime/AssetRegistry/Private/AssetRegistryImpl.h
@@ -188,7 +188,7 @@ public:
 	/** Send a log message with the search statistics. 
 	 *  StartTime is used to report wall clock search time in the case of background scan
 	 */
-	void LogSearchDiagnostics(double StartTime) const;
+	void LogSearchDiagnostics(double StartTime);
 	/** Look for and load a single AssetData result from the gatherer. */
 	void TickGatherPackage(Impl::FEventContext& EventContext, const FString& PackageName, const FString& LocalPath);
 	void ClearGathererCache();
@@ -388,6 +388,9 @@ private:
 	/** Lists of results from the background thread that are waiting to get processed by the main thread */
 	FAssetDataGatherer::FResults BackgroundResults;
 
+	/** Memory profiling information: How much memory is being used by the tags for each class. */
+	TMap<FTopLevelAssetPath, int64> TagSizeByClass;
+
 	/** Time spent processing Gather results */
 	float StoreGatherResultsTimeSeconds;
 	/** The highest number of pending results observed during initial gathering */
diff --git a/Engine/Source/Runtime/CoreUObject/Private/AssetRegistry/AssetDataTagMap.cpp b/Engine/Source/Runtime/CoreUObject/Private/AssetRegistry/AssetDataTagMap.cpp
--- a/Engine/Source/Runtime/CoreUObject/Private/AssetRegistry/AssetDataTagMap.cpp
+++ b/Engine/Source/Runtime/CoreUObject/Private/AssetRegistry/AssetDataTagMap.cpp
@@ -192,6 +192,11 @@ namespace FixedTagPrivate
 		return FTextStringHelper::CreateFromBuffer(*String);
 	}
 
+	int64 FMarshalledText::GetResourceSize() const
+	{
+		return String.GetAllocatedSize();
+	}
+
 	int32 FMarshalledText::CompareToCaseIgnored(const FMarshalledText& Other) const
 	{
 		return String.Compare(Other.String);
@@ -433,6 +438,27 @@ namespace FixedTagPrivate
 		return FString();
 	}
 
+	int64 FValueHandle::GetResourceSize() const
+	{
+		FStore& Store = GStores[StoreIndex];
+		uint32 Index = Id.Index;
+
+		switch (Id.Type)
+		{
+		case EValueType::AnsiString:			return TCString<ANSICHAR>::Strlen(Store.GetAnsiString(Index))*sizeof(ANSICHAR);
+		case EValueType::WideString:			return TCString<WIDECHAR>::Strlen(Store.GetWideString(Index))*sizeof(WIDECHAR);
+		case EValueType::NumberlessName:		return sizeof(FDisplayNameEntryId);
+		case EValueType::Name:					return sizeof(FName);
+		case EValueType::NumberlessExportPath:	return sizeof(FNumberlessExportPath);
+		case EValueType::ExportPath:			return sizeof(FAssetRegistryExportPath);
+		case EValueType::LocalizedText:			return sizeof(Store.Texts[Index]) + Store.Texts[Index].GetResourceSize();
+		}
+
+		check(false);
+		return 0;
+	}
+
+
 	FString	FValueHandle::AsDisplayString() const
 	{
 		return AsString<false /* bForStorage */>();
@@ -1377,6 +1403,11 @@ FString FAssetTagValueRef::ToLoose() const
 	return IsFixed() ? AsFixed().AsStorageString() : AsLoose();
 }
 
+int64 FAssetTagValueRef::GetResourceSize() const
+{
+	return IsFixed() ? AsFixed().GetResourceSize() : AsLoose().GetAllocatedSize();
+}
+
 bool FAssetTagValueRef::TryGetAsMarshalledText(FixedTagPrivate::FMarshalledText& Out) const
 {
 	if (IsFixed())
diff --git a/Engine/Source/Runtime/CoreUObject/Public/AssetRegistry/AssetDataTagMap.h b/Engine/Source/Runtime/CoreUObject/Public/AssetRegistry/AssetDataTagMap.h
--- a/Engine/Source/Runtime/CoreUObject/Public/AssetRegistry/AssetDataTagMap.h
+++ b/Engine/Source/Runtime/CoreUObject/Public/AssetRegistry/AssetDataTagMap.h
@@ -127,6 +127,7 @@ namespace FixedTagPrivate
 		COREUOBJECT_API bool						AsMarshalledText(FMarshalledText& Out) const;
 		COREUOBJECT_API bool						Equals(FStringView Str) const;
 		COREUOBJECT_API bool						Contains(const TCHAR* Str) const;
+		COREUOBJECT_API int64						GetResourceSize() const;
 
 	private:
 		template <bool bForStorage>
@@ -253,7 +254,7 @@ public:
 	FAssetTagValueRef& operator=(const FAssetTagValueRef&) = default;
 	FAssetTagValueRef& operator=(FAssetTagValueRef&&) = default;
 
-	bool						IsSet() const { return Bits != 0; }
+	bool										IsSet() const { return Bits != 0; }
 
 	COREUOBJECT_API FString						AsString() const;
 	COREUOBJECT_API FName						AsName() const;
@@ -261,15 +262,20 @@ public:
 	COREUOBJECT_API FText						AsText() const;
 	COREUOBJECT_API bool						TryGetAsText(FText& Out) const; // @return false if value isn't a localized string
 
-	FString						GetValue() const { return AsString(); }
+	FString										GetValue() const { return AsString(); }
 	/** Coerce the type to a Complex String capable of representing the type */
-	FString						GetStorageString() const { return ToLoose(); }
+	FString										GetStorageString() const { return ToLoose(); }
+	/**
+	 * Measure how much memory is used by the value. Does not account for deduplication, adding the results
+	 * for keys sharing a duplicated value will overreport how much memory is used.
+	 */
+	COREUOBJECT_API int64						GetResourceSize() const;
 
-	COREUOBJECT_API bool		Equals(FStringView Str) const;
+	COREUOBJECT_API bool						Equals(FStringView Str) const;
 
 private:
 	/** Return whether this's value is a MarshalledFText, and copy it into out parameter if so */
-	bool						TryGetAsMarshalledText(FixedTagPrivate::FMarshalledText& Out) const;
+	bool										TryGetAsMarshalledText(FixedTagPrivate::FMarshalledText& Out) const;
 	/**
 	 * Copy this's value (whether loose or fixed) into the loose format.
 	 * The returned loose value is in StorageFormat (e.g. complex strings) rather than display format.
diff --git a/Engine/Source/Runtime/CoreUObject/Public/AssetRegistry/AssetDataTagMapSerializationDetails.h b/Engine/Source/Runtime/CoreUObject/Public/AssetRegistry/AssetDataTagMapSerializationDetails.h
--- a/Engine/Source/Runtime/CoreUObject/Public/AssetRegistry/AssetDataTagMapSerializationDetails.h
+++ b/Engine/Source/Runtime/CoreUObject/Public/AssetRegistry/AssetDataTagMapSerializationDetails.h
@@ -34,6 +34,7 @@ namespace FixedTagPrivate
 		COREUOBJECT_API const FString& GetAsComplexString() const;
 		COREUOBJECT_API FText GetAsText() const;
 		COREUOBJECT_API int32 CompareToCaseIgnored(const FMarshalledText& Other) const;
+		COREUOBJECT_API int64 GetResourceSize() const;
 
 	private:
 		FString String;
