AssetRegistry: Change FAssetData returned from AR queries in the non-editor runtime to have GetChunkIds from the on-disk version of the data whether or not bIncludeOnlyOnDiskAssets is true, because GetChunkIds in non-editor runtimes always returns an empty array and is not what users expect. Add comments for the bIncludeOnlyOnDiskAssets contract.
#rnx

diff --git a/Engine/Source/Runtime/AssetRegistry/Private/AssetRegistry.cpp b/Engine/Source/Runtime/AssetRegistry/Private/AssetRegistry.cpp
--- a/Engine/Source/Runtime/AssetRegistry/Private/AssetRegistry.cpp
+++ b/Engine/Source/Runtime/AssetRegistry/Private/AssetRegistry.cpp
@@ -2213,6 +2213,14 @@ void EnumerateMemoryAssets(const FARCompiledFilter& InFilter, TSet<FName>& OutPa
 					{
 						PartialAssetData.TagsAndValues = FAssetDataTagMapSharedView(MoveTemp(*ModifiedTags));
 					}
+#if !WITH_EDITORONLY_DATA
+					// In non-editor builds, UObject::GetChunkIds returns an empty set.
+					// Like our contract for tags, when the information is missing from the UObject, our contract
+					// for that information in AssetRegistry queries is that we return the on-disk version of the data.
+					// The on-disk version of the data for GetChunkIds is the data that was stored in the generated
+					// AssetRegistry by calling AddChunkId for each chunkID that the cooker found the Asset to be in.
+					PartialAssetData.SetChunkIDs(OnDiskAssetData->GetChunkIDs());
+#endif
 				}
 			}
 			// After adding tags, PartialAssetData is now a full AssetData
diff --git a/Engine/Source/Runtime/AssetRegistry/Public/AssetRegistry/IAssetRegistry.h b/Engine/Source/Runtime/AssetRegistry/Public/AssetRegistry/IAssetRegistry.h
--- a/Engine/Source/Runtime/AssetRegistry/Public/AssetRegistry/IAssetRegistry.h
+++ b/Engine/Source/Runtime/AssetRegistry/Public/AssetRegistry/IAssetRegistry.h
@@ -150,6 +150,28 @@ class UAssetRegistry : public UInterface
 	GENERATED_UINTERFACE_BODY()
 };
 
+/**
+ * Global singleton interface for accessing a catalog of all packages (and some other content file types) that are
+ * stored in any mounted directory. In editor this information is gathered from the package files on disk during a
+ * gather step at editor startup. In cooked runtimes this information was calculated during cook and is serialized out
+ * of a single file (after pruning information not necessary at runtime.)
+ * 
+ * Some API notes:
+ *
+ * bIncludeOnlyOnDiskAssets
+ *     Most query functions that return FAssetData take this argument. If true, only data collected from disk and
+ *     stored in the AssetRegistry will be returned. If false, and the object is loaded in memory, the returned
+ *     AssetData will be calculated from the object in memory because the InMemoryData is more likely to be
+ *     up-to-date. The InMemory data will sometimes vary from the DiskGatheredData.
+ * 
+ *     When InMemoryData is returned some categories of data that are always missing from the object in memory (e.g.
+ *     GetAssetRegistryTagsExtended) are read from the DiskGatheredData and added to the InMemoryData.
+ * 
+ *     Setting this value to true will always be faster than setting it to false, because the same registry
+ *     lookups are performed in either case, but the InMemoryData lookup is skipped in the true case.
+ * 
+ *     The default is usually false.
+ */
 class IAssetRegistry
 {
 	GENERATED_IINTERFACE_BODY()
@@ -179,30 +201,39 @@ public:
 	 *
 	 * @param PackageName the package name for the requested assets (eg, /Game/MyFolder/MyAsset)
 	 * @param OutAssetData the list of assets in this path
+	 * @param bIncludeOnlyOnDiskAssets If true, use only DiskGatheredData, do not calculate from UObjects.
+	 *        @see IAssetRegistry class header for bIncludeOnlyOnDiskAssets.
 	 * @param bSkipARFilteredAssets If true, skips Objects that return true for IsAsset but are not assets in the current platform.
 	 */
 	UFUNCTION(BlueprintCallable, BlueprintPure=false, Category="AssetRegistry")
-	virtual bool GetAssetsByPackageName(FName PackageName, TArray<FAssetData>& OutAssetData, bool bIncludeOnlyOnDiskAssets = false, bool bSkipARFilteredAssets=true) const = 0;
+	virtual bool GetAssetsByPackageName(FName PackageName, TArray<FAssetData>& OutAssetData,
+		bool bIncludeOnlyOnDiskAssets = false, bool bSkipARFilteredAssets=true) const = 0;
 
 	/**
 	 * Gets asset data for all assets in the supplied folder path
 	 *
 	 * @param PackagePath the path to query asset data in (eg, /Game/MyFolder)
 	 * @param OutAssetData the list of assets in this path
+	 * @param bIncludeOnlyOnDiskAssets If true, use only DiskGatheredData, do not calculate from UObjects.
+	 *        @see IAssetRegistry class header for bIncludeOnlyOnDiskAssets.
 	 * @param bRecursive if true, all supplied paths will be searched recursively
 	 */
 	UFUNCTION(BlueprintCallable, BlueprintPure=false, Category = "AssetRegistry")
-	virtual bool GetAssetsByPath(FName PackagePath, TArray<FAssetData>& OutAssetData, bool bRecursive = false, bool bIncludeOnlyOnDiskAssets = false) const = 0;
+	virtual bool GetAssetsByPath(FName PackagePath, TArray<FAssetData>& OutAssetData, bool bRecursive = false,
+		bool bIncludeOnlyOnDiskAssets = false) const = 0;
 
 	/**
 	 * Gets asset data for all assets in any of the supplied folder paths
 	 *
 	 * @param PackagePaths the paths to query asset data in (eg, /Game/MyFolder)
 	 * @param OutAssetData the list of assets in this path
+	 * @param bIncludeOnlyOnDiskAssets If true, use only DiskGatheredData, do not calculate from UObjects.
+	 *        @see IAssetRegistry class header for bIncludeOnlyOnDiskAssets.
 	 * @param bRecursive if true, all supplied paths will be searched recursively
 	 */
 	UFUNCTION(BlueprintCallable, BlueprintPure = false, Category = "AssetRegistry")
-	virtual bool GetAssetsByPaths(TArray<FName> PackagePaths, TArray<FAssetData>& OutAssetData, bool bRecursive = false, bool bIncludeOnlyOnDiskAssets = false) const = 0;
+	virtual bool GetAssetsByPaths(TArray<FName> PackagePaths, TArray<FAssetData>& OutAssetData, bool bRecursive = false,
+		bool bIncludeOnlyOnDiskAssets = false) const = 0;
 
 	/**
 	 * Gets asset data for all assets with the supplied class
@@ -263,7 +294,8 @@ public:
 	 * Gets the asset data for the specified object path
 	 *
 	 * @param ObjectPath the path of the object to be looked up
-	 * @param bIncludeOnlyOnDiskAssets if true, in-memory objects will be ignored. The call will be faster.
+	 * @param bIncludeOnlyOnDiskAssets If true, use only DiskGatheredData, do not calculate from UObjects.
+	 *        @see IAssetRegistry class header for bIncludeOnlyOnDiskAssets.
 	 * @return the assets data;Will be invalid if object could not be found
 	 */
 	UE_DEPRECATED(5.1, "Asset path FNames have been deprecated, use Soft Object Path instead.")
@@ -274,22 +306,26 @@ public:
 	 * Gets the asset data for the specified object path
 	 *
 	 * @param ObjectPath the path of the object to be looked up
-	 * @param bIncludeOnlyOnDiskAssets if true, in-memory objects will be ignored. The call will be faster.
+	 * @param bIncludeOnlyOnDiskAssets If true, use only DiskGatheredData, do not calculate from UObjects.
+	 *        @see IAssetRegistry class header for bIncludeOnlyOnDiskAssets.
 	 * @param bSkipARFilteredAssets If true, skips Objects that return true for IsAsset but are not assets in the current platform.
 	 * @return the assets data;Will be invalid if object could not be found
 	 */
 	UFUNCTION(BlueprintCallable, BlueprintPure=false, Category = "AssetRegistry", DisplayName="Get Asset By Object Path")
-	ASSETREGISTRY_API virtual FAssetData K2_GetAssetByObjectPath(const FSoftObjectPath& ObjectPath, bool bIncludeOnlyOnDiskAssets = false, bool bSkipARFilteredAssets = true) const;
+	ASSETREGISTRY_API virtual FAssetData K2_GetAssetByObjectPath(const FSoftObjectPath& ObjectPath,
+		bool bIncludeOnlyOnDiskAssets = false, bool bSkipARFilteredAssets = true) const;
 
 	/**
 	 * Gets the asset data for the specified object path
 	 *
 	 * @param ObjectPath the path of the object to be looked up
-	 * @param bIncludeOnlyOnDiskAssets if true, in-memory objects will be ignored. The call will be faster.
+	 * @param bIncludeOnlyOnDiskAssets If true, use only DiskGatheredData, do not calculate from UObjects.
+	 *        @see IAssetRegistry class header for bIncludeOnlyOnDiskAssets.
 	 * @param bSkipARFilteredAssets If true, skips Objects that return true for IsAsset but are not assets in the current platform.
 	 * @return the assets data;Will be invalid if object could not be found
 	 */
-	virtual FAssetData GetAssetByObjectPath(const FSoftObjectPath& ObjectPath, bool bIncludeOnlyOnDiskAssets = false, bool bSkipARFilteredAssets = true) const = 0;
+	virtual FAssetData GetAssetByObjectPath(const FSoftObjectPath& ObjectPath,
+		bool bIncludeOnlyOnDiskAssets = false, bool bSkipARFilteredAssets = true) const = 0;
 
 	/**
 	 * Tries to get the asset data for the specified object path
@@ -323,8 +359,11 @@ public:
 	 * This method may be slow, use a filter if possible to avoid iterating over the entire registry.
 	 *
 	 * @param Callback function to call for each asset data enumerated
+	 * @param bIncludeOnlyOnDiskAssets If true, use only DiskGatheredData, do not calculate from UObjects.
+	 *        @see IAssetRegistry class header for bIncludeOnlyOnDiskAssets.
 	 */
-	virtual bool EnumerateAllAssets(TFunctionRef<bool(const FAssetData&)> Callback, bool bIncludeOnlyOnDiskAssets = false) const = 0;
+	virtual bool EnumerateAllAssets(TFunctionRef<bool(const FAssetData&)> Callback,
+		bool bIncludeOnlyOnDiskAssets = false) const = 0;
 
 	/**
 	 * Gets the LongPackageName for all packages with the given PackageName.
diff --git a/Engine/Source/Runtime/CoreUObject/Public/AssetRegistry/ARFilter.h b/Engine/Source/Runtime/CoreUObject/Public/AssetRegistry/ARFilter.h
--- a/Engine/Source/Runtime/CoreUObject/Public/AssetRegistry/ARFilter.h
+++ b/Engine/Source/Runtime/CoreUObject/Public/AssetRegistry/ARFilter.h
@@ -70,7 +70,10 @@ PRAGMA_ENABLE_DEPRECATION_WARNINGS
 	/** If true, subclasses of ClassPaths will also be included and RecursiveClassPathsExclusionSet will be excluded. */
 	bool bRecursiveClasses = false;
 
-	/** If true, only on-disk assets will be returned. Be warned that this is rarely what you want and should only be used for performance reasons */
+	/**
+	 * If true, use only DiskGatheredData, do not calculate from UObjects. @see IAssetRegistry class header for
+	 * bIncludeOnlyOnDiskAssets.
+	 */
 	bool bIncludeOnlyOnDiskAssets = false;
 
 	/** The exclusive filter component for package flags. Only assets without any of the specified flags will be returned. */
