AssetRegistry: Turn off collection of tag memory sizes when it is not being reported because verbose is not on. The collection is expensive in cputime.
#rnx

diff --git a/Engine/Source/Runtime/AssetRegistry/Private/AssetRegistry.cpp b/Engine/Source/Runtime/AssetRegistry/Private/AssetRegistry.cpp
--- a/Engine/Source/Runtime/AssetRegistry/Private/AssetRegistry.cpp
+++ b/Engine/Source/Runtime/AssetRegistry/Private/AssetRegistry.cpp
@@ -828,6 +828,11 @@ void FAssetRegistryImpl::Initialize(Impl::FInitializeContext& Context)
 	bInitialSearchCompleted = true;
 	GatherStatus = Impl::EGatherStatus::Active;
 	bSearchAllAssets = false;
+#if NO_LOGGING
+	bVerboseLogging = false;
+#else
+	bVerboseLogging = LogAssetRegistry.GetVerbosity() >= ELogVerbosity::Verbose;
+#endif
 	StoreGatherResultsTimeSeconds = 0.f;
 
 	// By default update the disk cache once on asset load, to incorporate changes made in PostLoad. This only happens in editor builds
@@ -4154,7 +4159,7 @@ void FAssetRegistryImpl::LogSearchDiagnostics(double StartTime)
 		Total, Diagnostics.DiscoveryTimeSeconds, Diagnostics.GatherTimeSeconds, StoreGatherResultsTimeSeconds);
 
 #if !NO_LOGGING
-	if (LogAssetRegistry.GetVerbosity() >= ELogVerbosity::Verbose)
+	if (bVerboseLogging)
 	{
 		UE_LOG(LogAssetRegistry, Verbose, TEXT("TagMemoryUse:"));
 		TagSizeByClass.ValueSort([](int64 A, int64 B) { return A > B; });
@@ -5088,11 +5093,16 @@ void FAssetRegistryImpl::AssetSearchDataGathered(Impl::FEventContext& EventConte
 			// The asset isn't in the cache yet, add it and notify subscribers
 			if (bPathIsMounted)
 			{
-				int64& ClassTagSizes = TagSizeByClass.FindOrAdd(BackgroundResult->AssetClassPath);
-				BackgroundResult->TagsAndValues.ForEach([&ClassTagSizes](const TPair<FName, FAssetTagValueRef>& Pair)
-					{
-						ClassTagSizes += Pair.Value.GetResourceSize();
-					});
+#if !NO_LOGGING
+				if (bVerboseLogging)
+				{
+					int64& ClassTagSizes = TagSizeByClass.FindOrAdd(BackgroundResult->AssetClassPath);
+					BackgroundResult->TagsAndValues.ForEach([&ClassTagSizes](const TPair<FName, FAssetTagValueRef>& Pair)
+						{
+							ClassTagSizes += Pair.Value.GetResourceSize();
+						});
+				}
+#endif
 
 #if WITH_EDITOR
 				PostLoadAssetRegistryTags(BackgroundResult.Get());
diff --git a/Engine/Source/Runtime/AssetRegistry/Private/AssetRegistryImpl.h b/Engine/Source/Runtime/AssetRegistry/Private/AssetRegistryImpl.h
--- a/Engine/Source/Runtime/AssetRegistry/Private/AssetRegistryImpl.h
+++ b/Engine/Source/Runtime/AssetRegistry/Private/AssetRegistryImpl.h
@@ -388,8 +388,10 @@ private:
 	/** Lists of results from the background thread that are waiting to get processed by the main thread */
 	FAssetDataGatherer::FResults BackgroundResults;
 
+#if !NO_LOGGING
 	/** Memory profiling information: How much memory is being used by the tags for each class. */
 	TMap<FTopLevelAssetPath, int64> TagSizeByClass;
+#endif
 
 	/** Time spent processing Gather results */
 	float StoreGatherResultsTimeSeconds;
@@ -419,6 +421,8 @@ private:
 	/** Record whether SearchAllAssets has been called; if so we will also search new mountpoints when added */
 	bool bSearchAllAssets;
 
+	bool bVerboseLogging;
+
 	/** List of all class names derived from Blueprint (including Blueprint itself) */
 	TSet<FTopLevelAssetPath> ClassGeneratorNames;
 
