AssetRegistry DependsNode: PR #10777: Fixed false positive FDependsNode automation test
#rnx

diff --git a/Engine/Source/Runtime/AssetRegistry/Private/DependsNodeTest.cpp b/Engine/Source/Runtime/AssetRegistry/Private/DependsNodeTest.cpp
--- a/Engine/Source/Runtime/AssetRegistry/Private/DependsNodeTest.cpp
+++ b/Engine/Source/Runtime/AssetRegistry/Private/DependsNodeTest.cpp
@@ -753,7 +753,7 @@ bool FDependsNodeTest::RunTest(const FString& Parameters)
 				TestTrue(TEXT("IterateOverDependencies two property corners expected values"), Dependency == &B && Category == EDependencyCategory::Package && (Properties == EDependencyProperty::Game || Properties == EDependencyProperty::Hard));
 				TestTrue(TEXT("IterateOverDependencies two property corners bDuplicate is set correctly"), IterationIndex == 0 || bDuplicate);
 				bHasGame = bHasGame || Properties == EDependencyProperty::Game;
-				bHasHard = bHasHard || Properties == EDependencyProperty::Game;
+				bHasHard = bHasHard || Properties == EDependencyProperty::Hard;
 				IterationIndex++;
 			});
 		TestTrue(TEXT("IterateOverDependencies two property corners hits both corners"), bHasGame && bHasHard);
@@ -767,7 +767,7 @@ bool FDependsNodeTest::RunTest(const FString& Parameters)
 				TestTrue(TEXT("IterateOverDependencies two property corners, node-specific, expected values"), Dependency == &B && Category == EDependencyCategory::Package && (Properties == EDependencyProperty::Game || Properties == EDependencyProperty::Hard));
 				TestTrue(TEXT("IterateOverDependencies two property corners, node-specific, bDuplicate is set correctly"), IterationIndex == 0 || bDuplicate);
 				bHasGame = bHasGame || Properties == EDependencyProperty::Game;
-				bHasHard = bHasHard || Properties == EDependencyProperty::Game;
+				bHasHard = bHasHard || Properties == EDependencyProperty::Hard;
 				IterationIndex++;
 			}, &B);
 		TestTrue(TEXT("IterateOverDependencies two property corners, node-specific, hits both corners"), bHasGame&& bHasHard);
@@ -816,7 +816,7 @@ bool FDependsNodeTest::RunTest(const FString& Parameters)
 				TestTrue(TEXT("AddPackageDependencySet only two elements"), IterationIndex < 2);
 				TestTrue(TEXT("AddPackageDependencySets expected values"), Dependency == &B && Category == EDependencyCategory::Package && (Properties == EDependencyProperty::Game || Properties == EDependencyProperty::Hard));
 				bHasGame = bHasGame || Properties == EDependencyProperty::Game;
-				bHasHard = bHasHard || Properties == EDependencyProperty::Game;
+				bHasHard = bHasHard || Properties == EDependencyProperty::Hard;
 				IterationIndex++;
 			});
 		TestTrue(TEXT("IterateOverDependencies two property corners hits both corners"), bHasGame && bHasHard);
@@ -1216,4 +1216,4 @@ bool FDependsNodeTest::RunTest(const FString& Parameters)
 	return true;
 }
 
-#endif
\ No newline at end of file
+#endif
