Don't attempt to use platform chunk installer to query chunk status if it has been disabled in the configuration.

#rnx
#preflighjt 64e8b587c87c17228eb60dbf

diff --git a/Engine/Source/Runtime/AssetRegistry/Private/AssetRegistry.cpp b/Engine/Source/Runtime/AssetRegistry/Private/AssetRegistry.cpp
--- a/Engine/Source/Runtime/AssetRegistry/Private/AssetRegistry.cpp
+++ b/Engine/Source/Runtime/AssetRegistry/Private/AssetRegistry.cpp
@@ -164,6 +164,10 @@ static FAutoConsoleVariableRef CVarAssetRegistryMaxSecondsPerFrame(
 	UE::AssetRegistry::Impl::MaxSecondsPerFrame,
 	TEXT("Maximum amount of time allowed for Asset Registry processing, in seconds"));
 
+#ifndef ENABLE_PLATFORM_CHUNK_INSTALL
+	#define ENABLE_PLATFORM_CHUNK_INSTALL (1)
+#endif
+
 /**
  * Implementation of IAssetRegistryInterface; forwards calls from the CoreUObject-accessible IAssetRegistryInterface into the AssetRegistry-accessible IAssetRegistry
  */
@@ -3280,6 +3284,7 @@ namespace UE::AssetRegistry::Utils
 
 EAssetAvailability::Type GetAssetAvailability(const FAssetData& AssetData)
 {
+#if ENABLE_PLATFORM_CHUNK_INSTALL
 	IPlatformChunkInstall* ChunkInstall = FPlatformMisc::GetPlatformChunkInstall();
 
 	EChunkLocation::Type BestLocation = EChunkLocation::DoesNotExist;
@@ -3316,6 +3321,9 @@ EAssetAvailability::Type GetAssetAvailability(const FAssetData& AssetData)
 		check(0);
 		return EAssetAvailability::LocalFast;
 	}
+#else
+	return EAssetAvailability::LocalFast;
+#endif
 }
 
 }
@@ -3330,11 +3338,13 @@ namespace UE::AssetRegistry::Utils
 
 float GetAssetAvailabilityProgress(const FAssetData& AssetData, EAssetAvailabilityProgressReportingType::Type ReportType)
 {
+	check(ReportType == EAssetAvailabilityProgressReportingType::PercentageComplete || ReportType == EAssetAvailabilityProgressReportingType::ETA);
+
+#if ENABLE_PLATFORM_CHUNK_INSTALL
 	IPlatformChunkInstall* ChunkInstall = FPlatformMisc::GetPlatformChunkInstall();
 	EChunkProgressReportingType::Type ChunkReportType = GetChunkAvailabilityProgressType(ReportType);
 
 	bool IsPercentageComplete = (ChunkReportType == EChunkProgressReportingType::PercentageComplete) ? true : false;
-	check(ReportType == EAssetAvailabilityProgressReportingType::PercentageComplete || ReportType == EAssetAvailabilityProgressReportingType::ETA);
 
 	float BestProgress = MAX_FLT;
 
@@ -3367,6 +3377,18 @@ float GetAssetAvailabilityProgress(const FAssetData& AssetData, EAssetAvailabili
 		BestProgress = 100.0f - BestProgress;
 	}
 	return BestProgress;
+
+#else
+	if (ReportType == EAssetAvailabilityProgressReportingType::PercentageComplete)
+	{
+		return 100.0f;
+	}
+	else
+	{
+		return 0.0f;
+	}
+#endif
+
 }
 
 }
@@ -3381,8 +3403,12 @@ namespace UE::AssetRegistry::Utils
 
 bool GetAssetAvailabilityProgressTypeSupported(EAssetAvailabilityProgressReportingType::Type ReportType)
 {
+#if ENABLE_PLATFORM_CHUNK_INSTALL
 	IPlatformChunkInstall* ChunkInstall = FPlatformMisc::GetPlatformChunkInstall();
 	return ChunkInstall->GetProgressReportingTypeSupported(GetChunkAvailabilityProgressType(ReportType));
+#else
+	return true;
+#endif
 }
 
 }
@@ -3397,6 +3423,7 @@ namespace UE::AssetRegistry::Utils
 
 void PrioritizeAssetInstall(const FAssetData& AssetData)
 {
+#if ENABLE_PLATFORM_CHUNK_INSTALL
 	IPlatformChunkInstall* ChunkInstall = FPlatformMisc::GetPlatformChunkInstall();
 
 	const TConstArrayView<int32> ChunkIDs = AssetData.GetChunkIDs();
@@ -3406,6 +3433,7 @@ void PrioritizeAssetInstall(const FAssetData& AssetData)
 	}
 
 	ChunkInstall->PrioritizePakchunk(ChunkIDs[0], EChunkPriority::Immediate);
+#endif
 }
 
 }
diff --git a/Engine/Source/Runtime/Engine/Private/AssetManager.cpp b/Engine/Source/Runtime/Engine/Private/AssetManager.cpp
--- a/Engine/Source/Runtime/Engine/Private/AssetManager.cpp
+++ b/Engine/Source/Runtime/Engine/Private/AssetManager.cpp
@@ -52,6 +52,10 @@ LLM_DEFINE_TAG(AssetManager);
 
 DEFINE_LOG_CATEGORY(LogAssetManager);
 
+#ifndef ENABLE_PLATFORM_CHUNK_INSTALL
+	#define ENABLE_PLATFORM_CHUNK_INSTALL (1)
+#endif
+
 /** Structure defining the current loading state of an asset */
 struct FPrimaryAssetLoadState
 {
@@ -2255,7 +2259,9 @@ bool UAssetManager::FindMissingChunkList(const TArray<FSoftObjectPath>& AssetLis
 	TMap<int32, EChunkLocation::Type> ChunkLocationCache;
 
 	// Grab chunk install
+#if ENABLE_PLATFORM_CHUNK_INSTALL
 	IPlatformChunkInstall* ChunkInstall = FPlatformMisc::GetPlatformChunkInstall();
+#endif
 
 	// Grab pak platform file
 	FPakPlatformFile* Pak = (FPakPlatformFile*)FPlatformFileManager::Get().FindPlatformFile(TEXT("PakFile"));
@@ -2271,7 +2277,11 @@ bool UAssetManager::FindMissingChunkList(const TArray<FSoftObjectPath>& AssetLis
 		{
 			if (!ChunkLocationCache.Contains(PakchunkId))
 			{
+#if ENABLE_PLATFORM_CHUNK_INSTALL
 				EChunkLocation::Type Location = ChunkInstall->GetPakchunkLocation(PakchunkId);
+#else
+				EChunkLocation::Type Location = EChunkLocation::LocalFast;
+#endif
 
 				// If chunk install thinks the chunk is available, we need to double check with the pak system that it isn't
 				// pending decryption
@@ -2342,6 +2352,7 @@ bool UAssetManager::FindMissingChunkList(const TArray<FSoftObjectPath>& AssetLis
 
 void UAssetManager::AcquireChunkList(const TArray<int32>& ChunkList, FAssetManagerAcquireResourceDelegate CompleteDelegate, EChunkPriority::Type Priority, TSharedPtr<FStreamableHandle> StalledHandle)
 {
+#if ENABLE_PLATFORM_CHUNK_INSTALL
 	FPendingChunkInstall* PendingChunkInstall = new(PendingChunkInstalls) FPendingChunkInstall;
 	PendingChunkInstall->ManualCallback = MoveTemp(CompleteDelegate);
 	PendingChunkInstall->RequestedChunks = ChunkList;
@@ -2359,6 +2370,7 @@ void UAssetManager::AcquireChunkList(const TArray<int32>& ChunkList, FAssetManag
 	{
 		ChunkInstall->PrioritizePakchunk(MissingChunk, Priority);
 	}
+#endif
 }
 
 void UAssetManager::AcquireResourcesForAssetList(const TArray<FSoftObjectPath>& AssetList, FAssetManagerAcquireResourceDelegate CompleteDelegate, EChunkPriority::Type Priority)
@@ -2445,6 +2457,7 @@ bool UAssetManager::GetResourceAcquireProgress(int32& OutAcquiredCount, int32& O
 
 void UAssetManager::OnChunkDownloaded(uint32 ChunkId, bool bSuccess)
 {
+#if ENABLE_PLATFORM_CHUNK_INSTALL
 	IPlatformChunkInstall* ChunkInstall = FPlatformMisc::GetPlatformChunkInstall();
 
 	// Iterate pending callbacks, in order they were added
@@ -2518,6 +2531,7 @@ void UAssetManager::OnChunkDownloaded(uint32 ChunkId, bool bSuccess)
 			}
 		}
 	}
+#endif
 }
 
 bool UAssetManager::OnAssetRegistryAvailableAfterInitialization(FName InName, FAssetRegistryState& OutNewState)
