Add and enable AssetRegistry.IgnoreEmptyDirectories cvar to ignore empty directories left behind by renames to improve load times
Add AssetRegistry.DeferDependencySort and DeferReferencerSort cvars to delay sorting dependency nodes, enabled DeferReferencerSort as it helps load times
Ignore directory rescan requests during the initial async scan as the last access time updating can lead to long main thread stalls
Optimize the computed dependencies map and skip package functions
Add a utility function to set sorting and cleanup some namespaces and dll exports

diff --git a/Engine/Source/Runtime/AssetRegistry/Private/AssetDataGatherer.cpp b/Engine/Source/Runtime/AssetRegistry/Private/AssetDataGatherer.cpp
--- a/Engine/Source/Runtime/AssetRegistry/Private/AssetDataGatherer.cpp
+++ b/Engine/Source/Runtime/AssetRegistry/Private/AssetDataGatherer.cpp
@@ -65,16 +65,21 @@ namespace AssetDataGathererConstants
 namespace UE::AssetDataGather::Private
 {
 
-static bool BlockPackagesWithMarkOfTheWeb = false;
+bool bBlockPackagesWithMarkOfTheWeb = false;
 static FAutoConsoleVariableRef CVarBlockPackagesWithMarkOfTheWeb(
 	TEXT("AssetRegistry.BlockPackagesWithMarkOfTheWeb"),
-	BlockPackagesWithMarkOfTheWeb,
-	TEXT("Whether package files with mark of the web are blocked from the asset registry"),
-	ECVF_SetBySystemSettingsIni);
+	bBlockPackagesWithMarkOfTheWeb,
+	TEXT("Whether package files with mark of the web are blocked from the asset registry"));
+
+bool bIgnoreEmptyDirectories = true;
+static FAutoConsoleVariableRef CVarIgnoreEmptyDirectories(
+	TEXT("AssetRegistry.IgnoreEmptyDirectories"),
+	bIgnoreEmptyDirectories,
+	TEXT("If true, completely empty leaf directories are ignored by the asset registry while scanning"));
 
 bool IsPackageBlocked(const FStringView& FilePath)
 {
-	return BlockPackagesWithMarkOfTheWeb && IPlatformFile::GetPlatformPhysical().HasMarkOfTheWeb(*FString(FilePath));
+	return bBlockPackagesWithMarkOfTheWeb && IPlatformFile::GetPlatformPhysical().HasMarkOfTheWeb(*FString(FilePath));
 }
 
 /** A structure to hold serialized cache data from async loads before adding it to the Gatherer's main cache. */
@@ -1461,7 +1466,7 @@ void FAssetDataDiscovery::TickInternal(bool bTickAll)
 							Data.ScanDir->SetScanResults(Data.DirLocalAbsPath, Data.ParentData, LocalSubDirs, LocalDiscoveredFiles);
 							if (!LocalSubDirs.IsEmpty() || !LocalDiscoveredFiles.IsEmpty())
 							{
-								AddDiscovered(Data.DirLocalAbsPath, LocalSubDirs, LocalDiscoveredFiles);
+								AddDiscovered(Data.DirLocalAbsPath, Data.DirLongPackageName, LocalSubDirs, LocalDiscoveredFiles);
 							}
 						}
 						Data.ScanDir->SetScanInFlight(false);
@@ -2433,7 +2438,7 @@ void FAssetDataDiscovery::OnDirectoryCreated(FStringView LocalAbsPath)
 	// The directory may also be scanned in the future because a parent directory is still yet pending to scan,
 	// we do not try to prevent that wasteful rescan because this is a rare event and it does not cause a behavior problem
 	SetIsIdle(false);
-	AddDiscovered(DirData.LocalAbsPath, TConstArrayView<FDiscoveredPathData>(&DirData, 1), TConstArrayView<FDiscoveredPathData>());
+	AddDiscovered(DirData.LocalAbsPath, DirData.LongPackageName, TConstArrayView<FDiscoveredPathData>(&DirData, 1), TConstArrayView<FDiscoveredPathData>());
 }
 
 void FAssetDataDiscovery::OnFilesCreated(TConstArrayView<FString> LocalAbsPaths)
@@ -2563,15 +2568,26 @@ int32 FAssetDataDiscovery::FindLowerBoundMountPoint(FStringView LocalAbsPath) co
 	);
 }
 
-void FAssetDataDiscovery::AddDiscovered(FStringView DirAbsPath, TConstArrayView<FDiscoveredPathData> SubDirs, TConstArrayView<FDiscoveredPathData> Files)
+void FAssetDataDiscovery::AddDiscovered(FStringView DirAbsPath, FStringView DirPackagePath, TConstArrayView<FDiscoveredPathData> SubDirs, TConstArrayView<FDiscoveredPathData> Files)
 {
 	// This function is inside the critical section so we have moved filtering results outside of it
 	// Caller is responsible for filtering SubDirs and Files by ShouldScan and packagename validity
 	FGathererScopeLock ResultsScopeLock(&ResultsLock);
-	for (const FDiscoveredPathData& SubDir : SubDirs)
+
+	if (UE::AssetDataGather::Private::bIgnoreEmptyDirectories)
 	{
-		DiscoveredDirectories.Add(FString(SubDir.LongPackageName));
+		// Only register this directory, this will get called for anything that has files or subdirectories
+		DiscoveredDirectories.Add(FString(DirPackagePath));
 	}
+	else 
+	{
+		// Register all of the subdirectories even if they are empty
+		for (const FDiscoveredPathData& SubDir : SubDirs)
+		{
+			DiscoveredDirectories.Add(FString(SubDir.LongPackageName));
+		}
+	}
+
 	if (Files.Num())
 	{
 		DiscoveredFiles.Emplace(DirAbsPath, Files);
diff --git a/Engine/Source/Runtime/AssetRegistry/Private/AssetDataGathererPrivate.h b/Engine/Source/Runtime/AssetRegistry/Private/AssetDataGathererPrivate.h
--- a/Engine/Source/Runtime/AssetRegistry/Private/AssetDataGathererPrivate.h
+++ b/Engine/Source/Runtime/AssetRegistry/Private/AssetDataGathererPrivate.h
@@ -658,7 +658,7 @@ private:
 	void SetIsIdle(bool bInIdle, double& TickStartTime);
 
 	/** Store the given discovered files and directories in the results. */
-	void AddDiscovered(FStringView DirAbsPath, TConstArrayView<FDiscoveredPathData> SubDirs,
+	void AddDiscovered(FStringView DirAbsPath, FStringView DirPackagePath, TConstArrayView<FDiscoveredPathData> SubDirs,
 		TConstArrayView<FDiscoveredPathData> Files);
 	/** Store the given specially reported single file in the results. */
 	void AddDiscoveredFile(FDiscoveredPathData&& File);
diff --git a/Engine/Source/Runtime/AssetRegistry/Private/AssetRegistry.cpp b/Engine/Source/Runtime/AssetRegistry/Private/AssetRegistry.cpp
--- a/Engine/Source/Runtime/AssetRegistry/Private/AssetRegistry.cpp
+++ b/Engine/Source/Runtime/AssetRegistry/Private/AssetRegistry.cpp
@@ -154,16 +154,29 @@ namespace UE::AssetRegistry::Impl
 {
 	/** The max time to spend in UAssetRegistryImpl::Tick */
 	float MaxSecondsPerFrame = 0.04f;
+	static FAutoConsoleVariableRef CVarAssetRegistryMaxSecondsPerFrame(
+		TEXT("AssetRegistry.MaxSecondsPerFrame"),
+		UE::AssetRegistry::Impl::MaxSecondsPerFrame,
+		TEXT("Maximum amount of time allowed for Asset Registry processing, in seconds"));
+
+	/** If true, defer sorting of dependencies until loading is complete */
+	bool bDeferDependencySort = false;
+	static FAutoConsoleVariableRef CVarAssetRegistryDeferDependencySort(
+		TEXT("AssetRegistry.DeferDependencySort"),
+		UE::AssetRegistry::Impl::bDeferDependencySort,
+		TEXT("If true, the dependency lists on dependency nodes will not be sorted until after the initial load is complete"));
+
+	/** If true, defer sorting of referencer data until loading is complete, this is enabled by default because of native packages with many referencers */
+	bool bDeferReferencerSort = true;
+	static FAutoConsoleVariableRef CVarAssetRegistryDeferReferencerSort(
+		TEXT("AssetRegistry.DeferReferencerSort"),
+		UE::AssetRegistry::Impl::bDeferReferencerSort,
+		TEXT("If true, the referencer list on dependency nodes will not be sorted until after the initial load is complete"));
 
 	/** Name of UObjectRedirector property */
 	const FName DestinationObjectFName(TEXT("DestinationObject"));
 }
 
-static FAutoConsoleVariableRef CVarAssetRegistryMaxSecondsPerFrame(
-	TEXT("AssetRegistry.MaxSecondsPerFrame"),
-	UE::AssetRegistry::Impl::MaxSecondsPerFrame,
-	TEXT("Maximum amount of time allowed for Asset Registry processing, in seconds"));
-
 #ifndef ENABLE_PLATFORM_CHUNK_INSTALL
 	#define ENABLE_PLATFORM_CHUNK_INSTALL (1)
 #endif
@@ -831,6 +844,8 @@ void FAssetRegistryImpl::Initialize(Impl::FInitializeContext& Context)
 	bInitialSearchStarted = false;
 	bInitialSearchCompleted = true;
 	GatherStatus = Impl::EGatherStatus::Active;
+	PerformanceMode = Impl::EPerformanceMode::MostlyStatic;
+
 	bSearchAllAssets = false;
 #if NO_LOGGING
 	bVerboseLogging = false;
@@ -941,7 +956,7 @@ void FAssetRegistryImpl::RebuildAssetDependencyGathererMapIfNeeded()
 		{
 			if (UClass* Class = Cast<UClass>(ClassObject); Class && Class->IsChildOf(AssetClass) && !Class->HasAnyClassFlags(CLASS_Abstract | CLASS_Deprecated | CLASS_NewerVersionExists))
 			{
-				RegisteredDependencyGathererClasses.Add(Class, RegisteredAssetDependencyGatherer);
+				RegisteredDependencyGathererClasses.Add(FTopLevelAssetPath(Class), RegisteredAssetDependencyGatherer);
 			}
 		}
 	});
@@ -1611,9 +1626,41 @@ void FAssetRegistryImpl::SearchAllAssetsInitialAsync(Impl::FEventContext& EventC
 	InitialSearchStartTime = FPlatformTime::Seconds();
 	bInitialSearchStarted = true;
 	bInitialSearchCompleted = false;
+	SetPerformanceMode(Impl::EPerformanceMode::BulkLoading);
 	SearchAllAssets(EventContext, InheritanceContext, false /* bSynchronousSearch */);
 }
 
+void FAssetRegistryImpl::SetPerformanceMode(Impl::EPerformanceMode NewMode)
+{
+	if (PerformanceMode != NewMode)
+	{
+		const bool bWereDependenciesSorted = ShouldSortDependencies();
+		const bool bWereReferencersSorted = ShouldSortReferencers();
+
+		PerformanceMode = NewMode;
+
+		const bool bShouldSortDependencies = ShouldSortDependencies();
+		const bool bShouldSortReferencers = ShouldSortReferencers();
+
+		if ((bWereDependenciesSorted != bShouldSortDependencies) || (bWereReferencersSorted != bShouldSortReferencers))
+		{
+			State.SetDependencyNodeSorting(bShouldSortDependencies, bShouldSortReferencers);
+		}		
+	}
+}
+
+bool FAssetRegistryImpl::ShouldSortDependencies() const
+{
+	// Always sort in static, sometimes sort during loading
+	return (PerformanceMode == Impl::MostlyStatic || (PerformanceMode == Impl::BulkLoading && !Impl::bDeferDependencySort));
+}
+
+bool FAssetRegistryImpl::ShouldSortReferencers() const
+{
+	// Always sort in static, sometimes sort during loading
+	return (PerformanceMode == Impl::MostlyStatic || (PerformanceMode == Impl::BulkLoading && !Impl::bDeferReferencerSort));
+}
+
 }
 
 void UAssetRegistryImpl::SearchAllAssets(bool bSynchronousSearch)
@@ -4157,6 +4204,10 @@ Impl::EGatherStatus FAssetRegistryImpl::TickGatherer(Impl::FEventContext& EventC
 			// update redirectors
 			UpdateRedirectCollector();
 #endif
+
+			// Handle any deferred loading operations
+			SetPerformanceMode(Impl::EPerformanceMode::MostlyStatic);
+
 			RecordTimer();
 			LogSearchDiagnostics(InitialSearchStartTime);
 			TRACE_END_REGION(TEXT("Asset Registry Scan"));
@@ -4335,7 +4386,10 @@ void FAssetRegistryImpl::LoadCalculatedDependencies(FName PackageName,
 	for (const FAssetData* AssetData : State.GetAssetsByPackageName(PackageName))
 	{
 		Gatherers.Reset();
-		RegisteredDependencyGathererClasses.MultiFind(AssetData->GetClass(), Gatherers);
+
+		// Check the class name instead of trying to load the actual class as that is slow
+		// This code could be moved somewhere where it doesn't need to re-query the asset data, but it needs to happen after both dependencies and data are handled
+		RegisteredDependencyGathererClasses.MultiFind(AssetData->AssetClassPath, Gatherers);
 		for (UE::AssetDependencyGatherer::Private::FRegisteredAssetDependencyGatherer* Gatherer : Gatherers)
 		{
 			if (!bOutHadActivity)
@@ -4980,6 +5034,15 @@ bool FAssetRegistryImpl::ShouldSkipGatheredAsset(FAssetData& AssetData)
 	// Extra validation for ExternalActors. If duplicate ExternalActors with the same object path exist
 	// then we intermittently will fail to find the correct one and WorldPartition will break.
 	// Validate that the PackageName matches what is expected from the ObjectPath.
+
+#if WITH_EDITORONLY_DATA
+	if (AssetData.GetOptionalOuterPathName().IsNone())
+	{
+		// If no outer path, this can't be an external asset
+		return false;
+	}
+#endif
+
 	FStringView ExternalActorsFolderName(FPackagePath::GetExternalActorsFolderName());
 	TStringBuilder<256> PackageNameStr;
 	AssetData.PackageName.ToString(PackageNameStr);
@@ -5231,6 +5294,8 @@ void FAssetRegistryImpl::DependencyDataGathered(const double TickStartTime, TMul
 				});
 
 			Node->ClearDependencies();
+			Node->SetIsDependencyListSorted(EDependencyCategory::All, ShouldSortDependencies());
+			Node->SetIsReferencersSorted(ShouldSortReferencers());
 
 			// Don't bother registering dependencies on these packages, every package in the game will depend on them
 			static TArray<FName> ScriptPackagesToSkip = TArray<FName>{ TEXT("/Script/CoreUObject"), TEXT("/Script/Engine"), TEXT("/Script/BlueprintGraph"), TEXT("/Script/UnrealEd") };
@@ -5271,6 +5336,9 @@ void FAssetRegistryImpl::DependencyDataGathered(const double TickStartTime, TMul
 				{
 					if (DependsNode->GetConnectionCount() == 0)
 					{
+						DependsNode->SetIsDependencyListSorted(EDependencyCategory::All, ShouldSortDependencies());
+						DependsNode->SetIsReferencersSorted(ShouldSortReferencers());
+
 						// This was newly created, see if we need to read the script package Guid
 						const FNameBuilder DependencyPackageNameStr(DependencyPackageName);
 
@@ -5725,8 +5793,17 @@ void FAssetRegistryImpl::OnDirectoryChanged(Impl::FEventContext& EventContext,
 	{
 		if (FileChangesProcessed[FileIdx].Action == FFileChangeData::FCA_RescanRequired)
 		{
-			OnDirectoryRescanRequired(EventContext, InheritanceContext, FileChangesProcessed[FileIdx].Filename,
-				FileChangesProcessed[FileIdx].TimeStamp);
+			if (bInitialSearchStarted && !bInitialSearchCompleted)
+			{
+				// Ignore rescan request during initial scan as it is probably caused by the scan itself
+				UE_LOG(LogAssetRegistry, Log, TEXT("FAssetRegistry ignoring rescan request for %s during startup"), *FileChangesProcessed[FileIdx].Filename);
+			}
+			else
+			{
+				OnDirectoryRescanRequired(EventContext, InheritanceContext, FileChangesProcessed[FileIdx].Filename,
+					FileChangesProcessed[FileIdx].TimeStamp);
+			}
+
 			continue;
 		}
 		FString LongPackageName;
@@ -6975,12 +7052,7 @@ void FAssetRegistryImpl::SetManageReferences(const TMultiMap<FAssetIdentifier, F
 		ExistingManagedNodes.Add(ManageNode);
 	}
 	// Restore all nodes to manage dependencies sorted and references sorted, so we can efficiently read them in future operations
-	for (const TPair<FAssetIdentifier, FDependsNode*>& Pair : State.CachedDependsNodes)
-	{
-		FDependsNode* DependsNode = Pair.Value;
-		DependsNode->SetIsDependencyListSorted(UE::AssetRegistry::EDependencyCategory::Manage, true);
-		DependsNode->SetIsReferencersSorted(true);
-	}
+	State.SetDependencyNodeSorting(ShouldSortDependencies(), ShouldSortReferencers());
 
 #if WITH_EDITOR
 	if (!GAssetRegistryManagementPathsPackageDebugName.IsEmpty())
diff --git a/Engine/Source/Runtime/AssetRegistry/Private/AssetRegistryImpl.h b/Engine/Source/Runtime/AssetRegistry/Private/AssetRegistryImpl.h
--- a/Engine/Source/Runtime/AssetRegistry/Private/AssetRegistryImpl.h
+++ b/Engine/Source/Runtime/AssetRegistry/Private/AssetRegistryImpl.h
@@ -100,6 +100,16 @@ namespace Impl
 		Complete,
 		UnableToProgress,
 	};
+
+	/** Affects how rules are applied to improve loading/runtime performance */
+	enum EPerformanceMode : uint8
+	{
+		// Handling slow async load
+		BulkLoading,
+
+		// Not changing, optimize for runtime queries
+		MostlyStatic,
+	};
 }
 
 /**
@@ -243,6 +253,12 @@ public:
 	uint64 GetClassGeneratorNamesRegisteredClassesVersionNumber() const { return ClassGeneratorNamesRegisteredClassesVersionNumber; }
 	/** Get a copy of the cached serialization options that were parsed from ini */
 	void CopySerializationOptions(FAssetRegistrySerializationOptions& OutOptions, ESerializationTarget Target) const;
+	
+	/** Query the performance mode, which modifies how data structures are loaded */
+	Impl::EPerformanceMode GetPerformanceMode() const { return PerformanceMode; }
+	void SetPerformanceMode(Impl::EPerformanceMode NewMode);
+	bool ShouldSortDependencies() const;
+	bool ShouldSortReferencers() const;
 
 	const FAssetRegistryState& GetState() const;
 	const FPathTree& GetCachedPathTree() const;
@@ -411,6 +427,8 @@ private:
 	bool bPreloadingComplete = false;
 	/** Status of the background search, so we can take actions when it changes to or from idle */
 	Impl::EGatherStatus GatherStatus;
+	/** What kind of performance mode this is in, used to optimize for initial loading vs runtime */
+	Impl::EPerformanceMode PerformanceMode;
 
 	/**
 	 * Enables extra check to make sure path still mounted before adding.
@@ -455,7 +473,7 @@ private:
 	TMultiMap<FString, FName> DirectoryReferencers;
 
 	/** A map of per asset class dependency gatherer called in LoadCalculatedDependencies */
-	TMultiMap<UClass*, UE::AssetDependencyGatherer::Private::FRegisteredAssetDependencyGatherer*> RegisteredDependencyGathererClasses;
+	TMultiMap<FTopLevelAssetPath, UE::AssetDependencyGatherer::Private::FRegisteredAssetDependencyGatherer*> RegisteredDependencyGathererClasses;
 	bool bRegisteredDependencyGathererClassesDirty;
 #endif
 #if WITH_ENGINE && WITH_EDITOR
diff --git a/Engine/Source/Runtime/AssetRegistry/Private/AssetRegistryState.cpp b/Engine/Source/Runtime/AssetRegistry/Private/AssetRegistryState.cpp
--- a/Engine/Source/Runtime/AssetRegistry/Private/AssetRegistryState.cpp
+++ b/Engine/Source/Runtime/AssetRegistry/Private/AssetRegistryState.cpp
@@ -317,12 +317,7 @@ void FAssetRegistryState::InitializeFromExistingAndPrune(const FAssetRegistrySta
 	}
 
 	// Restore the sortedness that we turned off for performance when creating each DependsNode
-	for (TPair<FAssetIdentifier, FDependsNode*> Pair : CachedDependsNodes)
-	{
-		FDependsNode* DependsNode = Pair.Value;
-		DependsNode->SetIsDependencyListSorted(UE::AssetRegistry::EDependencyCategory::All, true);
-		DependsNode->SetIsReferencersSorted(true);
-	}
+	SetDependencyNodeSorting(true, true);
 }
 
 void FAssetRegistryState::InitializeFromExisting(const FAssetDataMap& AssetDataMap, const TMap<FAssetIdentifier, FDependsNode*>& DependsNodeMap, 
@@ -459,12 +454,7 @@ void FAssetRegistryState::InitializeFromExisting(const FAssetDataMap& AssetDataM
 		}
 
 		// Restore the sortedness that we turned off for performance when creating each DependsNode
-		for (TPair<FAssetIdentifier, FDependsNode*> Pair : CachedDependsNodes)
-		{
-			FDependsNode* DependsNode = Pair.Value;
-			DependsNode->SetIsDependencyListSorted(UE::AssetRegistry::EDependencyCategory::All, true);
-			DependsNode->SetIsReferencersSorted(true);
-		}
+		SetDependencyNodeSorting(true, true);
 	}
 }
 
@@ -1890,6 +1880,16 @@ void FAssetRegistryState::SetTagsOnExistingAsset(FAssetData* AssetData, FAssetDa
 	AssetData->TagsAndValues = FAssetDataTagMapSharedView(MoveTemp(NewTags));
 }
 
+void FAssetRegistryState::SetDependencyNodeSorting(bool bSortDependencies, bool bSortReferencers)
+{
+	for (TPair<FAssetIdentifier, FDependsNode*>& Pair : CachedDependsNodes)
+	{
+		FDependsNode* DependsNode = Pair.Value;
+		DependsNode->SetIsDependencyListSorted(UE::AssetRegistry::EDependencyCategory::All, bSortDependencies);
+		DependsNode->SetIsReferencersSorted(bSortReferencers);
+	}
+}
+
 void FAssetRegistryState::UpdateAssetData(const FAssetData& NewAssetData, bool bCreateIfNotExists)
 {
 	if (FAssetData** AssetData = CachedAssets.Find(FCachedAssetKey(NewAssetData)))
diff --git a/Engine/Source/Runtime/AssetRegistry/Public/AssetRegistry/AssetRegistryState.h b/Engine/Source/Runtime/AssetRegistry/Public/AssetRegistry/AssetRegistryState.h
--- a/Engine/Source/Runtime/AssetRegistry/Public/AssetRegistry/AssetRegistryState.h
+++ b/Engine/Source/Runtime/AssetRegistry/Public/AssetRegistry/AssetRegistryState.h
@@ -575,27 +575,29 @@ private:
 	void Load(Archive&& Ar, const FAssetRegistryHeader& Header, const FAssetRegistryLoadOptions& Options);
 
 	/** Initialize the lookup maps */
-	ASSETREGISTRY_API void SetAssetDatas(TArrayView<FAssetData> AssetDatas, const FAssetRegistryLoadOptions& Options);
+	void SetAssetDatas(TArrayView<FAssetData> AssetDatas, const FAssetRegistryLoadOptions& Options);
 
 	/** Find the first non-redirector dependency node starting from InDependency. */
-	ASSETREGISTRY_API FDependsNode* ResolveRedirector(FDependsNode* InDependency, const FAssetDataMap& InAllowedAssets, TMap<FDependsNode*, FDependsNode*>& InCache);
+	FDependsNode* ResolveRedirector(FDependsNode* InDependency, const FAssetDataMap& InAllowedAssets, TMap<FDependsNode*, FDependsNode*>& InCache);
 
 	/** Finds an existing node for the given package and returns it, or returns null if one isn't found */
-	ASSETREGISTRY_API FDependsNode* FindDependsNode(const FAssetIdentifier& Identifier) const;
+	FDependsNode* FindDependsNode(const FAssetIdentifier& Identifier) const;
 
 	/** Creates a node in the CachedDependsNodes map or finds the existing node and returns it */
-	ASSETREGISTRY_API FDependsNode* CreateOrFindDependsNode(const FAssetIdentifier& Identifier);
+	FDependsNode* CreateOrFindDependsNode(const FAssetIdentifier& Identifier);
 
 	/** Removes the depends node and updates the dependencies to no longer contain it as as a referencer. */
-	ASSETREGISTRY_API bool RemoveDependsNode(const FAssetIdentifier& Identifier);
+	bool RemoveDependsNode(const FAssetIdentifier& Identifier);
 
 	/** Filter a set of tags and output a copy of the filtered set. */
-	static ASSETREGISTRY_API void FilterTags(const FAssetDataTagMapSharedView& InTagsAndValues, FAssetDataTagMap& OutTagsAndValues, const TSet<FName>* ClassSpecificFilterList, const FAssetRegistrySerializationOptions & Options);
+	static void FilterTags(const FAssetDataTagMapSharedView& InTagsAndValues, FAssetDataTagMap& OutTagsAndValues, const TSet<FName>* ClassSpecificFilterList, const FAssetRegistrySerializationOptions & Options);
 
-	ASSETREGISTRY_API void LoadDependencies(FArchive& Ar);
-	ASSETREGISTRY_API void LoadDependencies_BeforeFlags(FArchive& Ar, bool bSerializeDependencies, FAssetRegistryVersion::Type Version);
+	void LoadDependencies(FArchive& Ar);
+	void LoadDependencies_BeforeFlags(FArchive& Ar, bool bSerializeDependencies, FAssetRegistryVersion::Type Version);
 
-	ASSETREGISTRY_API void SetTagsOnExistingAsset(FAssetData* AssetData, FAssetDataTagMap&& NewTags);
+	void SetTagsOnExistingAsset(FAssetData* AssetData, FAssetDataTagMap&& NewTags);
+
+	void SetDependencyNodeSorting(bool bSortDependencies, bool bSortReferencers);
 
 	/** Set of asset data for assets saved to disk. Searched via path name types, implicitly converted to FCachedAssetKey. */
 	FAssetDataMap CachedAssets;
