Centralize ENABLE_PLATFORM_CHUNK_INSTALL declarations into GenericPlatformChunkInstall.h

#rnx

diff --git a/Engine/Source/Runtime/AssetRegistry/Private/AssetRegistry.cpp b/Engine/Source/Runtime/AssetRegistry/Private/AssetRegistry.cpp
--- a/Engine/Source/Runtime/AssetRegistry/Private/AssetRegistry.cpp
+++ b/Engine/Source/Runtime/AssetRegistry/Private/AssetRegistry.cpp
@@ -177,10 +177,6 @@ namespace UE::AssetRegistry::Impl
 	const FName DestinationObjectFName(TEXT("DestinationObject"));
 }
 
-#ifndef ENABLE_PLATFORM_CHUNK_INSTALL
-	#define ENABLE_PLATFORM_CHUNK_INSTALL (1)
-#endif
-
 /**
  * Implementation of IAssetRegistryInterface; forwards calls from the CoreUObject-accessible IAssetRegistryInterface into the AssetRegistry-accessible IAssetRegistry
  */
diff --git a/Engine/Source/Runtime/Core/Public/GenericPlatform/GenericPlatformChunkInstall.h b/Engine/Source/Runtime/Core/Public/GenericPlatform/GenericPlatformChunkInstall.h
--- a/Engine/Source/Runtime/Core/Public/GenericPlatform/GenericPlatformChunkInstall.h
+++ b/Engine/Source/Runtime/Core/Public/GenericPlatform/GenericPlatformChunkInstall.h
@@ -21,6 +21,10 @@ class IPlatformChunkInstall;
 
 CORE_API DECLARE_LOG_CATEGORY_EXTERN(LogChunkInstaller, Log, All);
 
+#ifndef ENABLE_PLATFORM_CHUNK_INSTALL
+	#define ENABLE_PLATFORM_CHUNK_INSTALL (1)
+#endif
+
 namespace EChunkLocation
 {
 	enum Type
diff --git a/Engine/Source/Runtime/Engine/Private/AssetManager.cpp b/Engine/Source/Runtime/Engine/Private/AssetManager.cpp
--- a/Engine/Source/Runtime/Engine/Private/AssetManager.cpp
+++ b/Engine/Source/Runtime/Engine/Private/AssetManager.cpp
@@ -52,9 +52,6 @@ LLM_DEFINE_TAG(AssetManager);
 
 DEFINE_LOG_CATEGORY(LogAssetManager);
 
-#ifndef ENABLE_PLATFORM_CHUNK_INSTALL
-	#define ENABLE_PLATFORM_CHUNK_INSTALL (1)
-#endif
 
 /** Structure defining the current loading state of an asset */
 struct FPrimaryAssetLoadState
diff --git a/Engine/Source/Runtime/PakFile/Private/IPlatformFilePak.cpp b/Engine/Source/Runtime/PakFile/Private/IPlatformFilePak.cpp
--- a/Engine/Source/Runtime/PakFile/Private/IPlatformFilePak.cpp
+++ b/Engine/Source/Runtime/PakFile/Private/IPlatformFilePak.cpp
@@ -72,9 +72,6 @@ int64 GTotalLoadedLastTick = 0;
 #ifndef HAS_PLATFORM_PAK_INSTALL_CHECK
 #define HAS_PLATFORM_PAK_INSTALL_CHECK 0
 #endif
-#ifndef ENABLE_PLATFORM_CHUNK_INSTALL
-#define ENABLE_PLATFORM_CHUNK_INSTALL 1
-#endif
 #ifndef ALL_PAKS_WILDCARD
 #define ALL_PAKS_WILDCARD "*.pak"
 #endif 
