Re-enable the default behavior of adding empty directories to the asset registry as it confused users of the content browser.
This can still be disabled at the project level using the AssetRegistry.IgnoreEmptyDirectories cvar.

diff --git a/Engine/Source/Runtime/AssetRegistry/Private/AssetDataGatherer.cpp b/Engine/Source/Runtime/AssetRegistry/Private/AssetDataGatherer.cpp
--- a/Engine/Source/Runtime/AssetRegistry/Private/AssetDataGatherer.cpp
+++ b/Engine/Source/Runtime/AssetRegistry/Private/AssetDataGatherer.cpp
@@ -71,7 +71,7 @@ static FAutoConsoleVariableRef CVarBlockPackagesWithMarkOfTheWeb(
 	bBlockPackagesWithMarkOfTheWeb,
 	TEXT("Whether package files with mark of the web are blocked from the asset registry"));
 
-bool bIgnoreEmptyDirectories = true;
+bool bIgnoreEmptyDirectories = false;
 static FAutoConsoleVariableRef CVarIgnoreEmptyDirectories(
 	TEXT("AssetRegistry.IgnoreEmptyDirectories"),
 	bIgnoreEmptyDirectories,
