Fixes for static analysis false positives.

diff --git a/Engine/Source/Runtime/AssetRegistry/Private/PackageReader.cpp b/Engine/Source/Runtime/AssetRegistry/Private/PackageReader.cpp
--- a/Engine/Source/Runtime/AssetRegistry/Private/PackageReader.cpp
+++ b/Engine/Source/Runtime/AssetRegistry/Private/PackageReader.cpp
@@ -66,7 +66,7 @@ bool FPackageReader::OpenPackageFile(FStringView InLongPackageName, FStringView
 	Loader = IFileManager::Get().CreateFileReader(*PackageFilename);
 	bLoaderOwner = true;
 	EOpenPackageResult Tmp;
-	return OpenPackageFile(OutErrorCode ? *OutErrorCode : Tmp);
+	return OpenPackageFile(*(OutErrorCode ? OutErrorCode : &Tmp)); // do the ternary on pointers rather than references, because static analysis complains
 }
 
 bool FPackageReader::OpenPackageFile(FArchive* InLoader, EOpenPackageResult* OutErrorCode)
@@ -77,7 +77,7 @@ bool FPackageReader::OpenPackageFile(FArchive* InLoader, EOpenPackageResult* Out
 	LongPackageName.Empty();
 	PackageFilename = Loader->GetArchiveName();
 	EOpenPackageResult Tmp;
-	return OpenPackageFile(OutErrorCode ? *OutErrorCode : Tmp);
+	return OpenPackageFile(*(OutErrorCode ? OutErrorCode : &Tmp)); // do the ternary on pointers rather than references, because static analysis complains
 }
 
 bool FPackageReader::OpenPackageFile(TUniquePtr<FArchive> InLoader, EOpenPackageResult* OutErrorCode)
@@ -88,7 +88,7 @@ bool FPackageReader::OpenPackageFile(TUniquePtr<FArchive> InLoader, EOpenPackage
 	LongPackageName.Empty();
 	PackageFilename = Loader->GetArchiveName();
 	EOpenPackageResult Tmp;
-	return OpenPackageFile(OutErrorCode ? *OutErrorCode : Tmp);
+	return OpenPackageFile(*(OutErrorCode ? OutErrorCode : &Tmp)); // do the ternary on pointers rather than references, because static analysis complains
 }
 
 bool FPackageReader::OpenPackageFile(EOpenPackageResult& OutErrorCode)
diff --git a/Engine/Source/Runtime/Core/Private/Containers/String.cpp b/Engine/Source/Runtime/Core/Private/Containers/String.cpp
--- a/Engine/Source/Runtime/Core/Private/Containers/String.cpp
+++ b/Engine/Source/Runtime/Core/Private/Containers/String.cpp
@@ -1622,7 +1622,7 @@ void FString::ConvertTabsToSpacesInline(const int32 InSpacesPerTab)
 	//must call this with at least 1 space so the modulus operation works
 	check(InSpacesPerTab > 0);
 
-	int32 TabIndex;
+	int32 TabIndex = 0;
 	while ((TabIndex = Find(CHARTEXT(ElementType, "\t"), ESearchCase::CaseSensitive)) != INDEX_NONE )
 	{
 		FString RightSide = Mid(TabIndex+1);
