Initialize local variables

diff --git a/Engine/Source/Runtime/AssetRegistry/Private/PackageReader.cpp b/Engine/Source/Runtime/AssetRegistry/Private/PackageReader.cpp
--- a/Engine/Source/Runtime/AssetRegistry/Private/PackageReader.cpp
+++ b/Engine/Source/Runtime/AssetRegistry/Private/PackageReader.cpp
@@ -65,8 +65,8 @@ bool FPackageReader::OpenPackageFile(FStringView InLongPackageName, FStringView
 	PackageFilename = InPackageFilename;
 	Loader = IFileManager::Get().CreateFileReader(*PackageFilename);
 	bLoaderOwner = true;
-	EOpenPackageResult Tmp;
-	return OpenPackageFile(*(OutErrorCode ? OutErrorCode : &Tmp)); // do the ternary on pointers rather than references, because static analysis complains
+	EOpenPackageResult Tmp = EOpenPackageResult::Success;
+	return OpenPackageFile(OutErrorCode ? *OutErrorCode : Tmp); 
 }
 
 bool FPackageReader::OpenPackageFile(FArchive* InLoader, EOpenPackageResult* OutErrorCode)
@@ -76,8 +76,8 @@ bool FPackageReader::OpenPackageFile(FArchive* InLoader, EOpenPackageResult* Out
 	bLoaderOwner = false;
 	LongPackageName.Empty();
 	PackageFilename = Loader->GetArchiveName();
-	EOpenPackageResult Tmp;
-	return OpenPackageFile(*(OutErrorCode ? OutErrorCode : &Tmp)); // do the ternary on pointers rather than references, because static analysis complains
+	EOpenPackageResult Tmp = EOpenPackageResult::Success;
+	return OpenPackageFile(OutErrorCode ? *OutErrorCode : Tmp); 
 }
 
 bool FPackageReader::OpenPackageFile(TUniquePtr<FArchive> InLoader, EOpenPackageResult* OutErrorCode)
@@ -87,8 +87,8 @@ bool FPackageReader::OpenPackageFile(TUniquePtr<FArchive> InLoader, EOpenPackage
 	bLoaderOwner = true;
 	LongPackageName.Empty();
 	PackageFilename = Loader->GetArchiveName();
-	EOpenPackageResult Tmp;
-	return OpenPackageFile(*(OutErrorCode ? OutErrorCode : &Tmp)); // do the ternary on pointers rather than references, because static analysis complains
+	EOpenPackageResult Tmp = EOpenPackageResult::Success;
+	return OpenPackageFile(OutErrorCode ? *OutErrorCode : Tmp); 
 }
 
 bool FPackageReader::OpenPackageFile(EOpenPackageResult& OutErrorCode)
