AssetRegistry: Optimize AssetRegistryState::EnumerateAssets in general and specifically in the case when we have a classfilter with a large number of results combined with a packagefilter with a small number of results.
#rnx

diff --git a/Engine/Source/Runtime/AssetRegistry/Private/AssetRegistryState.cpp b/Engine/Source/Runtime/AssetRegistry/Private/AssetRegistryState.cpp
--- a/Engine/Source/Runtime/AssetRegistry/Private/AssetRegistryState.cpp
+++ b/Engine/Source/Runtime/AssetRegistry/Private/AssetRegistryState.cpp
@@ -4,6 +4,7 @@
 
 #include "Algo/Compare.h"
 #include "Algo/Sort.h"
+#include "Algo/Unique.h"
 #include "AssetRegistry/ARFilter.h"
 #include "AssetRegistryArchive.h"
 #include "AssetRegistryImpl.h"
@@ -655,125 +656,266 @@ bool FAssetRegistryState::GetAssets(const FARCompiledFilter& Filter, const TSet<
 	bSkipARFilteredAssets);
 }
 
-template<class ArrayType, typename KeyType>
-TArray<FAssetData*> FindAssets(const TMap<KeyType, ArrayType>& Map, const TSet<KeyType>& Keys)
+namespace UE::AssetRegistry::Private
 {
-	TArray<TArrayView<FAssetData* const>> Matches;
+
+bool DecideIntersectionMethod(int32 PreviousSize, int32 FilterResultsSize, int32 FilterComplexity)
+{
+	// Cost of intersection of previous results with new results is the cost to construct a TMap of the smaller
+	// set plus the cost to query larger set against the TMap. TMap construction is more expensive than TMap query.
+	constexpr uint64 TMapConstructionCost = 3;
+	uint64 SmallSize;
+	uint64 LargeSize;
+	if (PreviousSize < FilterResultsSize)
+	{
+		SmallSize = (uint64)PreviousSize;
+		LargeSize = (uint64)FilterResultsSize;
+	}
+	else
+	{
+		SmallSize = (uint64)FilterResultsSize;
+		LargeSize = (uint64)PreviousSize;
+	}
+	uint64 ArrayCost = SmallSize * TMapConstructionCost + LargeSize;
+	// Cost of filtering previous results by FilterFunction is filtercomplexity times size of the previous results
+	uint64 FilterCost = ((uint64)FilterComplexity) * ((uint64)PreviousSize);
+
+	// Our two sets of cost calculation are not on the same scale; they are off by some factor that is dependent 
+	// upon ArrayIntersection code, TMap code, and hardware dependent factors. But we assume they are for on the
+	// same scale for simplicity. Despite the invalid assumption, our comparison will still work in the important
+	// cases: a large FilterComplexity will use ArrayIntersection and a large FilterResultsSize will use filtering.
+	return FilterCost < ArrayCost;
+}
+
+void ArrayIntersection(TArray<const FAssetData*>& InOutResults, TConstArrayView<TArrayView<const FAssetData*>> Matches, int32 TotalMatches)
+{
+	if (InOutResults.Num() < TotalMatches)
+	{
+		TMap<const FAssetData*, bool> Exists;
+		Exists.Reserve(InOutResults.Num());
+		for (const FAssetData* Result : InOutResults)
+		{
+			Exists.Add(Result, false);
+		}
+		InOutResults.Empty();
+		for (TArrayView<const FAssetData*> Assets : Matches)
+		{
+			for (const FAssetData* Asset : Assets)
+			{
+				bool* Result = Exists.Find(Asset);
+				if (Result && 
+					!(*Result) // If there are duplicates of an Asset in multiple elements of Matches, only add the first one
+				)
+				{
+					*Result = true;
+					InOutResults.Add(Asset);
+				}
+ 			}
+		}
+	}
+	else
+	{
+		TSet<const FAssetData*> Exists;
+		Exists.Reserve(TotalMatches);
+		for (TArrayView<const FAssetData*> Assets : Matches)
+		{
+			for (const FAssetData* Asset : Assets)
+			{
+				Exists.Add(Asset);
+			}
+		}
+
+		InOutResults.RemoveAllSwap([&Exists](const FAssetData* Asset)
+			{
+				return !Exists.Contains(Asset);
+			});
+	}
+}
+
+template<class ArrayType, typename KeyType, typename CallbackType>
+void FilterAssets(TArray<const FAssetData*>& InOutResults, const TMap<KeyType, ArrayType>& AccelerationMap,
+	const TSet<KeyType>& Keys, CallbackType&& FunctionToKeepAsset, int32 FilterComplexity)
+{
+	TArray<TArrayView<const FAssetData*>, TInlineAllocator<10>> Matches;
 	Matches.Reserve(Keys.Num());
 	uint32 TotalMatches = 0;
 
 	for (const KeyType& Key : Keys)
 	{
-		if (const ArrayType* Assets = Map.Find(Key))
+		if (const ArrayType* Assets = AccelerationMap.Find(Key))
 		{
-			Matches.Add(MakeArrayView(*Assets));
+			const FAssetData** AssetPtr = const_cast<const FAssetData**>(Assets->GetData());
+			Matches.Add(TArrayView<const FAssetData*>(AssetPtr, Assets->Num()));
 			TotalMatches += Assets->Num();
 		}
 	}
 
-	TArray<FAssetData*> Out;
-	Out.Reserve(TotalMatches);
-	for (TArrayView<FAssetData* const> Assets : Matches)
+	// Keys is a TSet and entries in the AccelerationMap do not overlap, so there should be no duplicates to remove in Matches
+	if (InOutResults.IsEmpty())
+	{
+		// No previous Results; set Results equal to the values found in AccelerationMap
+		InOutResults.Reserve(TotalMatches);
+		for (TArrayView<const FAssetData*> Assets : Matches)
+		{
+			InOutResults.Append(Assets);
+		}
+	}
+	else
 	{
-		Out.Append(Assets.GetData(), Assets.Num());
+		bool bUseFiltering = DecideIntersectionMethod(InOutResults.Num(), TotalMatches, FilterComplexity);
+		if (bUseFiltering)
+		{
+			InOutResults.RemoveAllSwap([&FunctionToKeepAsset](const FAssetData* AssetData)
+				{
+					return !FunctionToKeepAsset(AssetData);
+				});
+		}
+		else
+		{
+			ArrayIntersection(InOutResults, TConstArrayView<TArrayView<const FAssetData*>>(Matches), TotalMatches);
+		}
 	}
-
-	return Out;
 }
 
-bool FAssetRegistryState::EnumerateAssets(const FARCompiledFilter& Filter, const TSet<FName>& PackageNamesToSkip,
-	TFunctionRef<bool(const FAssetData&)> Callback, bool bSkipARFilteredAssets) const
+template<typename CallbackType>
+void FilterAssets(TArray<const FAssetData*>&InOutResults,
+	const UE::AssetRegistry::Private::FAssetDataMap &AccelerationMap, const TSet<FSoftObjectPath>& Keys,
+	CallbackType&& FunctionToKeepAsset, int32 FilterComplexity)
 {
-	// Verify filter input. If all assets are needed, use EnumerateAllAssets() instead.
-	if (Filter.IsEmpty() || !IsFilterValid(Filter))
+	TArray<const FAssetData*, TInlineAllocator<10>> Matches;
+	Matches.Reserve(Keys.Num());
+
+	for (const FSoftObjectPath& Key : Keys)
 	{
-		return false;
+		if (FAssetData* const* AssetDataPtr = AccelerationMap.Find(UE::AssetRegistry::Private::FCachedAssetKey(Key)))
+		{
+			Matches.Add(*AssetDataPtr);
+		}
 	}
 
-	const uint32 FilterWithoutPackageFlags = Filter.WithoutPackageFlags;
-	const uint32 FilterWithPackageFlags = Filter.WithPackageFlags;
-
-	// The assets that match each filter
-	TArray<TArray<FAssetData*>, TInlineAllocator<5>> FilterResults;
-	
-	// On disk package names
-	if (Filter.PackageNames.Num() > 0)
+	// Keys is a TSet, so there should be no duplicates to remove in Matches
+	if (InOutResults.IsEmpty())
 	{
-		FilterResults.Emplace(FindAssets(CachedAssetsByPackageName, Filter.PackageNames));
+		// No previous Results; set Results equal to the values found in AccelerationMap
+		InOutResults = MoveTemp(Matches);
 	}
-
-	// On disk package paths
-	if (Filter.PackagePaths.Num() > 0)
+	else
 	{
-		FilterResults.Emplace(FindAssets(CachedAssetsByPath, Filter.PackagePaths));
+		bool bUseFiltering = DecideIntersectionMethod(InOutResults.Num(), Matches.Num(), FilterComplexity);
+		if (bUseFiltering)
+		{
+			InOutResults.RemoveAllSwap([&FunctionToKeepAsset](const FAssetData* AssetData)
+				{
+					return !FunctionToKeepAsset(AssetData);
+				});
+		}
+		else
+		{
+			TArrayView<const FAssetData*> ArrayView(Matches);
+			TConstArrayView<TArrayView<const FAssetData*>> ArrayViewOfArrayViews(&ArrayView, 1);
+			ArrayIntersection(InOutResults, ArrayViewOfArrayViews, Matches.Num());
+		}
 	}
+}
 
-	// On disk classes
-	if (Filter.ClassPaths.Num() > 0)
+bool AssetDataMatchesTag(const FAssetData* AssetData, const TPair<FName, TOptional<FString>>& TagPair)
+{
+	if (!AssetData)
 	{
-		FilterResults.Emplace(FindAssets(CachedAssetsByClass, Filter.ClassPaths));
+		return false;
 	}
-
-	// On disk object paths
-	if (Filter.SoftObjectPaths.Num() > 0)
+	if (!TagPair.Value.IsSet())
 	{
-		TArray<FAssetData*>& ObjectPathsFilter = FilterResults.Emplace_GetRef();
-		ObjectPathsFilter.Reserve(Filter.SoftObjectPaths.Num());
+		return AssetData->TagsAndValues.Contains(TagPair.Key);
+	}
+	else
+	{
+		return AssetData->TagsAndValues.ContainsKeyValue(TagPair.Key, TagPair.Value.GetValue());
+	}
+}
 
-		for (const FSoftObjectPath& ObjectPath : Filter.SoftObjectPaths)
+template<typename CallbackType>
+void FilterAssets(TArray<const FAssetData*>& InOutResults, const TMap<FName, TArray<FAssetData*>>& AccelerationMap,
+	const TMultiMap<FName, TOptional<FString>>& TagsAndValues, CallbackType&& FunctionToKeepAsset, int32 FilterComplexity)
+{
+	TArray<TArray<const FAssetData*>, TInlineAllocator<10>> Matches;
+	Matches.Reserve(TagsAndValues.Num());
+	uint32 TotalMatches = 0;
+
+	for (const TPair<FName, TOptional<FString>>& TagPair : TagsAndValues)
+	{
+		TArray<const FAssetData*>& Results = Matches.Emplace_GetRef();
+		if (const TArray<FAssetData*>* TagAssets = AccelerationMap.Find(TagPair.Key))
 		{
-			if (FAssetData* const* AssetDataPtr = CachedAssets.Find(FCachedAssetKey(ObjectPath)))
+			Results.Reserve(TagAssets->Num());
+			for (FAssetData* AssetData : *TagAssets)
 			{
-				ObjectPathsFilter.Add(*AssetDataPtr);
+				if (AssetDataMatchesTag(AssetData, TagPair))
+				{
+					Results.Add(AssetData);
+				}
 			}
 		}
+		TotalMatches += Results.Num();
 	}
 
-	// On disk tags and values
-	if (Filter.TagsAndValues.Num() > 0)
+	if (InOutResults.IsEmpty())
 	{
-		TArray<FAssetData*>& TagAndValuesFilter = FilterResults.Emplace_GetRef();
-		// Sometimes number of assets matching this filter is correlated to number of assets matching previous filters 
-		if (FilterResults.Num())
+		// No previous Results; set Results equal to the values found in AccelerationMap
+		InOutResults.Reserve(TotalMatches);
+		for (TArray<const FAssetData*>& Assets : Matches)
 		{
-			TagAndValuesFilter.Reserve(FilterResults[0].Num());
+			InOutResults.Append(Assets);
 		}
-
-		for (auto FilterTagIt = Filter.TagsAndValues.CreateConstIterator(); FilterTagIt; ++FilterTagIt)
+		// Remove duplicates
+		Algo::Sort(InOutResults);
+		InOutResults.SetNum(Algo::Unique(InOutResults));
+	}
+	else
+	{
+		bool bUseFiltering = DecideIntersectionMethod(InOutResults.Num(), TotalMatches, FilterComplexity);
+		if (bUseFiltering)
 		{
-			const FName Tag = FilterTagIt.Key();
-			const TOptional<FString>& Value = FilterTagIt.Value();
-
-			if (const TArray<FAssetData*>* TagAssets = CachedAssetsByTag.Find(Tag))
-			{
-				for (FAssetData* AssetData : *TagAssets)
+			InOutResults.RemoveAllSwap([&FunctionToKeepAsset](const FAssetData* AssetData)
 				{
-					if (AssetData != nullptr)
-					{
-						bool bAccept;
-						if (!Value.IsSet())
-						{
-							bAccept = AssetData->TagsAndValues.Contains(Tag);
-						}
-						else
-						{
-							bAccept = AssetData->TagsAndValues.ContainsKeyValue(Tag, Value.GetValue());
-						}
-						if (bAccept)
-						{
-							TagAndValuesFilter.Add(AssetData);
-						}
-					}
-				}
+					return !FunctionToKeepAsset(AssetData);
+				});
+		}
+		else
+		{
+			// Convert Array of Arrays into format required by ArrayIntersection: Array of ArrayViews
+			TArray<TArrayView<const FAssetData*>, TInlineAllocator<10>> ArrayViewMatches;
+			ArrayViewMatches.Reserve(Matches.Num());
+			for (TArray<const FAssetData*>& MatchesElement : Matches)
+			{
+				ArrayViewMatches.Emplace(MatchesElement);
 			}
+
+			// ArrayIntersection handles removing any duplicates from Matches
+			ArrayIntersection(InOutResults, ArrayViewMatches, TotalMatches);
 		}
 	}
+}
 
-	// Perform callback for assets that match all filters
-	if (FilterResults.Num() > 0)
+}
+
+bool FAssetRegistryState::EnumerateAssets(const FARCompiledFilter& Filter, const TSet<FName>& PackageNamesToSkip,
+	TFunctionRef<bool(const FAssetData&)> Callback, bool bSkipARFilteredAssets) const
+{
+	using namespace UE::AssetRegistry::Private;
+
+	// Verify filter input. If all assets are needed, use EnumerateAllAssets() instead.
+	if (Filter.IsEmpty() || !IsFilterValid(Filter))
 	{
-		auto SkipAssetData = [&](const FAssetData* AssetData) 
-		{ 
+		return false;
+	}
+
+	const uint32 FilterWithoutPackageFlags = Filter.WithoutPackageFlags;
+	const uint32 FilterWithPackageFlags = Filter.WithPackageFlags;
+	auto ShouldSkipAssetData =
+		[&PackageNamesToSkip, bSkipARFilteredAssets, FilterWithoutPackageFlags, FilterWithPackageFlags]
+		(const FAssetData* AssetData)
+		{
 			if (PackageNamesToSkip.Contains(AssetData->PackageName) |			//-V792
 				AssetData->HasAnyPackageFlags(FilterWithoutPackageFlags) |		//-V792
 				!AssetData->HasAllPackageFlags(FilterWithPackageFlags))			//-V792
@@ -785,51 +927,104 @@ bool FAssetRegistryState::EnumerateAssets(const FARCompiledFilter& Filter, const
 				UE::AssetRegistry::FFiltering::ShouldSkipAsset(AssetData->AssetClassPath, AssetData->PackageFlags);
 		};
 
-		int32 NumFilterResults = FilterResults.Num();
-		if (NumFilterResults > 1)
-		{
-			// Mark which filters each asset passes
-			uint32 PassAllFiltersValue = (1 << NumFilterResults) - 1; // 1 in every bit for the lowest n bits
-			TMap<FAssetData*, uint32> PassBits;
-			for (int32 FilterIndex = 0; FilterIndex < NumFilterResults; ++FilterIndex)
+
+	// Some of our filters are accelerated: we have TMaps that list for each value of the filter all of the assets
+	// that pass that filter. But some of those assets-passing-FilterN-ValueV are very large, and just merging the
+	// lists of FAssetData* can be expensive. So for each new filter we need to decide whether it is more expensive to
+	// merge previous results with the acceleration list or to apply the filter to every element in previous results.
+	// This decision is handled by FilterAssets.
+	// To benefit from the filter method we want to have as small a list of results as possible at each step, so
+	// order the filters from most-likely to have few results to least-likely to have few results.
+	TArray<const FAssetData*> AccumulatedResults;
+
+	if (Filter.SoftObjectPaths.Num() > 0)
+	{
+		FilterAssets(AccumulatedResults, CachedAssets, Filter.SoftObjectPaths,
+			[&Filter](const FAssetData* AssetData)
 			{
-				const TArray<FAssetData*>& FilterEvaluation = FilterResults[FilterIndex];
-				PassBits.Reserve(FilterEvaluation.Num());
-				
-				for (FAssetData* AssetData : FilterEvaluation)
-				{
-					PassBits.FindOrAdd(AssetData) |= (1 << FilterIndex);
-				}
-			}
+				return Filter.SoftObjectPaths.Contains(AssetData->GetSoftObjectPath());
+			},
+			Filter.SoftObjectPaths.Num());
+		if (AccumulatedResults.IsEmpty())
+		{
+			return true;
+		}
+	}
 
-			// Include assets that pass all filters
-			for (TPair<FAssetData*, uint32> PassPair : PassBits)
+	if (Filter.PackageNames.Num() > 0)
+	{
+		FilterAssets(AccumulatedResults, CachedAssetsByPackageName, Filter.PackageNames,
+			[&Filter](const FAssetData* AssetData)
 			{
-				const FAssetData* AssetData = PassPair.Key;
-				if (PassPair.Value != PassAllFiltersValue || SkipAssetData(AssetData))
-				{
-					continue;
-				}
-				else if (!Callback(*AssetData))
-				{
-					return true;
-				}
-			}
+				return Filter.PackageNames.Contains(AssetData->PackageName);
+			},
+			Filter.PackageNames.Num());
+		if (AccumulatedResults.IsEmpty())
+		{
+			return true;
 		}
-		else
+	}
+
+	if (Filter.PackagePaths.Num() > 0)
+	{
+		FilterAssets(AccumulatedResults, CachedAssetsByPath, Filter.PackagePaths,
+			[&Filter](const FAssetData* AssetData)
+			{
+				return Filter.PackagePaths.Contains(AssetData->PackagePath);
+			},
+			Filter.PackagePaths.Num());
+		if (AccumulatedResults.IsEmpty())
 		{
-			// All matched assets passed the single filter
-			for (const FAssetData* AssetData : FilterResults[0])
+			return true;
+		}
+	}
+
+	if (Filter.TagsAndValues.Num() > 0)
+	{
+		FilterAssets(AccumulatedResults, CachedAssetsByTag, Filter.TagsAndValues,
+			[&Filter](const FAssetData* AssetData)
 			{
-				if (SkipAssetData(AssetData))
-				{
-					continue;
-				}
-				else if (!Callback(*AssetData))
+				for (const TPair<FName, TOptional<FString>>& TagPair : Filter.TagsAndValues)
 				{
-					return true;
+					if (AssetDataMatchesTag(AssetData, TagPair))
+					{
+						return true; // keep
+					}
 				}
-			}
+				return false; // remove
+			},
+			Filter.TagsAndValues.Num());
+		if (AccumulatedResults.IsEmpty())
+		{
+			return true;
+		}
+	}
+
+	if (Filter.ClassPaths.Num() > 0)
+	{
+		FilterAssets(AccumulatedResults, CachedAssetsByClass, Filter.ClassPaths,
+			[&Filter](const FAssetData* AssetData)
+			{
+				return Filter.ClassPaths.Contains(AssetData->AssetClassPath);
+			},
+			Filter.ClassPaths.Num());
+		if (AccumulatedResults.IsEmpty())
+		{
+			return true;
+		}
+	}
+
+	// Run the remaining non-accelerated filters on every element of AccumulatedResults
+	for (const FAssetData* AssetData : AccumulatedResults)
+	{
+		if (ShouldSkipAssetData(AssetData))
+		{
+			continue;
+		}
+		bool bContinueFiltering = Callback(*AssetData);
+		if (!bContinueFiltering)
+		{
+			return true;
 		}
 	}
 
