Allow UAssetRegistryImpl::AppendState to be called from background threads.

diff --git a/Engine/Source/Runtime/AssetRegistry/Private/AssetRegistry.cpp b/Engine/Source/Runtime/AssetRegistry/Private/AssetRegistry.cpp
--- a/Engine/Source/Runtime/AssetRegistry/Private/AssetRegistry.cpp
+++ b/Engine/Source/Runtime/AssetRegistry/Private/AssetRegistry.cpp
@@ -4499,17 +4499,14 @@ void UAssetRegistryImpl::AppendState(const FAssetRegistryState& InState)
 		GuardedData.AppendState(EventContext, InState);
 	}
 
-	Broadcast(EventContext);
-	checkf(IsInGameThread(), TEXT("AppendState is not yet implemented as callable from other threads"));
-
+	// AppendState does not create the AssetAdded events; create them here and then broadcast
 	TArray<FAssetData> TempAssets;
 	InState.GetAllAssets(TSet<FName>(), TempAssets, true /* bARFiltering */);
-	for (const FAssetData& AssetData : TempAssets)
+	for (FAssetData& AssetData : TempAssets)
 	{
-		// Let subscribers know that the new asset was added to the registry
-		AssetAddedEvent.Broadcast(AssetData);
+		EventContext.AssetEvents.Emplace(MoveTemp(AssetData), UE::AssetRegistry::Impl::FEventContext::EEvent::Added);
 	}
-	OnAssetsAdded().Broadcast(TempAssets);
+	Broadcast(EventContext);
 }
 
 namespace UE::AssetRegistry
