Build path tree from leaf up to avoid re-checking existence of parent nodes as frequently.
Require callers to strip trailing slashes from input.
Add batch delegates for path add/remove.

diff --git a/Engine/Plugins/Editor/ContentBrowser/ContentBrowserAssetDataSource/Source/ContentBrowserAssetDataSource/Private/ContentBrowserAssetDataSource.cpp b/Engine/Plugins/Editor/ContentBrowser/ContentBrowserAssetDataSource/Source/ContentBrowserAssetDataSource/Private/ContentBrowserAssetDataSource.cpp
--- a/Engine/Plugins/Editor/ContentBrowser/ContentBrowserAssetDataSource/Source/ContentBrowserAssetDataSource/Private/ContentBrowserAssetDataSource.cpp
+++ b/Engine/Plugins/Editor/ContentBrowser/ContentBrowserAssetDataSource/Source/ContentBrowserAssetDataSource/Private/ContentBrowserAssetDataSource.cpp
@@ -64,8 +64,8 @@ void UContentBrowserAssetDataSource::Initialize(const bool InAutoRegister)
 	AssetRegistry->OnAssetRenamed().AddUObject(this, &UContentBrowserAssetDataSource::OnAssetRenamed);
 	AssetRegistry->OnAssetUpdated().AddUObject(this, &UContentBrowserAssetDataSource::OnAssetUpdated);
 	AssetRegistry->OnAssetUpdatedOnDisk().AddUObject(this, &UContentBrowserAssetDataSource::OnAssetUpdatedOnDisk);
-	AssetRegistry->OnPathAdded().AddUObject(this, &UContentBrowserAssetDataSource::OnPathAdded);
-	AssetRegistry->OnPathRemoved().AddUObject(this, &UContentBrowserAssetDataSource::OnPathRemoved);
+	AssetRegistry->OnPathsAdded().AddUObject(this, &UContentBrowserAssetDataSource::OnPathsAdded);
+	AssetRegistry->OnPathsRemoved().AddUObject(this, &UContentBrowserAssetDataSource::OnPathsRemoved);
 
 	// Listen for when assets are loaded or changed
 	FCoreUObjectDelegates::OnObjectPropertyChanged.AddUObject(this, &UContentBrowserAssetDataSource::OnObjectPropertyChanged);
@@ -188,8 +188,8 @@ void UContentBrowserAssetDataSource::Shutdown()
 			AssetRegistryMaybe->OnAssetRenamed().RemoveAll(this);
 			AssetRegistryMaybe->OnAssetUpdated().RemoveAll(this);
 			AssetRegistryMaybe->OnAssetUpdatedOnDisk().RemoveAll(this);	
-			AssetRegistryMaybe->OnPathAdded().RemoveAll(this);
-			AssetRegistryMaybe->OnPathRemoved().RemoveAll(this);
+			AssetRegistryMaybe->OnPathsAdded().RemoveAll(this);
+			AssetRegistryMaybe->OnPathsRemoved().RemoveAll(this);
 			AssetRegistryMaybe->OnFilesLoaded().RemoveAll(this);
 		}
 	}
@@ -2857,42 +2857,48 @@ void UContentBrowserAssetDataSource::OnObjectPreSave(UObject* InObject, FObjectP
 	}
 }
 
-void UContentBrowserAssetDataSource::OnPathAdded(const FString& InPath)
+void UContentBrowserAssetDataSource::OnPathsAdded(TConstArrayView<FStringView> Paths)
 {
-	// Completely ignore paths that do not pass the most inclusive filter
-	if (!ContentBrowserDataUtils::PathPassesAttributeFilter(InPath, 0, EContentBrowserItemAttributeFilter::IncludeAll))
+	for (FStringView InPath : Paths)
 	{
-		return;
-	}
+		// Completely ignore paths that do not pass the most inclusive filter
+		if (!ContentBrowserDataUtils::PathPassesAttributeFilter(InPath, 0, EContentBrowserItemAttributeFilter::IncludeAll))
+		{
+			return;
+		}
 
-	FName PathName(InPath);
-	RecentlyPopulatedAssetFolders.Empty();
-	
-	QueueItemDataUpdate(FContentBrowserItemDataUpdate::MakeItemAddedUpdate(CreateAssetFolderItem(PathName)));
+		FName PathName(InPath);
+		RecentlyPopulatedAssetFolders.Empty();
+		
+		QueueItemDataUpdate(FContentBrowserItemDataUpdate::MakeItemAddedUpdate(CreateAssetFolderItem(PathName)));
 
-	FStringView PathView(InPath);
-	// Minus one because the test depth start at zero
-	const int32 CurrentDepth = ContentBrowserDataUtils::CalculateFolderDepthOfPath(PathView) - 1;
-	int32 Index;
-	if (PathView.FindLastChar(TEXT('/'), Index))
-	{ 
-		uint32 PathNameHash = GetTypeHash(PathName);
-		FName ParentPath(PathView.Left(Index));
-		uint32 ParentPathHash = GetTypeHash(ParentPath);
-		OnAssetPathAddedDelegate.Broadcast(PathName, PathView, PathNameHash, ParentPath, ParentPathHash, CurrentDepth);
+		FStringView PathView(InPath);
+		// Minus one because the test depth start at zero
+		const int32 CurrentDepth = ContentBrowserDataUtils::CalculateFolderDepthOfPath(PathView) - 1;
+		int32 Index;
+		if (PathView.FindLastChar(TEXT('/'), Index))
+		{ 
+			uint32 PathNameHash = GetTypeHash(PathName);
+			FName ParentPath(PathView.Left(Index));
+			uint32 ParentPathHash = GetTypeHash(ParentPath);
+			OnAssetPathAddedDelegate.Broadcast(PathName, PathView, PathNameHash, ParentPath, ParentPathHash, CurrentDepth);
+		}
 	}
 }
 
-void UContentBrowserAssetDataSource::OnPathRemoved(const FString& InPath)
+void UContentBrowserAssetDataSource::OnPathsRemoved(TConstArrayView<FStringView> Paths)
 {
-	// Deleted paths are no longer relevant for tracking
-	FName PathName(InPath);
-	RecentlyPopulatedAssetFolders.Remove(PathName);
-	AssetFolderToAttributes.Remove(PathName);
+	for (FStringView InPath : Paths)
+	{
+		// Deleted paths are no longer relevant for tracking
+		FName PathName(InPath);
+		RecentlyPopulatedAssetFolders.Remove(PathName);
+		AssetFolderToAttributes.Remove(PathName);
 
-	QueueItemDataUpdate(FContentBrowserItemDataUpdate::MakeItemRemovedUpdate(CreateAssetFolderItem(PathName)));
+		QueueItemDataUpdate(FContentBrowserItemDataUpdate::MakeItemRemovedUpdate(CreateAssetFolderItem(PathName)));
 
-	OnAssetPathRemovedDelegate.Broadcast(PathName, GetTypeHash(PathName));
+		OnAssetPathRemovedDelegate.Broadcast(PathName, GetTypeHash(PathName));
+	}
 }
 
 void UContentBrowserAssetDataSource::OnPathPopulated(const FAssetData& InAssetData)
diff --git a/Engine/Plugins/Editor/ContentBrowser/ContentBrowserAssetDataSource/Source/ContentBrowserAssetDataSource/Public/ContentBrowserAssetDataSource.h b/Engine/Plugins/Editor/ContentBrowser/ContentBrowserAssetDataSource/Source/ContentBrowserAssetDataSource/Public/ContentBrowserAssetDataSource.h
--- a/Engine/Plugins/Editor/ContentBrowser/ContentBrowserAssetDataSource/Source/ContentBrowserAssetDataSource/Public/ContentBrowserAssetDataSource.h
+++ b/Engine/Plugins/Editor/ContentBrowser/ContentBrowserAssetDataSource/Source/ContentBrowserAssetDataSource/Public/ContentBrowserAssetDataSource.h
@@ -417,9 +417,9 @@ private:
 
 	void OnObjectPreSave(UObject* InObject, class FObjectPreSaveContext InObjectPreSaveContext);
 
-	void OnPathAdded(const FString& InPath);
+	void OnPathsAdded(TConstArrayView<FStringView> Paths);
 
-	void OnPathRemoved(const FString& InPath);
+	void OnPathsRemoved(TConstArrayView<FStringView> Paths);
 
 	void OnPathPopulated(const FAssetData& InAssetData);
 
diff --git a/Engine/Source/Programs/LowLevelTests/FoundationTests.Build.cs b/Engine/Source/Programs/LowLevelTests/FoundationTests.Build.cs
--- a/Engine/Source/Programs/LowLevelTests/FoundationTests.Build.cs
+++ b/Engine/Source/Programs/LowLevelTests/FoundationTests.Build.cs
@@ -29,7 +29,8 @@ public class FoundationTests : TestModuleRules
 				"Core",
 				"Cbor",
 				"CoreUObject",
-				"TelemetryUtils"
+				"TelemetryUtils",
+				"AssetRegistry"
 			});
 
 		if (Target.bBuildWithEditorOnlyData)
diff --git a/Engine/Source/Runtime/AssetRegistry/Private/AssetRegistry.cpp b/Engine/Source/Runtime/AssetRegistry/Private/AssetRegistry.cpp
--- a/Engine/Source/Runtime/AssetRegistry/Private/AssetRegistry.cpp
+++ b/Engine/Source/Runtime/AssetRegistry/Private/AssetRegistry.cpp
@@ -35,6 +35,7 @@
 #include "Serialization/ArrayReader.h"
 #include "Serialization/CompactBinarySerialization.h"
 #include "Serialization/CompactBinaryWriter.h"
+#include "String/RemoveFrom.h"
 #include "Templates/UnrealTemplate.h"
 #include "TelemetryRouter.h"
 #include "UObject/ConstructorHelpers.h"
@@ -916,7 +917,7 @@ void FAssetRegistryImpl::Initialize(Impl::FInitializeContext& Context)
 	Context.bUpdateDiskCacheAfterLoad = bUpdateDiskCacheAfterLoad;
 
 	// Content roots always exist; add them as paths
-	FPackageName::QueryRootContentPaths(Context.RootContentPaths);
+	FPackageName::QueryRootContentPaths(Context.RootContentPaths, false, false, true);
 	for (const FString& AssetPath : Context.RootContentPaths)
 	{
 		AddPath(Context.Events, AssetPath);
@@ -3544,7 +3545,7 @@ bool UAssetRegistryImpl::AddPath(const FString& PathToAdd)
 namespace UE::AssetRegistry
 {
 
-bool FAssetRegistryImpl::AddPath(Impl::FEventContext& EventContext, const FString& PathToAdd)
+bool FAssetRegistryImpl::AddPath(Impl::FEventContext& EventContext, FStringView PathToAdd)
 {
 	bool bIsDenied = false;
 	// If no GlobalGatherer, then we are in the game or non-cook commandlet and we do not implement deny listing
@@ -3560,7 +3561,7 @@ bool FAssetRegistryImpl::AddPath(Impl::FEventContext& EventContext, const FStrin
 	{
 		return false;
 	}
-	return AddAssetPath(EventContext, FName(*PathToAdd));
+	return AddAssetPath(EventContext, FName(PathToAdd));
 }
 
 }
@@ -5226,6 +5227,8 @@ void FAssetRegistryImpl::PathDataGathered(Impl::FEventContext& EventContext, con
 		FPackageName::QueryRootContentPaths(MountPointsArray, /*bIncludeReadOnlyRoots=*/ true, /*bWithoutLeadingSlashes*/ false, /*WithoutTrailingSlashes=*/ true);
 		MountPoints.Append(MoveTemp(MountPointsArray));
 	}
+	
+	CachedPathTree.EnsureAdditionalCapacity(PathResults.Num());
 
 	while (PathResults.Num() > 0)
 	{
@@ -6371,7 +6374,7 @@ void FAssetRegistryImpl::OnContentPathMounted(Impl::FEventContext& EventContext,
 	const FString& AssetPathWithTrailingSlash, const FString& FileSystemPath)
 {
 	// Content roots always exist
-	AddPath(EventContext, AssetPathWithTrailingSlash);
+	AddPath(EventContext, UE::String::RemoveFromEnd(FStringView(AssetPathWithTrailingSlash), TEXTVIEW("/")));
 
 	if (GlobalGatherer.IsValid() && bSearchAllAssets)
 	{
@@ -7293,17 +7296,42 @@ void UAssetRegistryImpl::Broadcast(UE::AssetRegistry::Impl::FEventContext& Event
 
 	if (EventContext.PathEvents.Num())
 	{
+		// Batch add/remove events 
+		TArray<FStringView> Params;
+		// Ensure loop batch condition is always false first iteration
+		bool bCurrentBatchIsAdd = EventContext.PathEvents[0].Get<1>() == FEventContext::EEvent::Added;
 		for (const TPair<FString, FEventContext::EEvent>& PathEvent : EventContext.PathEvents)
 		{
 			const FString& Path = PathEvent.Get<0>();
-			switch (PathEvent.Get<1>())
+			bool bEventIsAdd = PathEvent.Get<1>() == FEventContext::EEvent::Added;
+			if (bEventIsAdd != bCurrentBatchIsAdd)
 			{
-			case FEventContext::EEvent::Added:
-				PathAddedEvent.Broadcast(Path);
-				break;
-			case FEventContext::EEvent::Removed:
-				PathRemovedEvent.Broadcast(Path);
-				break;
+				(bCurrentBatchIsAdd ? PathsAddedEvent : PathsRemovedEvent).Broadcast(MakeArrayView(Params));
+				Params.Reset();
+				bCurrentBatchIsAdd = bEventIsAdd;
+			}
+			Params.Add(FStringView(Path));
+		}
+		if (Params.Num() != 0)
+		{
+			(bCurrentBatchIsAdd ? PathsAddedEvent : PathsRemovedEvent).Broadcast(MakeArrayView(Params));
+		}
+
+		// Legacy single events 
+		if (PathAddedEvent.IsBound() || PathRemovedEvent.IsBound())
+		{
+			for (const TPair<FString, FEventContext::EEvent>& PathEvent : EventContext.PathEvents)
+			{
+				const FString& Path = PathEvent.Get<0>();
+				switch (PathEvent.Get<1>())
+				{
+				case FEventContext::EEvent::Added:
+					PathAddedEvent.Broadcast(Path);
+					break;
+				case FEventContext::EEvent::Removed:
+					PathRemovedEvent.Broadcast(Path);
+					break;
+				}
 			}
 		}
 		EventContext.PathEvents.Empty();
@@ -7415,6 +7443,16 @@ UAssetRegistryImpl::FFilesBlockedEvent& UAssetRegistryImpl::OnFilesBlocked()
 	return FilesBlockedEvent;
 }
 
+UAssetRegistryImpl::FPathsEvent& UAssetRegistryImpl::OnPathsAdded()
+{
+	return PathsAddedEvent;
+}
+
+UAssetRegistryImpl::FPathsEvent& UAssetRegistryImpl::OnPathsRemoved()
+{
+	return PathsRemovedEvent;
+}
+
 UAssetRegistryImpl::FPathAddedEvent& UAssetRegistryImpl::OnPathAdded()
 {
 	return PathAddedEvent;
diff --git a/Engine/Source/Runtime/AssetRegistry/Private/AssetRegistry.h b/Engine/Source/Runtime/AssetRegistry/Private/AssetRegistry.h
--- a/Engine/Source/Runtime/AssetRegistry/Private/AssetRegistry.h
+++ b/Engine/Source/Runtime/AssetRegistry/Private/AssetRegistry.h
@@ -127,6 +127,10 @@ public:
 	DECLARE_DERIVED_EVENT( UAssetRegistryImpl, IAssetRegistry::FFilesBlockedEvent, FFilesBlockedEvent);
 	virtual FFilesBlockedEvent& OnFilesBlocked() override;
 
+	DECLARE_DERIVED_EVENT( UAssetRegistryImpl, IAssetRegistry::FPathsEvent, FPathsEvent);
+	virtual FPathsEvent& OnPathsAdded() override;
+	virtual FPathsEvent& OnPathsRemoved() override;
+	
 	DECLARE_DERIVED_EVENT( UAssetRegistryImpl, IAssetRegistry::FPathAddedEvent, FPathAddedEvent);
 	virtual FPathAddedEvent& OnPathAdded() override;
 
@@ -295,6 +299,12 @@ private:
 	/** The delegate to execute when one or more files have been blocked from the registry */
 	FFilesBlockedEvent FilesBlockedEvent;
 
+	/** The delegate to execute when a batch of paths are added to the registry */
+	FPathsEvent PathsAddedEvent;
+
+	/** The delegate to execute when a batch of paths are removed from the registry */
+	FPathsEvent PathsRemovedEvent;
+	
 	/** The delegate to execute when an asset path is added to the registry */
 	FPathAddedEvent PathAddedEvent;
 
diff --git a/Engine/Source/Runtime/AssetRegistry/Private/AssetRegistryImpl.h b/Engine/Source/Runtime/AssetRegistry/Private/AssetRegistryImpl.h
--- a/Engine/Source/Runtime/AssetRegistry/Private/AssetRegistryImpl.h
+++ b/Engine/Source/Runtime/AssetRegistry/Private/AssetRegistryImpl.h
@@ -142,7 +142,7 @@ public:
 		FARCompiledFilter& OutCompiledFilter) const;
 	void SetTemporaryCachingMode(bool bEnable);
 	void SetTemporaryCachingModeInvalidated();
-	bool AddPath(Impl::FEventContext& EventContext, const FString& PathToAdd);
+	bool AddPath(Impl::FEventContext& EventContext, FStringView PathToAdd);
 	void SearchAllAssets(Impl::FEventContext& EventContext, Impl::FClassInheritanceContext& InheritanceContext,
 		bool bSynchronousSearch);
 	bool GetVerseFilesByPath(FName PackagePath, TArray<FName>* OutFilePaths, bool bRecursive) const;
diff --git a/Engine/Source/Runtime/AssetRegistry/Private/PathTree.cpp b/Engine/Source/Runtime/AssetRegistry/Private/PathTree.cpp
--- a/Engine/Source/Runtime/AssetRegistry/Private/PathTree.cpp
+++ b/Engine/Source/Runtime/AssetRegistry/Private/PathTree.cpp
@@ -2,96 +2,81 @@
 
 #include "AssetRegistry/PathTree.h"
 
+#include "Misc/ReverseIterate.h"
 #include "UObject/NameTypes.h"
 
-bool FPathTree::CachePath(FName Path, TFunctionRef<void(FName)> OnPathAdded)
+FPathTree::FPathTree()
 {
-	if (Path.IsNone())
+	// Ensure an entry for the root of the path
+	static const FName PathRoot = "/";
+
+	ParentPathToChildPaths.FindOrAdd(PathRoot, {});
+}
+
+void FPathTree::EnsureAdditionalCapacity(int32 NumNewPaths)
+{
+	ParentPathToChildPaths.Reserve(ParentPathToChildPaths.Num() + NumNewPaths); 
+	ChildPathToParentPath.Reserve(ChildPathToParentPath.Num() + NumNewPaths);
+}
+
+bool FPathTree::CachePath(FName InPath, TFunctionRef<void(FName)> OnPathAdded)
+{
+	if (InPath.IsNone())
 	{
 		return false;
 	}
 
-	if (ParentPathToChildPaths.Contains(Path))
+	if (ParentPathToChildPaths.Contains(InPath))
 	{
 		// Already cached - nothing more to do
 		return false;
 	}
 
-	FString PathStr = Path.ToString();
-	check(PathStr.Len() >= 2);	// Must be at least "/A"
-	check(PathStr[0] == '/');	// Must start with a "/"
-
-	// Paths are cached without their trailing slash, so if the given path has a trailing slash, test it again now as it may already be cached
-	if (PathStr[PathStr.Len() - 1] == '/')
-	{
-		PathStr.RemoveAt(PathStr.Len() - 1, 1, /*bAllowShrinking*/false);
-		Path = *PathStr;
+	TStringBuilder<FName::StringBufferSize> PathBuffer(InPlace, InPath);
+	FStringView PathView = PathBuffer.ToView();
+	check(PathView.Len() >= 2);	// Must be at least "/A"
+	check(PathView [0] == '/');	// Must start with a "/"
+	check(PathView[PathView.Len() - 1] != '/'); // Must not contain trailing slash 
 
-		if (ParentPathToChildPaths.Contains(Path))
-		{
-			// Already cached - nothing more to do
-			return false;
-		}
-	}
+	static const FName Root("/");
 
-	FName LastPath;
+	TArray<FName, TInlineAllocator<16>> NewPaths;
+	NewPaths.Add(InPath);
+	ParentPathToChildPaths.FindOrAdd(InPath, {}); // Add this new path with no children
 
-	// Ensure an entry for the root of the path
+	// Walk backwards through the string until we encounter a path we've already created
+	FName LastPath = InPath;
+	FName ParentPath;
+	while (ParentPath != Root)
 	{
-		static const FName PathRoot = "/";
+		// Strip the last path element from PathView to get the parent of LastPath
+		// i.e. /Game/Maps/Something -> /Game/Maps 
+		int32 SlashIndex = UE::String::FindLastChar(PathView, '/');	
+		PathView.LeftInline(SlashIndex);
+		ParentPath = PathView.IsEmpty() ? Root : FName(PathView);
+		
+		ChildPathToParentPath.FindOrAdd(LastPath, ParentPath);
 
-		if (!ParentPathToChildPaths.Contains(PathRoot))
+		uint32 Hash = GetTypeHash(ParentPath);
+		TSet<FName>* Children = ParentPathToChildPaths.FindByHash(Hash, ParentPath);
+		if (Children)
 		{
-			ParentPathToChildPaths.Add(PathRoot);
+			// Parent path already existed in tree, no need to continue looking at parents
+			Children->Add(LastPath);
+			break; 
 		}
-
-		LastPath = PathRoot;
-	}
-
-	// Walk each part of the path, adding known path entries if required
-	// This manipulates PathStr in-place to avoid making any string copies
-	TCHAR* PathCharPtr = &PathStr[1]; // Skip over the first / when scanning
-	for (;;)
-	{
-		const TCHAR PathChar = *PathCharPtr;
-		if (PathChar == '/' || PathChar == 0)
-		{
-			// We've found a path separator (or the end of the string), so process this part of the path
-			(*PathCharPtr) = 0;			// Null terminate this part of the string so we can create an FName from it
-			const FName CurrentPath = *PathStr;
-			(*PathCharPtr) = PathChar;	// Restore the original character now
-
-			check(!CurrentPath.IsNone());	// Path parts cannot be empty
-			check(*(PathCharPtr-1) != '/'); // The previous character cannot be a /, as that would suggest a malformed path such as "/Game//MyAsset"
-
-			if (!ParentPathToChildPaths.Contains(CurrentPath))
-			{
-				ParentPathToChildPaths.Add(CurrentPath);
-
-				// Only update the parent maps when this child is first added
-				if (!LastPath.IsNone())
-				{
-					// Add us as a known child of our parent path
-					TSet<FName>& ChildPaths = ParentPathToChildPaths.FindChecked(LastPath);
-					ChildPaths.Add(CurrentPath);
-
-					// Make sure we know how to find our parent again later on
-					ChildPathToParentPath.Add(CurrentPath, LastPath);
-				}
-
-				OnPathAdded(CurrentPath);
-			}
-
-			LastPath = CurrentPath;
-		}
-
-		if (PathChar == 0)
+		else
 		{
-			// End of the string
-			break;
+			ParentPathToChildPaths.AddByHash(Hash, ParentPath).Add(LastPath);
+			NewPaths.Add(ParentPath);
+			LastPath = ParentPath;
 		}
-
-		++PathCharPtr;
+	} 
+	
+	// Notify caller of each path created in order from root to leaf
+	for (FName NewPath : ReverseIterate(NewPaths))
+	{
+		OnPathAdded(NewPath);
 	}
 
 	return true;
@@ -104,27 +89,12 @@ bool FPathTree::RemovePath(FName Path, TFunctionRef<void(FName)> OnPathRemoved)
 		return false;
 	}
 
-	if (!ParentPathToChildPaths.Contains(Path))
+	TSet<FName>* Children = ParentPathToChildPaths.Find(Path);
+	if (!Children)
 	{
-		// Paths are cached without their trailing slash, so if the given path has a trailing slash, test it again now as it may already be cached
-		// We do this after the initial map test as: a) Most paths are well formed, b) This avoids string allocations until we know we need them
-		FString PathStr = Path.ToString();
-		if (PathStr[PathStr.Len() - 1] == '/')
-		{
-			PathStr.RemoveAt(PathStr.Len() - 1, 1, /*bAllowShrinking*/false);
-			Path = *PathStr;
-
-			if (!ParentPathToChildPaths.Contains(Path))
-			{
-				// Doesn't exist - nothing more to do
-				return false;
-			}
-		}
-		else
-		{
-			// Doesn't exist - nothing more to do
-			return false;
-		}
+		TStringBuilder<FName::StringBufferSize> PathString(InPlace, Path);
+		checkf(!PathString.ToView().EndsWith(TEXT("/")), TEXT("Path tree arguments should not have trailing slashes: %s"), *PathString);
+		return false;
 	}
 
 	// We also need to gather up and remove any children of this path
@@ -150,8 +120,8 @@ bool FPathTree::RemovePath(FName Path, TFunctionRef<void(FName)> OnPathRemoved)
 		const FName* ParentPathPtr = ChildPathToParentPath.Find(Path);
 		if (ParentPathPtr)
 		{
-			TSet<FName>& ChildPaths = ParentPathToChildPaths.FindChecked(*ParentPathPtr);
-			ChildPaths.Remove(Path);
+			TSet<FName>* ChildPaths = ParentPathToChildPaths.Find(*ParentPathPtr);
+			ChildPaths->Remove(Path);
 		}
 	}
 
@@ -259,7 +229,7 @@ bool FPathTree::EnumerateSubPaths(FName BasePath, TFunctionRef<bool(FName)> Call
 
 	for (const FName& ChildPath : *ChildPathsPtr)
 	{
-		check(ParentPathToChildPaths.Contains(ChildPath)); // This failing is an integrity violation as this entry lists a child that we don't know about
+		checkf(ParentPathToChildPaths.Find(ChildPath) != nullptr, TEXT("PathTree integrity failure, expected to contain %s"), *WriteToString<FName::StringBufferSize>(ChildPath)); // This failing is an integrity violation as this entry lists a child that we don't know about
 
 		if (!Callback(ChildPath))
 		{
diff --git a/Engine/Source/Runtime/AssetRegistry/Public/AssetRegistry/IAssetRegistry.h b/Engine/Source/Runtime/AssetRegistry/Public/AssetRegistry/IAssetRegistry.h
--- a/Engine/Source/Runtime/AssetRegistry/Public/AssetRegistry/IAssetRegistry.h
+++ b/Engine/Source/Runtime/AssetRegistry/Public/AssetRegistry/IAssetRegistry.h
@@ -642,12 +642,17 @@ public:
 	DECLARE_EVENT_OneParam( IAssetRegistry, FFilesBlockedEvent, const TArray<FString>& /*Files*/ );
 	virtual FFilesBlockedEvent& OnFilesBlocked() = 0;
 
-	/** Event for when paths are added to the registry */
+	/**
+	 * Events for when paths (folders/directories) are added to and removed from the asset registry.
+	 * Both batch and singular events will be called for the same set of paths in all cases.
+	 * Prefer the batched versions as the singular versions will be deprecated in a future release.
+	 */
+	DECLARE_TS_MULTICAST_DELEGATE_OneParam( FPathsEvent, TConstArrayView<FStringView> /* Paths */);
 	DECLARE_TS_MULTICAST_DELEGATE_OneParam( FPathAddedEvent, const FString& /*Path*/ );
-	virtual FPathAddedEvent& OnPathAdded() = 0;
-
-	/** Event for when paths are removed from the registry */
 	DECLARE_EVENT_OneParam( IAssetRegistry, FPathRemovedEvent, const FString& /*Path*/ );
+	virtual FPathsEvent& OnPathsAdded() = 0;
+	virtual FPathsEvent& OnPathsRemoved() = 0;
+	virtual FPathAddedEvent& OnPathAdded() = 0;
 	virtual FPathRemovedEvent& OnPathRemoved() = 0;
 
 	/** Informs the asset registry that an in-memory asset has been created */
diff --git a/Engine/Source/Runtime/AssetRegistry/Public/AssetRegistry/PathTree.h b/Engine/Source/Runtime/AssetRegistry/Public/AssetRegistry/PathTree.h
--- a/Engine/Source/Runtime/AssetRegistry/Public/AssetRegistry/PathTree.h
+++ b/Engine/Source/Runtime/AssetRegistry/Public/AssetRegistry/PathTree.h
@@ -16,6 +16,12 @@ template <typename FuncType> class TFunctionRef;
 class FPathTree
 {
 public:
+	ASSETREGISTRY_API FPathTree();
+	
+	/** Reserve extra space for the expected number of paths that will be added in future. 
+	 * Count is added to the current number of stored paths and not any previous reservation. */
+	ASSETREGISTRY_API void EnsureAdditionalCapacity(int32 NumNewPaths);
+	
 	/** Adds the specified path to the tree, creating nodes as needed and calling OnPathAdded for any new paths added. Returns true if the specified path was actually added (as opposed to already existed) */
 	ASSETREGISTRY_API bool CachePath(FName Path, TFunctionRef<void(FName)> OnPathAdded);
 
diff --git a/Engine/Source/Runtime/AssetRegistry/Tests/PathTreeTests.cpp b/Engine/Source/Runtime/AssetRegistry/Tests/PathTreeTests.cpp
new file mode 100644
--- /dev/null
+++ b/Engine/Source/Runtime/AssetRegistry/Tests/PathTreeTests.cpp
@@ -0,0 +1,76 @@
+// Copyright Epic Games, Inc. All Rights Reserved.
+
+#if WITH_LOW_LEVEL_TESTS
+
+#include "TestHarness.h"
+#include "TestMacros/Assertions.h"
+#include "AssetRegistry/PathTree.h"
+
+TEST_CASE("UE::AssetRegistry::PathTree")
+{
+    FPathTree Tree;
+    TArray<FName> NewPaths;
+    auto Callback = [&NewPaths](FName NewPath) { NewPaths.Add(NewPath); };
+    Tree.CachePath(FName("/Game/Maps/Arena"), Callback);
+    Tree.CachePath(FName("/Game/Maps/Town"), Callback);
+    Tree.CachePath(FName("/Game/Characters/Knight"), Callback);
+    
+    SECTION("NewPathCallbacks")
+    {    
+        // No duplicates should be added from successive operations
+        CHECK(NewPaths.Num() == (3 + 1 + 2));
+        CHECK(NewPaths.Contains(FName("/Game")));
+        CHECK(NewPaths.Contains(FName("/Game/Maps")));
+        CHECK(NewPaths.Contains(FName("/Game/Maps/Arena")));
+
+        CHECK(NewPaths.Contains(FName("/Game/Maps/Town")));
+        
+        CHECK(NewPaths.Contains(FName("/Game/Characters")));
+        CHECK(NewPaths.Contains(FName("/Game/Characters/Knight")));
+    }
+    
+    SECTION("AllPaths")
+    {
+        TSet<FName> AllPaths;
+        Tree.GetAllPaths(AllPaths);
+
+        CHECK(AllPaths.Num() == 7);        
+        CHECK(AllPaths.Contains(FName("/")));
+        CHECK(AllPaths.Contains(FName("/Game")));
+        CHECK(AllPaths.Contains(FName("/Game/Maps")));
+        CHECK(AllPaths.Contains(FName("/Game/Maps/Arena")));
+        CHECK(AllPaths.Contains(FName("/Game/Maps/Town")));
+        CHECK(AllPaths.Contains(FName("/Game/Characters")));
+        CHECK(AllPaths.Contains(FName("/Game/Characters/Knight")));
+    }
+    
+    SECTION("SubPaths")
+    {
+        TSet<FName> MapPaths;
+        Tree.GetSubPaths(FName("/Game/Maps"), MapPaths);
+
+        CHECK(MapPaths.Num() == 2);
+        CHECK(MapPaths.Contains(FName("/Game/Maps/Arena")));
+        CHECK(MapPaths.Contains(FName("/Game/Maps/Town")));
+    }
+
+    SECTION("Remove")
+    {
+        TSet<FName> RemovedPaths;
+        Tree.RemovePath(FName("/Game/Maps"), [&RemovedPaths](FName Path) { RemovedPaths.Add(Path); });
+        
+        CHECK(RemovedPaths.Num() == 3);
+        CHECK(RemovedPaths.Contains("/Game/Maps"));
+        CHECK(RemovedPaths.Contains("/Game/Maps/Arena"));
+        CHECK(RemovedPaths.Contains("/Game/Maps/Town"));
+        
+        TSet<FName> GamePaths;
+        Tree.GetSubPaths(FName("/Game"), GamePaths);
+
+        CHECK(GamePaths.Num() == 2);
+        CHECK(GamePaths.Contains(FName("/Game/Characters")));
+        CHECK(GamePaths.Contains(FName("/Game/Characters/Knight")));
+    }
+}
+
+#endif
\ No newline at end of file
diff --git a/Engine/Source/Runtime/CoreUObject/Private/Misc/PackageName.cpp b/Engine/Source/Runtime/CoreUObject/Private/Misc/PackageName.cpp
--- a/Engine/Source/Runtime/CoreUObject/Private/Misc/PackageName.cpp
+++ b/Engine/Source/Runtime/CoreUObject/Private/Misc/PackageName.cpp
@@ -892,6 +892,11 @@ FString FPackageName::FilenameToLongPackageName(const FString& InFilename)
 }
 
 bool FPackageName::TryConvertLongPackageNameToFilename(const FString& InLongPackageName, FString& OutFilename, const FString& InExtension)
+{
+	return TryConvertLongPackageNameToFilename(FStringView(InLongPackageName), OutFilename, InExtension);
+}
+
+bool FPackageName::TryConvertLongPackageNameToFilename(FStringView InLongPackageName, FString& OutFilename, FStringView InExtension)
 {
 	const auto& Paths = FLongPackagePathsSingleton::Get();
 	FReadScopeLock ScopeLock(Paths.MountLock);
diff --git a/Engine/Source/Runtime/CoreUObject/Public/Misc/PackageName.h b/Engine/Source/Runtime/CoreUObject/Public/Misc/PackageName.h
--- a/Engine/Source/Runtime/CoreUObject/Public/Misc/PackageName.h
+++ b/Engine/Source/Runtime/CoreUObject/Public/Misc/PackageName.h
@@ -122,6 +122,7 @@ public:
 	 * @return Package filename.
 	 */
 	static COREUOBJECT_API bool TryConvertLongPackageNameToFilename(const FString& InLongPackageName, FString& OutFilename, const FString& InExtension = TEXT(""));
+	static COREUOBJECT_API bool TryConvertLongPackageNameToFilename(FStringView InLongPackageName, FString& OutFilename, FStringView InExtension = {});
 
 	/** 
 	 * Find the MountPoint for a LocalPath, LongPackageName, or ObjectPath and return its elements. Use this function instead of TryConvertFilenameToLongPackageName or
@@ -149,6 +150,7 @@ public:
 
 	/** 
 	 * Returns the path to the specified package, excluding the short package name
+	 * e.g. given /Game/Maps/MyMap returns /Game/Maps
 	 *
 	 * @param InLongPackageName Long Package Name.
 	 * @return The path containing the specified package.
