AssetRegistry: When ScanPathsSynchronous is called on AssetDataGatherer but is not forced and nothing new is found, set idle back to true if there is no further work, so that we do not waste time on the next call to WaitForCompletion.
#rnx

diff --git a/Engine/Source/Runtime/AssetRegistry/Private/AssetDataGatherer.cpp b/Engine/Source/Runtime/AssetRegistry/Private/AssetDataGatherer.cpp
--- a/Engine/Source/Runtime/AssetRegistry/Private/AssetDataGatherer.cpp
+++ b/Engine/Source/Runtime/AssetRegistry/Private/AssetDataGatherer.cpp
@@ -3198,7 +3198,7 @@ FAssetDataGatherer::FAssetDataGatherer(const TArray<FString>& InLongPackageNames
 	, bIsIdle(false)
 	, bFirstTickAfterIdle(true)
 	, bFinishedInitialDiscovery(false)
-	, WaitBatchCount(0)
+	, WaitBatchCount(-1)
 	, LastMonolithicCacheSaveUncachedAssetFiles(0)
 	, CacheInUseCount(0)
 	, bIsSavingAsyncCache(false)
@@ -3539,35 +3539,57 @@ void FAssetDataGatherer::TickInternal(bool& bOutIsTickInterrupt, double& TickSta
 
 		// Take a batch off of the work list. If we're waiting only on the first WaitBatchCount results don't take more than that
 		int32 NumToProcess = FMath::Min<int32>(BatchSize-LocalFilesToSearch.Num(), FilesToSearch->GetNumAvailable());
-		if (WaitBatchCount > 0)
+		// If no work is available mark idle and exit
+		if (NumToProcess == 0)
 		{
+			if (WaitBatchCount != -1)
+			{
+				WaitBatchCount = -1; // WaitBatchCount was set higher than FilesToSearch->GetNumAvailable(), mark it completed
+				bOutIsTickInterrupt = true;
+			}
+
+			if (bDiscoveryIsComplete)
+			{
+				bOutIsTickInterrupt = true;
+				const bool bWasInitialDiscoveryFinished = bFinishedInitialDiscovery;
+				SetIsIdle(true, TickStartTime);
+				if (!bWasInitialDiscoveryFinished && bFinishedInitialDiscovery)
+				{
+					UE_LOG(LogAssetRegistry, Display, TEXT("Triggering cache save on discovery complete"));
+					bSaveAsyncCacheTriggered = true;
+				}
+			}
+			return;
+		}
+
+		if (WaitBatchCount >= 0)
+		{
+			if (WaitBatchCount == 0)
+			{
+				// We've finished executing the caller's requested batchcount (and we have restored idle if we
+				// are idle), so exit now without doing any further work.
+				WaitBatchCount = -1;
+				bOutIsTickInterrupt = true;
+				return;
+			}
+
+			// Otherwise we still have some work to do for the caller's requested batchcount, so do work up to
+			// that batchcount.
 			bWaitBatchCountDecremented = true;
 			NumToProcess = FMath::Min(NumToProcess, WaitBatchCount);
 			WaitBatchCount -= NumToProcess;
 			if (WaitBatchCount == 0)
 			{
+				// Mark that WaitBatchCount has been consumed and is no longer active
+				WaitBatchCount = -1;
 				bOutIsTickInterrupt = true;
 			}
 		}
-
 		DependencyResults.Reserve(FilesToSearch->GetNumAvailable() + DependencyResults.Num());
+		check(NumToProcess > 0);
 		FilesToSearch->PopFront(LocalFilesToSearch, NumToProcess);
+		check(LocalFilesToSearch.Num() > 0);
 
-		// If no work is available mark idle and exit
-		if (LocalFilesToSearch.Num() == 0 && bDiscoveryIsComplete)
-		{
-			WaitBatchCount = 0; // Clear WaitBatchCount in case it was set higher than FilesToSearch->GetNumAvailable().
-			bOutIsTickInterrupt = true;
-
-			const bool bWasInitialDiscoveryFinished = bFinishedInitialDiscovery;
-			SetIsIdle(true, TickStartTime);
-			if (!bWasInitialDiscoveryFinished && bFinishedInitialDiscovery) 
-			{
-				UE_LOG(LogAssetRegistry, Display, TEXT("Triggering cache save on discovery complete"));
-				bSaveAsyncCacheTriggered = true;
-			}
-			return;
-		}
 		if (bReadMonolithicCache && !bHasLoadedMonolithicCache)
 		{
 			bLoadMonolithicCache = true;
@@ -3802,6 +3824,9 @@ void FAssetDataGatherer::TickInternal(bool& bOutIsTickInterrupt, double& TickSta
 					FilesToSearch->AddFileAgainAfterTimeout(MoveTemp(ReadContext.AssetFileData));
 					if (bWaitBatchCountDecremented)
 					{
+						// Restore WaitBatchCount to 0 if we marked it done by setting it to -1
+						WaitBatchCount = FMath::Max(0, WaitBatchCount);
+						// Then add back on the work that we thought we were going to do but didn't get to
 						++WaitBatchCount;
 					}
 				}
@@ -4105,10 +4130,8 @@ void FAssetDataGatherer::WaitOnPathsInternal(TArrayView<UE::AssetDataGather::Pri
 
 			int32 NumDiscoveredPaths;
 			SortPathsByPriority(QueryPaths, UE::AssetDataGather::Private::EPriority::Blocking, NumDiscoveredPaths);
-			if (NumDiscoveredPaths == 0)
-			{
-				return;
-			}
+			// Set WaitBatchCount to valid (non-negative) but possibly to 0. If it is 0 we still want to call TickInternal
+			// to SetIdle(false) if necessary.
 			WaitBatchCount = NumDiscoveredPaths;
 		}
 	}
@@ -4122,7 +4145,7 @@ void FAssetDataGatherer::WaitOnPathsInternal(TArrayView<UE::AssetDataGather::Pri
 	{
 		InnerTickLoop(true /* bInSynchronousTick */, bContributeToCacheSave);
 		FGathererScopeLock ResultsScopeLock(&ResultsLock); // WaitBatchCount requires the lock
-		if (WaitBatchCount == 0)
+		if (WaitBatchCount < 0)
 		{
 			break;
 		}
diff --git a/Engine/Source/Runtime/AssetRegistry/Private/AssetDataGatherer.h b/Engine/Source/Runtime/AssetRegistry/Private/AssetDataGatherer.h
--- a/Engine/Source/Runtime/AssetRegistry/Private/AssetDataGatherer.h
+++ b/Engine/Source/Runtime/AssetRegistry/Private/AssetDataGatherer.h
@@ -476,7 +476,7 @@ private:
 	TMap<FName, FDiskCachedAssetData*> DiskCachedAssetDataMap;
 	/** Map of PackageName to cached discovered assets that will be written to disk at shutdown. */
 	TMap<FName, FDiskCachedAssetData*> NewCachedAssetDataMap;
-	/** Used to block on gather results. If non-zero, tick should end when WaitBatchCount files have been processed. */
+	/** Used to block on gather results. If non-negative, tick should end when WaitBatchCount files have been processed. */
 	int32 WaitBatchCount;
 	/** How many uncached asset files had been discovered at the last async cache save */
 	int32 LastMonolithicCacheSaveUncachedAssetFiles;
