Avoid the FAssetDataDiscovery::AddMountPoint API causing the asset data gatherer from spuriously thinking it is not idle after calls to ScanPathsSynchronous.  This would lead to wasted time in UAssetRegistryImpl::WaitForCompletion.

diff --git a/Engine/Source/Runtime/AssetRegistry/Private/AssetDataGatherer.cpp b/Engine/Source/Runtime/AssetRegistry/Private/AssetDataGatherer.cpp
--- a/Engine/Source/Runtime/AssetRegistry/Private/AssetDataGatherer.cpp
+++ b/Engine/Source/Runtime/AssetRegistry/Private/AssetDataGatherer.cpp
@@ -2304,25 +2304,29 @@ void FAssetDataDiscovery::Shrink()
 	DirToScanBuffers.Empty();
 }
 
-void FAssetDataDiscovery::AddMountPoint(const FString& LocalAbsPath, FStringView LongPackageName)
+void FAssetDataDiscovery::AddMountPoint(const FString& LocalAbsPath, FStringView LongPackageName, bool& bOutAlreadyExisted)
 {
 	CHECK_IS_NOT_LOCKED_CURRENT_THREAD(ResultsLock);
 	FGathererScopeLock TreeScopeLock(&TreeLock);
-	SetIsIdle(false);
-	AddMountPointInternal(LocalAbsPath, LongPackageName);
+	bool bAlreadyExisted = false;
+	AddMountPointInternal(LocalAbsPath, LongPackageName, bAlreadyExisted);
+	if (!bAlreadyExisted)
+	{
+		SetIsIdle(false);
+	}
 }
 
-void FAssetDataDiscovery::AddMountPointInternal(const FString& LocalAbsPath, FStringView LongPackageName)
+void FAssetDataDiscovery::AddMountPointInternal(const FString& LocalAbsPath, FStringView LongPackageName, bool& bOutAlreadyExisted)
 {
 	CHECK_IS_LOCKED_CURRENT_THREAD(TreeLock);
 	TArray<FMountDir*> ChildMounts;
 	FMountDir* ParentMount = nullptr;
-	bool bExists = false;
+	bOutAlreadyExisted = false;
 	for (TUniquePtr<FMountDir>& ExistingMount : MountDirs)
 	{
 		if (FPathViews::Equals(ExistingMount->GetLocalAbsPath(), LocalAbsPath))
 		{
-			bExists = true;
+			bOutAlreadyExisted = true;
 			break;
 		}
 		else if (FPathViews::IsParentPathOf(ExistingMount->GetLocalAbsPath(), LocalAbsPath))
@@ -2341,7 +2345,7 @@ void FAssetDataDiscovery::AddMountPointInternal(const FString& LocalAbsPath, FSt
 			}
 		}
 	}
-	if (bExists)
+	if (bOutAlreadyExisted)
 	{
 		return;
 	}
@@ -4658,7 +4662,10 @@ void FAssetDataGatherer::Shrink()
 
 void FAssetDataGatherer::AddMountPoint(FStringView LocalPath, FStringView LongPackageName)
 {
-	Discovery->AddMountPoint(NormalizeLocalPath(LocalPath), NormalizeLongPackageName(LongPackageName));
+	bool bAlreadyExisted = false;
+	Discovery->AddMountPoint(NormalizeLocalPath(LocalPath), NormalizeLongPackageName(LongPackageName), bAlreadyExisted);
+
+	if (!bAlreadyExisted)
 	{
 		FGathererScopeLock ResultsScopeLock(&ResultsLock);
 		SetIsIdle(false);
@@ -4675,13 +4682,18 @@ void FAssetDataGatherer::AddRequiredMountPoints(TArrayView<FString> LocalPaths)
 	TStringBuilder<128> MountPackageName;
 	TStringBuilder<128> MountFilePath;
 	TStringBuilder<128> RelPath;
+	bool bAllExisted = true;
 	for (const FString& LocalPath : LocalPaths)
 	{
 		if (FPackageName::TryGetMountPointForPath(LocalPath, MountPackageName, MountFilePath, RelPath))
 		{
-			Discovery->AddMountPoint(NormalizeLocalPath(MountFilePath), NormalizeLongPackageName(MountPackageName));
+			bool bAlreadyExisted = false;
+			Discovery->AddMountPoint(NormalizeLocalPath(MountFilePath), NormalizeLongPackageName(MountPackageName), bAlreadyExisted);
+			bAllExisted = bAllExisted && bAlreadyExisted;
 		}
 	}
+
+	if (!bAllExisted)
 	{
 		FGathererScopeLock ResultsScopeLock(&ResultsLock);
 		SetIsIdle(false);
diff --git a/Engine/Source/Runtime/AssetRegistry/Private/AssetDataGathererPrivate.h b/Engine/Source/Runtime/AssetRegistry/Private/AssetDataGathererPrivate.h
--- a/Engine/Source/Runtime/AssetRegistry/Private/AssetDataGathererPrivate.h
+++ b/Engine/Source/Runtime/AssetRegistry/Private/AssetDataGathererPrivate.h
@@ -622,7 +622,7 @@ public:
 	 * Register the given LocalAbsPath/LongPackageName pair that came from FPackageName's list of mount points as
 	 * a mountpoint to track. Will not be scanned until allow listed.
 	 */
-	void AddMountPoint(const FString& LocalAbsPath, FStringView LongPackageName);
+	void AddMountPoint(const FString& LocalAbsPath, FStringView LongPackageName, bool& bOutAlreadyExisted);
 	/** Remove the mountpoint because FPackageName has removed it. */
 	void RemoveMountPoint(const FString& LocalAbsPath);
 	/** Raise the priority until completion of scans of the given path and its subdirs. */
@@ -682,7 +682,7 @@ private:
 		const UE::AssetDataGather::Private::FSetPathProperties& Properties, bool bConfirmedExists);
 
 	/** Add the given path as a MountPoint and update child registrations. */
-	void AddMountPointInternal(const FString& LocalAbsPath, FStringView LongPackageName);
+	void AddMountPointInternal(const FString& LocalAbsPath, FStringView LongPackageName, bool& bOutAlreadyExisted);
 	/** Remove the given path as a MountPoint and update child registrations. */
 	void RemoveMountPointInternal(const FString& LocalAbsPath);
 
