Asset Registry:
- Move FWorldPartitionHelpers::FixupRedirectedAssetPath to UAssetRegistryHelpers::FixupRedirectedAssetPath

diff --git a/Engine/Source/Runtime/AssetRegistry/Private/AssetRegistryHelpers.cpp b/Engine/Source/Runtime/AssetRegistry/Private/AssetRegistryHelpers.cpp
--- a/Engine/Source/Runtime/AssetRegistry/Private/AssetRegistryHelpers.cpp
+++ b/Engine/Source/Runtime/AssetRegistry/Private/AssetRegistryHelpers.cpp
@@ -10,6 +10,10 @@
 #include "Modules/ModuleManager.h"
 #include "UObject/LinkerLoad.h"
 
+#if WITH_EDITOR
+#include "Misc/RedirectCollector.h"
+#endif
+
 TScriptInterface<IAssetRegistry> UAssetRegistryHelpers::GetAssetRegistry()
 {
 	return &UAssetRegistryImpl::Get();
@@ -227,3 +231,98 @@ UAssetRegistryHelpers::FTemporaryCachingModeScope::~FTemporaryCachingModeScope()
 {
 	UAssetRegistryHelpers::GetAssetRegistry()->SetTemporaryCachingMode(PreviousCachingMode);
 }
+
+bool UAssetRegistryHelpers::FixupRedirectedAssetPath(FSoftObjectPath& InOutSoftObjectPath)
+{
+	if (InOutSoftObjectPath.IsNull())
+	{
+		// Empty path, no redirect
+		return true;
+	}
+
+#if WITH_EDITOR
+	// Check GRedirectCollector first for faster fixup
+	FSoftObjectPath FoundRedirection = GRedirectCollector.GetAssetPathRedirection(InOutSoftObjectPath.GetWithoutSubPath());
+	if (!FoundRedirection.IsNull())
+	{
+		InOutSoftObjectPath.SetPath(FoundRedirection.GetAssetPath(), InOutSoftObjectPath.GetSubPathString());
+		return true;
+	}
+#endif
+
+	if (InOutSoftObjectPath.GetAssetName().IsEmpty())
+	{
+		// A package name. No need to ask the asset registry for assets, it wont find any
+		return true;
+	}
+
+	const FAssetData* AssetData;
+	IAssetRegistry& AssetRegistry = IAssetRegistry::GetChecked();
+
+	FTopLevelAssetPath AssetPath = InOutSoftObjectPath.GetAssetPath();
+	TSet<FTopLevelAssetPath> Visited;
+	Visited.Add(AssetPath);
+
+	for (;;)
+	{
+		TArray<FAssetData> Assets;
+		AssetRegistry.ScanFilesSynchronous({ AssetPath.GetPackageName().ToString() }, /*bForceRescan*/false);
+		AssetRegistry.GetAssetsByPackageName(AssetPath.GetPackageName(), Assets, /*bIncludeOnlyOnDiskAssets*/true);
+
+		if (!Assets.Num())
+		{
+			UE_LOG(LogCore, Warning, TEXT("Failed to find assets for asset path '%s'"), *AssetPath.ToString());
+			return false;
+		}
+
+		AssetData = Assets.FindByPredicate([&AssetPath](const FAssetData& AssetData)
+		{
+			return (AssetData.ToSoftObjectPath().GetAssetPath() == AssetPath);
+		});
+
+		if (!AssetData)
+		{
+			UE_LOG(LogCore, Warning, TEXT("Failed to find asset for asset path '%s'"), *AssetPath.ToString());
+			return false;
+		}
+
+		if (!AssetData->IsRedirector())
+		{
+			break;
+		}
+
+		FString DestinationObjectPath;
+		if (!AssetData->GetTagValue(TEXT("DestinationObject"), DestinationObjectPath))
+		{
+			UE_LOG(LogCore, Warning, TEXT("Failed to follow redirector for '%s'"), *AssetPath.ToString());
+			return false;
+		}
+
+		// Update asset path
+		AssetPath = FTopLevelAssetPath(DestinationObjectPath);
+
+		bool bAlreadyVisited = false;
+		Visited.Add(AssetPath, &bAlreadyVisited);
+		if (bAlreadyVisited)
+		{
+			UE_LOG(LogCore, Warning, TEXT("Asset path was already visited '%s'"), *AssetPath.ToString());
+			return false;
+		}
+	}
+
+	InOutSoftObjectPath.SetPath(AssetPath, InOutSoftObjectPath.GetSubPathString());
+
+	return true;
+}
+
+bool UAssetRegistryHelpers::FixupRedirectedAssetPath(FName& InOutAssetPath)
+{
+	FSoftObjectPath SoftObjectPath(InOutAssetPath.ToString());
+	if (FixupRedirectedAssetPath(SoftObjectPath))
+	{
+		InOutAssetPath = FName(*SoftObjectPath.ToString());
+		return true;
+	}
+
+	return false;
+}
diff --git a/Engine/Source/Runtime/AssetRegistry/Public/AssetRegistry/AssetRegistryHelpers.h b/Engine/Source/Runtime/AssetRegistry/Public/AssetRegistry/AssetRegistryHelpers.h
--- a/Engine/Source/Runtime/AssetRegistry/Public/AssetRegistry/AssetRegistryHelpers.h
+++ b/Engine/Source/Runtime/AssetRegistry/Public/AssetRegistry/AssetRegistryHelpers.h
@@ -137,4 +137,17 @@ public:
 	/** Checks to see if the given asset data is a blueprint with a base class in the ClassNameSet. This checks the parent asset tag */
 	static bool ASSETREGISTRY_API IsAssetDataBlueprintOfClassSet(const FAssetData& AssetData, const TSet<FTopLevelAssetPath>& ClassNameSet);
 
+	/**
+	 * Resolves the provided asset path using asset redirectors.
+	 * @param	InOutAssetPath	Asset path to resolve [In/Out].
+	 * @return	False if an error occurred while trying to resolve, else true.
+	 */
+	static ASSETREGISTRY_API bool FixupRedirectedAssetPath(FName& InOutAssetPath);
+	
+	/**
+	 * Resolves the provided soft object asset path using asset redirectors.
+	 * @param	InOutSoftObjectPath	Asset path to resolve [In/Out].
+	 * @return	False if an error occurred while trying to resolve, else true.
+	 */
+	static ASSETREGISTRY_API bool FixupRedirectedAssetPath(FSoftObjectPath& InOutSoftObjectPath);
 };
diff --git a/Engine/Source/Runtime/Engine/Private/WorldPartition/DataLayer/WorldDataLayersActorDesc.cpp b/Engine/Source/Runtime/Engine/Private/WorldPartition/DataLayer/WorldDataLayersActorDesc.cpp
--- a/Engine/Source/Runtime/Engine/Private/WorldPartition/DataLayer/WorldDataLayersActorDesc.cpp
+++ b/Engine/Source/Runtime/Engine/Private/WorldPartition/DataLayer/WorldDataLayersActorDesc.cpp
@@ -4,6 +4,7 @@
 
 #if WITH_EDITOR
 
+#include "AssetRegistry/AssetRegistryHelpers.h"
 #include "WorldPartition/WorldPartitionHelpers.h"
 #include "WorldPartition/DataLayer/WorldDataLayers.h"
 #include "WorldPartition/DataLayer/DataLayerInstanceWithAsset.h"
@@ -38,7 +39,7 @@ FArchive& operator<<(FArchive& Ar, FDataLayerInstanceDesc& Desc)
 	// Fixup redirected data layer asset path
 	if (Ar.IsLoading() && Desc.bIsUsingAsset)
 	{
-		FWorldPartitionHelpers::FixupRedirectedAssetPath(Desc.AssetPath);
+		UAssetRegistryHelpers::FixupRedirectedAssetPath(Desc.AssetPath);
 	}
 
 	if (Ar.CustomVer(FFortniteMainBranchObjectVersion::GUID) >= FFortniteMainBranchObjectVersion::WorldPartitionActorFilter)
diff --git a/Engine/Source/Runtime/Engine/Private/WorldPartition/Filter/WorldPartitionActorFilter.cpp b/Engine/Source/Runtime/Engine/Private/WorldPartition/Filter/WorldPartitionActorFilter.cpp
--- a/Engine/Source/Runtime/Engine/Private/WorldPartition/Filter/WorldPartitionActorFilter.cpp
+++ b/Engine/Source/Runtime/Engine/Private/WorldPartition/Filter/WorldPartitionActorFilter.cpp
@@ -4,6 +4,7 @@
 #include "WorldPartition/WorldPartitionActorLoaderInterface.h"
 #include "WorldPartition/WorldPartitionHelpers.h"
 #include "UObject/FortniteMainBranchObjectVersion.h"
+#include "AssetRegistry/AssetRegistryHelpers.h"
 
 #if WITH_EDITORONLY_DATA
 FWorldPartitionActorFilter::FOnWorldPartitionActorFilterChanged FWorldPartitionActorFilter::OnWorldPartitionActorFilterChanged;
@@ -171,7 +172,7 @@ bool FWorldPartitionActorFilter::Serialize(FArchive& Ar)
 			{
 				Ar << AssetPath;
 			}
-			FWorldPartitionHelpers::FixupRedirectedAssetPath(AssetPath);
+			UAssetRegistryHelpers::FixupRedirectedAssetPath(AssetPath);
 
 			bool bIncluded;
 			Ar << bIncluded;
diff --git a/Engine/Source/Runtime/Engine/Private/WorldPartition/WorldPartitionActorDesc.cpp b/Engine/Source/Runtime/Engine/Private/WorldPartition/WorldPartitionActorDesc.cpp
--- a/Engine/Source/Runtime/Engine/Private/WorldPartition/WorldPartitionActorDesc.cpp
+++ b/Engine/Source/Runtime/Engine/Private/WorldPartition/WorldPartitionActorDesc.cpp
@@ -8,6 +8,7 @@
 #include "Misc/ArchiveMD5.h"
 #include "UObject/MetaData.h"
 #include "AssetRegistry/AssetData.h"
+#include "AssetRegistry/AssetRegistryHelpers.h"
 #include "Serialization/MemoryReader.h"
 #include "Serialization/MemoryWriter.h"
 #include "UObject/UE5MainStreamObjectVersion.h"
@@ -689,7 +690,7 @@ PRAGMA_ENABLE_DEPRECATION_WARNINGS
 	{
 		for (FName& DataLayer : DataLayers)
 		{
-			FWorldPartitionHelpers::FixupRedirectedAssetPath(DataLayer);
+			UAssetRegistryHelpers::FixupRedirectedAssetPath(DataLayer);
 		}
 	}
 }
diff --git a/Engine/Source/Runtime/Engine/Private/WorldPartition/WorldPartitionActorDescArchive.cpp b/Engine/Source/Runtime/Engine/Private/WorldPartition/WorldPartitionActorDescArchive.cpp
--- a/Engine/Source/Runtime/Engine/Private/WorldPartition/WorldPartitionActorDescArchive.cpp
+++ b/Engine/Source/Runtime/Engine/Private/WorldPartition/WorldPartitionActorDescArchive.cpp
@@ -12,6 +12,7 @@
 #include "UObject/FortniteSeasonBranchObjectVersion.h"
 #include "UObject/FortniteMainBranchObjectVersion.h"
 #include "AssetRegistry/AssetData.h"
+#include "AssetRegistry/AssetRegistryHelpers.h"
 
 FActorDescArchive::FActorDescArchive(FArchive& InArchive, FWorldPartitionActorDesc* InActorDesc)
 	: FArchiveProxy(InArchive)
@@ -117,7 +118,7 @@ FArchive& FActorDescArchive::operator<<(FSoftObjectPath& Value)
 
 	if (IsLoading())
 	{
-		FWorldPartitionHelpers::FixupRedirectedAssetPath(Value);
+		UAssetRegistryHelpers::FixupRedirectedAssetPath(Value);
 	}
 
 	return *this;
diff --git a/Engine/Source/Runtime/Engine/Private/WorldPartition/WorldPartitionActorDescUtils.cpp b/Engine/Source/Runtime/Engine/Private/WorldPartition/WorldPartitionActorDescUtils.cpp
--- a/Engine/Source/Runtime/Engine/Private/WorldPartition/WorldPartitionActorDescUtils.cpp
+++ b/Engine/Source/Runtime/Engine/Private/WorldPartition/WorldPartitionActorDescUtils.cpp
@@ -4,6 +4,7 @@
 #include "WorldPartition/WorldPartitionActorDescUtils.h"
 #include "AssetRegistry/ARFilter.h"
 #include "AssetRegistry/AssetData.h"
+#include "AssetRegistry/AssetRegistryHelpers.h"
 #include "AssetRegistry/AssetRegistryModule.h"
 #include "Modules/ModuleManager.h"
 #include "WorldPartition/WorldPartitionLog.h"
@@ -174,6 +175,6 @@ void FWorldPartitionActorDescUtils::ReplaceActorDescriptorPointerFromActor(const
 
 bool FWorldPartitionActorDescUtils::FixupRedirectedAssetPath(FName& InOutAssetPath)
 {
-	return FWorldPartitionHelpers::FixupRedirectedAssetPath(InOutAssetPath);
+	return UAssetRegistryHelpers::FixupRedirectedAssetPath(InOutAssetPath);
 }
 #endif
diff --git a/Engine/Source/Runtime/Engine/Private/WorldPartition/WorldPartitionHelpers.cpp b/Engine/Source/Runtime/Engine/Private/WorldPartition/WorldPartitionHelpers.cpp
--- a/Engine/Source/Runtime/Engine/Private/WorldPartition/WorldPartitionHelpers.cpp
+++ b/Engine/Source/Runtime/Engine/Private/WorldPartition/WorldPartitionHelpers.cpp
@@ -16,6 +16,7 @@
 #if WITH_EDITOR
 #include "Misc/RedirectCollector.h"
 #include "AssetRegistry/AssetRegistryModule.h"
+#include "AssetRegistry/AssetRegistryHelpers.h"
 #include "Modules/ModuleManager.h"
 #endif
 
@@ -326,85 +327,12 @@ bool FWorldPartitionHelpers::ConvertRuntimePathToEditorPath(const FSoftObjectPat
 
 bool FWorldPartitionHelpers::FixupRedirectedAssetPath(FSoftObjectPath& InOutSoftObjectPath)
 {
-	if (InOutSoftObjectPath.IsNull())
-	{
-		// Empty path, no redirect
-		return true;
-	}
-
-	// Check GRedirectCollector first for faster fixup
-	FSoftObjectPath FoundRedirection = GRedirectCollector.GetAssetPathRedirection(InOutSoftObjectPath.GetWithoutSubPath());
-	if (!FoundRedirection.IsNull())
-	{
-		InOutSoftObjectPath.SetPath(FoundRedirection.GetAssetPath(), InOutSoftObjectPath.GetSubPathString());
-		return true;
-	}
-
-	if (InOutSoftObjectPath.GetAssetName().IsEmpty())
-	{
-		// A package name. No need to ask the asset registry for assets, it wont find any
-		return true;
-	}
-
-	const FAssetData* AssetData;
-	IAssetRegistry& AssetRegistry = FModuleManager::LoadModuleChecked<FAssetRegistryModule>(TEXT("AssetRegistry")).Get();
-
-	FTopLevelAssetPath AssetPath = InOutSoftObjectPath.GetAssetPath();
-
-	for (;;)
-	{
-		TArray<FAssetData> Assets;
-		AssetRegistry.ScanFilesSynchronous({ AssetPath.GetPackageName().ToString() }, /*bForceRescan*/false);
-		AssetRegistry.GetAssetsByPackageName(AssetPath.GetPackageName(), Assets, /*bIncludeOnlyOnDiskAssets*/true);
-
-		if (!Assets.Num())
-		{
-			UE_LOG(LogWorldPartition, Warning, TEXT("Failed to find assets for asset path '%s'"), *AssetPath.ToString());
-			return false;
-		}
-
-		AssetData = Assets.FindByPredicate([&AssetPath](const FAssetData& AssetData)
-		{
-			return (AssetData.ToSoftObjectPath().GetAssetPath() == AssetPath);
-		});
-
-		if (!AssetData)
-		{
-			UE_LOG(LogWorldPartition, Warning, TEXT("Failed to find asset for asset path '%s'"), *AssetPath.ToString());
-			return false;
-		}
-
-		if (!AssetData->IsRedirector())
-		{
-			break;
-		}
-
-		FString DestinationObjectPath;
-		if (!AssetData->GetTagValue(TEXT("DestinationObject"), DestinationObjectPath))
-		{
-			UE_LOG(LogWorldPartition, Warning, TEXT("Failed to follow redirector for '%s'"), *AssetPath.ToString());
-			return false;
-		}
-
-		// Update asset path
-		AssetPath = FTopLevelAssetPath(DestinationObjectPath);
-	}
-
-	InOutSoftObjectPath.SetPath(AssetPath, InOutSoftObjectPath.GetSubPathString());
-
-	return true;
+	return UAssetRegistryHelpers::FixupRedirectedAssetPath(InOutSoftObjectPath);
 }
 
 bool FWorldPartitionHelpers::FixupRedirectedAssetPath(FName& InOutAssetPath)
 {
-	FSoftObjectPath SoftObjectPath(InOutAssetPath.ToString());
-	if (FixupRedirectedAssetPath(SoftObjectPath))
-	{
-		InOutAssetPath = FName(*SoftObjectPath.ToString());
-		return true;
-	}
-
-	return false;
+	return UAssetRegistryHelpers::FixupRedirectedAssetPath(InOutAssetPath);
 }
 
 #endif // #if WITH_EDITOR
diff --git a/Engine/Source/Runtime/Engine/Public/WorldPartition/WorldPartitionHelpers.h b/Engine/Source/Runtime/Engine/Public/WorldPartition/WorldPartitionHelpers.h
--- a/Engine/Source/Runtime/Engine/Public/WorldPartition/WorldPartitionHelpers.h
+++ b/Engine/Source/Runtime/Engine/Public/WorldPartition/WorldPartitionHelpers.h
@@ -142,18 +142,10 @@ public:
 	// Runtime/Editor conversions
 	static ENGINE_API bool ConvertRuntimePathToEditorPath(const FSoftObjectPath& InPath, FSoftObjectPath& OutPath);
 
-	/**
-	 * Resolves the provided asset path using asset redirectors.
-	 * @param	InOutAssetPath	Asset path to resolve [In/Out].
-	 * @return	False if an error occured while trying to resolve, else true.
-	 */
+	UE_DEPRECATED(5.4, "Use UAssetRegistryHelpers::FixupRedirectedAssetPath instead.")
 	static ENGINE_API bool FixupRedirectedAssetPath(FName& InOutAssetPath);
 	
-	/**
-	 * Resolves the provided soft object asset path using asset redirectors.
-	 * @param	InOutSoftObjectPath	Asset path to resolve [In/Out].
-	 * @return	False if an error occured while trying to resolve, else true.
-	 */
+	UE_DEPRECATED(5.4, "Use UAssetRegistryHelpers::FixupRedirectedAssetPath instead.")
 	static ENGINE_API bool FixupRedirectedAssetPath(FSoftObjectPath& InOutSoftObjectPath);
 #endif // WITH_EDITOR
 
