Change the asset registry to call the final FilesLoaded callback after processing any changes in the same tick.
This significantly improves editor load times from a synchronous flush because many systems register expensive add/remove callbacks from FilesLoaded

diff --git a/Engine/Source/Runtime/AssetRegistry/Private/AssetRegistry.cpp b/Engine/Source/Runtime/AssetRegistry/Private/AssetRegistry.cpp
--- a/Engine/Source/Runtime/AssetRegistry/Private/AssetRegistry.cpp
+++ b/Engine/Source/Runtime/AssetRegistry/Private/AssetRegistry.cpp
@@ -7282,18 +7282,6 @@ void UAssetRegistryImpl::Broadcast(UE::AssetRegistry::Impl::FEventContext& Event
 		return;
 	}
 
-	if (EventContext.bFileLoadedEventBroadcast)
-	{
-		FileLoadedEvent.Broadcast();
-		EventContext.bFileLoadedEventBroadcast = false;
-	}
-
-	if (EventContext.ProgressUpdateData.IsSet())
-	{
-		FileLoadProgressUpdatedEvent.Broadcast(*EventContext.ProgressUpdateData);
-		EventContext.ProgressUpdateData.Reset();
-	}
-
 	if (EventContext.PathEvents.Num())
 	{
 		// Batch add/remove events 
@@ -7435,6 +7423,23 @@ void UAssetRegistryImpl::Broadcast(UE::AssetRegistry::Impl::FEventContext& Event
 		FilesBlockedEvent.Broadcast(EventContext.BlockedFiles);
 		EventContext.BlockedFiles.Empty();
 	}
+
+	if (EventContext.ProgressUpdateData.IsSet())
+	{
+		FileLoadProgressUpdatedEvent.Broadcast(*EventContext.ProgressUpdateData);
+		EventContext.ProgressUpdateData.Reset();
+	}
+
+	// FileLoadedEvent needs to come after all of the AssetEvents. Some systems do more expensive work for
+	// AssetEvents after receiving FileLoadedEvent, because they batched up that work for all assets in the initial load in
+	// their FileLoadedEvent handler. The AssetEvents precede the FileLoadedEvent in the broadcast that is sent from
+	// TickGatherer, so it is correct to make them precede it in the order in which we broadcast the events.
+
+	if (EventContext.bFileLoadedEventBroadcast)
+	{
+		FileLoadedEvent.Broadcast();
+		EventContext.bFileLoadedEventBroadcast = false;
+	}
 }
 
 
