Added new function IAssetRegistry::FlushDirectoryWatcherAndWaitForCompletion which sleeps for 1ms, then ticks the directory watcher, then waits for completion

#rnx

diff --git a/Engine/Source/Runtime/AssetRegistry/Private/AssetRegistry.cpp b/Engine/Source/Runtime/AssetRegistry/Private/AssetRegistry.cpp
--- a/Engine/Source/Runtime/AssetRegistry/Private/AssetRegistry.cpp
+++ b/Engine/Source/Runtime/AssetRegistry/Private/AssetRegistry.cpp
@@ -1791,6 +1791,24 @@ void UAssetRegistryImpl::WaitForCompletion()
 	}
 }
 
+#if WITH_EDITOR
+void UAssetRegistryImpl::FlushDirectoryWatcherAndWaitForCompletion()
+{
+	TRACE_CPUPROFILER_EVENT_SCOPE(UAssetRegistryImpl::FlushDirectoryWatcherAndWaitForCompletion);
+
+	// Operating system might take a few milliseconds before detecting the change
+	FPlatformProcess::Sleep(0.100f);
+
+	// Tick the directory watcher to pick up the new files and pass them into OnDirectoryChanged
+	IDirectoryWatcher* DirectoryWatcher = FModuleManager::LoadModuleChecked<FDirectoryWatcherModule>(TEXT("DirectoryWatcher")).Get();
+	DirectoryWatcher->Tick(0.0f);
+
+	// Wait for the results to be processed.
+	// TODO: Change this wait to only wait for the new files
+	WaitForCompletion();
+}
+#endif
+
 void UAssetRegistryImpl::ClearGathererCache()
 {
 	LLM_SCOPE(ELLMTag::AssetRegistry);
diff --git a/Engine/Source/Runtime/AssetRegistry/Private/AssetRegistry.h b/Engine/Source/Runtime/AssetRegistry/Private/AssetRegistry.h
--- a/Engine/Source/Runtime/AssetRegistry/Private/AssetRegistry.h
+++ b/Engine/Source/Runtime/AssetRegistry/Private/AssetRegistry.h
@@ -99,6 +99,9 @@ public:
 	virtual bool IsSearchAllAssets() const override;
 	virtual bool IsSearchAsync() const override;
 	virtual void WaitForCompletion() override;
+#if WITH_EDITOR
+	virtual void FlushDirectoryWatcherAndWaitForCompletion() override;
+#endif
 	virtual void ClearGathererCache() override;
 	virtual void WaitForPackage(const FString& PackageName) override;
 	virtual void ScanSynchronous(const TArray<FString>& InPaths, const TArray<FString>& InFilePaths, UE::AssetRegistry::EScanFlags InScanFlags = UE::AssetRegistry::EScanFlags::None) override;
diff --git a/Engine/Source/Runtime/AssetRegistry/Public/AssetRegistry/IAssetRegistry.h b/Engine/Source/Runtime/AssetRegistry/Public/AssetRegistry/IAssetRegistry.h
--- a/Engine/Source/Runtime/AssetRegistry/Public/AssetRegistry/IAssetRegistry.h
+++ b/Engine/Source/Runtime/AssetRegistry/Public/AssetRegistry/IAssetRegistry.h
@@ -620,6 +620,13 @@ public:
 	UFUNCTION(BlueprintCallable, Category = "AssetRegistry")
 	virtual void WaitForCompletion() = 0;
 
+#if WITH_EDITOR
+	/**
+	 * Make directory watcher pick up new changes to files on disk, then wait for scan to be complete.
+	 */
+	virtual void FlushDirectoryWatcherAndWaitForCompletion() = 0;
+#endif
+
 	/**
 	 * Empty the global gatherer's cache and disable further caching of scans from disk.
 	 * Used to save memory when cooking after the scan is complete.
