AssetRegistry: Fix deadlock caused by calling arbitrary callback code (GetAssetRegistryTags) while holding the global UObject hash lock due to a FThreadSafeObjectIterator
#rnx

diff --git a/Engine/Source/Runtime/AssetRegistry/Private/AssetRegistry.cpp b/Engine/Source/Runtime/AssetRegistry/Private/AssetRegistry.cpp
--- a/Engine/Source/Runtime/AssetRegistry/Private/AssetRegistry.cpp
+++ b/Engine/Source/Runtime/AssetRegistry/Private/AssetRegistry.cpp
@@ -2119,78 +2119,107 @@ void EnumerateMemoryAssetsHelper(const FARCompiledFilter& InFilter, TSet<FName>&
 	const uint32 FilterWithoutPackageFlags = InFilter.WithoutPackageFlags | PKG_ForDiffing;
 	const uint32 FilterWithPackageFlags = InFilter.WithPackageFlags;
 
-	auto FilterInMemoryObjectLambda = [&](const UObject* Obj, bool& OutContinue)
+	struct FFilterData
 	{
-		if (Obj->IsAsset())
+		const UObject* Object;
+		const UPackage* Package;
+		FString PackageNameStr;
+		FSoftObjectPath ObjectPath;
+	};
+
+	/**
+	 * The portions of the filter that are safe to execute even in the UObject global hash lock in FThreadSafeObjectIterator
+	 * Returns true if the object passes the filter and should be copied into an array for calling the rest of the filter
+	 * outside the lock.
+	 */
+	auto PassesLockSafeFilter =
+		[&InFilter, bSkipARFilteredAssets, FilterWithoutPackageFlags, FilterWithPackageFlags]
+	(const UObject* Obj, FFilterData& FilterData)
+	{
+		if (!Obj->IsAsset())
 		{
-			// Skip assets that are currently loading
-			if (Obj->HasAnyFlags(RF_NeedLoad))
-			{
-				return;
-			}
+			return false;
+		}
 
-			check(!Obj->GetPackage()->HasAnyPackageFlags(PKG_PlayInEditor));
-			check(!Obj->GetOutermostObject()->GetPackage()->HasAnyPackageFlags(PKG_PlayInEditor));
+		// Skip assets that are currently loading
+		if (Obj->HasAnyFlags(RF_NeedLoad))
+		{
+			return false;
+		}
 
-			UPackage* InMemoryPackage = Obj->GetOutermost();
+		check(!Obj->GetPackage()->HasAnyPackageFlags(PKG_PlayInEditor));
+		check(!Obj->GetOutermostObject()->GetPackage()->HasAnyPackageFlags(PKG_PlayInEditor));
 
-			// Skip assets with any of the specified 'without' package flags 
-			if (InMemoryPackage->HasAnyPackageFlags(FilterWithoutPackageFlags))
-			{
-				return;
-			}
+		FilterData.Package = Obj->GetOutermost();
 
-			// Skip assets without any the specified 'with' packages flags
-			if (!InMemoryPackage->HasAllPackagesFlags(FilterWithPackageFlags))
-			{
-				return;
-			}
+		// Skip assets with any of the specified 'without' package flags 
+		if (FilterData.Package->HasAnyPackageFlags(FilterWithoutPackageFlags))
+		{
+			return false;
+		}
 
-			// Skip classes that report themselves as assets but that the editor AssetRegistry is currently not counting as assets
-			if (bSkipARFilteredAssets && UE::AssetRegistry::FFiltering::ShouldSkipAsset(Obj))
-			{
-				return;
-			}
+		// Skip assets without any the specified 'with' packages flags
+		if (!FilterData.Package->HasAllPackagesFlags(FilterWithPackageFlags))
+		{
+			return false;
+		}
 
-			// Package name
-			const FName PackageName = InMemoryPackage->GetFName();
+		// Skip classes that report themselves as assets but that the editor AssetRegistry is currently not counting as assets
+		if (bSkipARFilteredAssets && UE::AssetRegistry::FFiltering::ShouldSkipAsset(Obj))
+		{
+			return false;
+		}
 
-			OutPackageNamesWithAssets.Add(PackageName);
+		// Package name
+		const FName PackageName = FilterData.Package->GetFName();
 
-			if (InFilter.PackageNames.Num() && !InFilter.PackageNames.Contains(PackageName))
-			{
-				return;
-			}
+		if (InFilter.PackageNames.Num() && !InFilter.PackageNames.Contains(PackageName))
+		{
+			return false;
+		}
 
-			// Asset Path
-			const FSoftObjectPath ObjectPath = FSoftObjectPath(Obj);
-			if (InFilter.SoftObjectPaths.Num() > 0)
+		// Asset Path
+		FilterData.ObjectPath = FSoftObjectPath(Obj);
+		if (InFilter.SoftObjectPaths.Num() > 0)
+		{
+			if (!InFilter.SoftObjectPaths.Contains(FilterData.ObjectPath))
 			{
-				if (!InFilter.SoftObjectPaths.Contains(ObjectPath))
-				{
-					return;
-				}
+				return false;
 			}
+		}
 
-			// Package path
-			const FString PackageNameStr = InMemoryPackage->GetName();
-			if (InFilter.PackagePaths.Num() > 0)
+		// Package path
+		PackageName.ToString(FilterData.PackageNameStr);
+		if (InFilter.PackagePaths.Num() > 0)
+		{
+			const FName PackagePath = FName(*FPackageName::GetLongPackagePath(FilterData.PackageNameStr));
+			if (!InFilter.PackagePaths.Contains(PackagePath))
 			{
-				const FName PackagePath = FName(*FPackageName::GetLongPackagePath(PackageNameStr));
-				if (!InFilter.PackagePaths.Contains(PackagePath))
-				{
-					return;
-				}
+				return false;
 			}
+		}
 
-			// Could perhaps save some FName -> String conversions by creating this a bit earlier using the UObject constructor
-			// to get package name and path.
-			FAssetData PartialAssetData(PackageNameStr, ObjectPath.ToString(), Obj->GetClass()->GetClassPathName(), FAssetDataTagMap(),
-				InMemoryPackage->GetChunkIDs(), InMemoryPackage->GetPackageFlags());
+		FilterData.Object = Obj;
+		return true;
+	};
 
-			// All filters passed, except for AssetRegistry filter; caller must check that one
-			OutContinue = Callback(Obj, MoveTemp(PartialAssetData));
-		}
+	auto RunUnsafeFilterAndCallback =
+		[&Callback, &OutPackageNamesWithAssets]
+	(FFilterData& FilterData, bool& bOutContinue)
+	{
+		// We mark the package found for this passing asset, so that any followup search for assets on disk will not
+		// add a duplicate of this Asset. We do this here for convenience; it would be more correct to call it only for assets that
+		// pass the callers remaining filters inside of Callback
+		OutPackageNamesWithAssets.Add(FilterData.Package->GetFName());
+
+		// Could perhaps save some FName -> String conversions by creating this a bit earlier using the UObject constructor
+		// to get package name and path.
+		FAssetData PartialAssetData(MoveTemp(FilterData.PackageNameStr), FilterData.ObjectPath.ToString(),
+			FilterData.Object->GetClass()->GetClassPathName(), FAssetDataTagMap(),
+			FilterData.Package->GetChunkIDs(), FilterData.Package->GetPackageFlags());
+
+		// All filters passed, except for AssetRegistry filter; caller must check that one
+		bOutContinue = Callback(FilterData.Object, MoveTemp(PartialAssetData));
 	};
 
 	// Iterate over all in-memory assets to find the ones that pass the filter components
@@ -2235,23 +2264,37 @@ void EnumerateMemoryAssetsHelper(const FARCompiledFilter& InFilter, TSet<FName>&
 			}
 		}
 
-		for (UObject* Object : InMemoryObjects)
+		FFilterData ScratchFilterData;
+		for (const UObject* Object : InMemoryObjects)
 		{
-			bool bContinue = true;
-			FilterInMemoryObjectLambda(Object, bContinue);
-			if (!bContinue)
+			if (PassesLockSafeFilter(Object, ScratchFilterData))
 			{
-				bOutStopIteration = true;
-				return;
+				bool bContinue = true;
+				RunUnsafeFilterAndCallback(ScratchFilterData, bContinue);
+				if (!bContinue)
+				{
+					bOutStopIteration = true;
+					return;
+				}
 			}
 		}
 	}
 	else
 	{
+		TArray<FFilterData> FirstPassFilterResults;
+		FFilterData ScratchFilterData;
 		for (FThreadSafeObjectIterator ObjIt; ObjIt; ++ObjIt)
+		{
+			if (PassesLockSafeFilter(*ObjIt, ScratchFilterData))
+			{
+				FirstPassFilterResults.Add(MoveTemp(ScratchFilterData));
+			}
+		}
+
+		for (FFilterData& FilterData : FirstPassFilterResults)
 		{
 			bool bContinue = true;
-			FilterInMemoryObjectLambda(*ObjIt, bContinue);
+			RunUnsafeFilterAndCallback(FilterData, bContinue);
 			if (!bContinue)
 			{
 				bOutStopIteration = true;
