AssetRegistryState: DumpState function: sort AssetTags and AssetBundles sections to remove indeterminism in comparisons. Add the AssetTags section to the DumpAssetRegistryForCooker dumped state.
#rnx

diff --git a/Engine/Source/Editor/UnrealEd/Private/CookOnTheFlyServer.cpp b/Engine/Source/Editor/UnrealEd/Private/CookOnTheFlyServer.cpp
--- a/Engine/Source/Editor/UnrealEd/Private/CookOnTheFlyServer.cpp
+++ b/Engine/Source/Editor/UnrealEd/Private/CookOnTheFlyServer.cpp
@@ -8324,7 +8324,8 @@ void DumpAssetRegistryForCooker(IAssetRegistry* AssetRegistry)
 	DeleteReportDir.DeleteDirectory(DumpDir);
 	PlatformFile.CreateDirectoryTree(*DumpDir);
 	TArray<FString> Pages;
-	TArray<FString> Arguments({ TEXT("ObjectPath"),TEXT("PackageName"),TEXT("Path"),TEXT("Class"),TEXT("Tag"), TEXT("DependencyDetails"), TEXT("PackageData"), TEXT("LegacyDependencies") });
+	TArray<FString> Arguments({ TEXT("ObjectPath"),TEXT("PackageName"),TEXT("Path"),TEXT("Class"),TEXT("Tag"),
+		TEXT("DependencyDetails"), TEXT("PackageData"), TEXT("LegacyDependencies"), TEXT("AssetTags") });
 	AssetRegistry->DumpState(Arguments, Pages, 10000 /* LinesPerPage */);
 	int PageIndex = 0;
 	TStringBuilder<256> FileName;
diff --git a/Engine/Source/Runtime/AssetRegistry/Private/AssetRegistryState.cpp b/Engine/Source/Runtime/AssetRegistry/Private/AssetRegistryState.cpp
--- a/Engine/Source/Runtime/AssetRegistry/Private/AssetRegistryState.cpp
+++ b/Engine/Source/Runtime/AssetRegistry/Private/AssetRegistryState.cpp
@@ -161,28 +161,28 @@ void FAssetRegistryState::FilterTags(const FAssetDataTagMapSharedView& InTagsAnd
 	// Exclude denied tags or include only allowed tags, based on how we were configured in ini
 	for (const auto& TagPair : InTagsAndValues)
 	{
-		const bool bInAllClassesList = AllClassesFilterList && (AllClassesFilterList->Contains(TagPair.Key) || AllClassesFilterList->Contains(UE::AssetRegistry::WildcardFName));
-		const bool bInClassSpecificList = ClassSpecificFilterList && (ClassSpecificFilterList->Contains(TagPair.Key) || ClassSpecificFilterList->Contains(UE::AssetRegistry::WildcardFName));
-		if (Options.bUseAssetRegistryTagsAllowListInsteadOfDenyList)
-		{
-			// It's an allow list, only include it if it is in the all classes list or in the class specific list
+			const bool bInAllClassesList = AllClassesFilterList && (AllClassesFilterList->Contains(TagPair.Key) || AllClassesFilterList->Contains(UE::AssetRegistry::WildcardFName));
+			const bool bInClassSpecificList = ClassSpecificFilterList && (ClassSpecificFilterList->Contains(TagPair.Key) || ClassSpecificFilterList->Contains(UE::AssetRegistry::WildcardFName));
+			if (Options.bUseAssetRegistryTagsAllowListInsteadOfDenyList)
+			{
+				// It's an allow list, only include it if it is in the all classes list or in the class specific list
 			if (bInAllClassesList || bInClassSpecificList)
 			{
 				// It is in the allow list. Keep it.
 				OutTagsAndValues.Add(TagPair.Key, TagPair.Value.ToLoose());
 			}
-		}
-		else
-		{
-			// It's a deny list, include it unless it is in the all classes list or in the class specific list
-			if (!bInAllClassesList && !bInClassSpecificList)
+			}
+			else
 			{
+				// It's a deny list, include it unless it is in the all classes list or in the class specific list
+			if (!bInAllClassesList && !bInClassSpecificList)
+		{
 				// It isn't in the deny list. Keep it.
-				OutTagsAndValues.Add(TagPair.Key, TagPair.Value.ToLoose());
-			}
+			OutTagsAndValues.Add(TagPair.Key, TagPair.Value.ToLoose());
 		}
 	}
 }
+}
 
 void FAssetRegistryState::InitializeFromExistingAndPrune(const FAssetRegistryState & ExistingState, const TSet<FName>& RequiredPackages, const TSet<FName>& RemovePackages,
 	const TSet<int32> ChunksToKeep, const FAssetRegistrySerializationOptions& Options)
@@ -2747,6 +2747,18 @@ void FAssetRegistryState::Dump(const TArray<FString>& Arguments, TArray<FString>
 			});
 	}
 
+	TArray<FAssetData*> SortedAssets;
+	auto InitializeSortedAssets = [&SortedAssets, this]()
+	{
+		if (SortedAssets.Num() != CachedAssets.Num())
+		{
+			SortedAssets = CachedAssets.Array();
+			Algo::Sort(SortedAssets, [](const FAssetData* A, FAssetData* B)
+				{ return A->GetSoftObjectPath().LexicalLess(B->GetSoftObjectPath()); }
+			);
+		}
+	};
+
 	if (bAllFields || Arguments.Contains(TEXT("AssetTags")))
 	{
 		int32 Counter = 0;
@@ -2754,7 +2766,8 @@ void FAssetRegistryState::Dump(const TArray<FString>& Arguments, TArray<FString>
 		PageBuffer.Append(TEXT("--- Begin AssetTags ---"));
 		AddLine();
 
-		for (const FAssetData* AssetData : CachedAssets)
+		InitializeSortedAssets();
+		for (const FAssetData* AssetData : SortedAssets)
 		{
 			if (AssetData->TagsAndValues.Num() == 0)
 			{
@@ -2961,7 +2974,8 @@ void FAssetRegistryState::Dump(const TArray<FString>& Arguments, TArray<FString>
 		PageBuffer.Append(TEXT("--- Begin AssetBundles ---"));
 		AddLine();
 
-		for (FAssetData* AssetData : CachedAssets)
+		InitializeSortedAssets();
+		for (const FAssetData* AssetData : SortedAssets)
 		{
 			if (AssetData->TaggedAssetBundles.IsValid())
 			{
