Add VerseCreated/Deleted API to the AssetRegistry

diff --git a/Engine/Source/Runtime/AssetRegistry/Private/AssetRegistry.cpp b/Engine/Source/Runtime/AssetRegistry/Private/AssetRegistry.cpp
--- a/Engine/Source/Runtime/AssetRegistry/Private/AssetRegistry.cpp
+++ b/Engine/Source/Runtime/AssetRegistry/Private/AssetRegistry.cpp
@@ -4005,6 +4005,64 @@ void UAssetRegistryImpl::AssetTagsFinalized(const UObject& FinalizedAsset)
 #endif
 }
 
+bool UAssetRegistryImpl::VerseCreated(const FString& FilePathOnDisk)
+{
+	checkf(GIsEditor, TEXT("Updating the AssetRegistry is only available in editor"));
+	if (!FAssetDataGatherer::IsVerseFile(FilePathOnDisk))
+	{
+		return false;
+	}
+
+	FString PackageName;
+	if (!FPackageName::TryConvertFilenameToLongPackageName(FilePathOnDisk, PackageName))
+	{
+		return false;
+	}
+
+	FNameBuilder VersePackagePathName;
+	VersePackagePathName.Append(PackageName);
+	VersePackagePathName.Append(FPathViews::GetExtension(FilePathOnDisk, /*bIncludeDot=*/true));
+
+	UE::AssetRegistry::Impl::FEventContext EventContext;
+	{
+		LLM_SCOPE(ELLMTag::AssetRegistry);
+		FWriteScopeLock InterfaceScopeLock(InterfaceLock);
+		GuardedData.AddVerseFile(EventContext, *VersePackagePathName);
+	}
+	Broadcast(EventContext);
+
+	return true;
+}
+
+bool UAssetRegistryImpl::VerseDeleted(const FString& FilePathOnDisk)
+{
+	checkf(GIsEditor, TEXT("Updating the AssetRegistry is only available in editor"));
+	if (!FAssetDataGatherer::IsVerseFile(FilePathOnDisk))
+	{
+		return false;
+	}
+
+	FString PackageName;
+	if (!FPackageName::TryConvertFilenameToLongPackageName(FilePathOnDisk, PackageName))
+	{
+		return false;
+	}
+
+	FNameBuilder VersePackagePathName;
+	VersePackagePathName.Append(PackageName);
+	VersePackagePathName.Append(FPathViews::GetExtension(FilePathOnDisk, /*bIncludeDot=*/true));
+
+	UE::AssetRegistry::Impl::FEventContext EventContext;
+	{
+		LLM_SCOPE(ELLMTag::AssetRegistry);
+		FWriteScopeLock InterfaceScopeLock(InterfaceLock);
+		GuardedData.RemoveVerseFile(EventContext, *VersePackagePathName);
+	}
+	Broadcast(EventContext);
+
+	return true;
+}
+
 void UAssetRegistryImpl::PackageDeleted(UPackage* DeletedPackage)
 {
 	checkf(GIsEditor, TEXT("Updating the AssetRegistry is only available in editor"));
@@ -5497,15 +5555,7 @@ void FAssetRegistryImpl::VerseFilesGathered(Impl::FEventContext& EventContext, c
 	{
 		FName VerseFilePath = VerseResults.PopFrontValue();
 
-		bool bAlreadyExists = false;
-		CachedVerseFiles.Add(VerseFilePath, &bAlreadyExists);
-		if (!bAlreadyExists)
-		{
-			FName VerseDirectoryPath(FPathViews::GetPath(WriteToString<256>(VerseFilePath)));
-			TArray<FName>& FilePathsArray = CachedVerseFilesByPath.FindOrAdd(VerseDirectoryPath);
-			FilePathsArray.Add(VerseFilePath);
-			EventContext.VerseEvents.Emplace(VerseFilePath, Impl::FEventContext::EEvent::Added);
-		}
+		AddVerseFile(EventContext, VerseFilePath);
 
 		// Check to see if we have run out of time in this tick
 		if (TickStartTime >= 0 && (FPlatformTime::Seconds() - TickStartTime) > Impl::MaxSecondsPerFrame)
@@ -5760,6 +5810,19 @@ void FAssetRegistryImpl::RemovePackageData(Impl::FEventContext& EventContext, co
 	}
 }
 
+void FAssetRegistryImpl::AddVerseFile(Impl::FEventContext& EventContext, FName VerseFilePathToAdd)
+{
+	bool bAlreadyExists = false;
+	CachedVerseFiles.Add(VerseFilePathToAdd, &bAlreadyExists);
+	if (!bAlreadyExists)
+	{
+		FName VerseDirectoryPath(FPathViews::GetPath(WriteToString<256>(VerseFilePathToAdd)));
+		TArray<FName>& FilePathsArray = CachedVerseFilesByPath.FindOrAdd(VerseDirectoryPath);
+		FilePathsArray.Add(VerseFilePathToAdd);
+		EventContext.VerseEvents.Emplace(VerseFilePathToAdd, Impl::FEventContext::EEvent::Added);
+	}
+}
+
 void FAssetRegistryImpl::RemoveVerseFile(Impl::FEventContext& EventContext, FName VerseFilePathToRemove)
 {
 	if (CachedVerseFiles.Remove(VerseFilePathToRemove))
diff --git a/Engine/Source/Runtime/AssetRegistry/Private/AssetRegistry.h b/Engine/Source/Runtime/AssetRegistry/Private/AssetRegistry.h
--- a/Engine/Source/Runtime/AssetRegistry/Private/AssetRegistry.h
+++ b/Engine/Source/Runtime/AssetRegistry/Private/AssetRegistry.h
@@ -148,6 +148,9 @@ public:
 	virtual void AssetFullyUpdateTags(UObject* Object) override;
 	virtual void AssetTagsFinalized(const UObject& FinalizedAsset) override;
 
+	virtual bool VerseCreated(const FString& FilePath) override;
+	virtual bool VerseDeleted(const FString& FilePath) override;
+
 	virtual void PackageDeleted(UPackage* DeletedPackage) override;
 
 	DECLARE_DERIVED_EVENT( UAssetRegistryImpl, IAssetRegistry::FAssetAddedEvent, FAssetAddedEvent);
diff --git a/Engine/Source/Runtime/AssetRegistry/Private/AssetRegistryImpl.h b/Engine/Source/Runtime/AssetRegistry/Private/AssetRegistryImpl.h
--- a/Engine/Source/Runtime/AssetRegistry/Private/AssetRegistryImpl.h
+++ b/Engine/Source/Runtime/AssetRegistry/Private/AssetRegistryImpl.h
@@ -236,6 +236,8 @@ public:
 	bool RemoveAssetPath(Impl::FEventContext& EventContext, FName PathToRemove, bool bEvenIfAssetsStillExist = false);
 	/** Removes the asset data associated with this package from the look-up maps */
 	void RemovePackageData(Impl::FEventContext& EventContext, const FName PackageName);
+	/** Adds the Verse file to the look up maps */
+	void AddVerseFile(Impl::FEventContext& EventContext, FName VerseFilePathToAdd);
 	/** Removes the Verse file from the look-up maps */
 	void RemoveVerseFile(Impl::FEventContext& EventContext, FName VerseFilePathToRemove);
 
diff --git a/Engine/Source/Runtime/AssetRegistry/Public/AssetRegistry/IAssetRegistry.h b/Engine/Source/Runtime/AssetRegistry/Public/AssetRegistry/IAssetRegistry.h
--- a/Engine/Source/Runtime/AssetRegistry/Public/AssetRegistry/IAssetRegistry.h
+++ b/Engine/Source/Runtime/AssetRegistry/Public/AssetRegistry/IAssetRegistry.h
@@ -688,6 +688,12 @@ public:
 	 */
 	virtual void AssetFullyUpdateTags(UObject* Object) = 0;
 
+	/** Informs the asset registry that a Verse file has been created on disk */
+	virtual bool VerseCreated(const FString& FilePathOnDisk) = 0;
+
+	/** Informs the asset registry that a Verse file has been deleted from disk */
+	virtual bool VerseDeleted(const FString& FilePathOnDisk) = 0;
+
 	/** Informs the asset registry that an in-memory package has been deleted, and all associated assets should be removed */
 	virtual void PackageDeleted (UPackage* DeletedPackage) = 0;
 
