Fixed an issue where the engine might query the asset registry before the premade asset registry had finished loading. Added a method to IAssetRegistry interface to flush the premade asset registry load and added a call in UAssetManager::StartInitialLoading.

diff --git a/Engine/Source/Runtime/AssetRegistry/Private/AssetRegistry.cpp b/Engine/Source/Runtime/AssetRegistry/Private/AssetRegistry.cpp
--- a/Engine/Source/Runtime/AssetRegistry/Private/AssetRegistry.cpp
+++ b/Engine/Source/Runtime/AssetRegistry/Private/AssetRegistry.cpp
@@ -1792,6 +1792,23 @@ void UAssetRegistryImpl::WaitForCompletion()
 	}
 }
 
+void UAssetRegistryImpl::WaitForPremadeAssetRegistry()
+{
+	TRACE_CPUPROFILER_EVENT_SCOPE(UAssetRegistryImpl::WaitForPremadeAssetRegistry);
+	using namespace UE::AssetRegistry::Impl;
+
+	FEventContext EventContext;
+	{
+		LLM_SCOPE(ELLMTag::AssetRegistry);
+		FWriteScopeLock InterfaceScopeLock(InterfaceLock);
+		FClassInheritanceContext InheritanceContext;
+		FClassInheritanceBuffer InheritanceBuffer;
+		GetInheritanceContextWithRequiredLock(InterfaceScopeLock, InheritanceContext, InheritanceBuffer);
+		GuardedData.ConditionalLoadPremadeAssetRegistry(*this, EventContext, InterfaceScopeLock);
+	}
+	Broadcast(EventContext);
+}
+
 #if WITH_EDITOR
 void UAssetRegistryImpl::FlushDirectoryWatcherAndWaitForCompletion()
 {
diff --git a/Engine/Source/Runtime/AssetRegistry/Private/AssetRegistry.h b/Engine/Source/Runtime/AssetRegistry/Private/AssetRegistry.h
--- a/Engine/Source/Runtime/AssetRegistry/Private/AssetRegistry.h
+++ b/Engine/Source/Runtime/AssetRegistry/Private/AssetRegistry.h
@@ -99,6 +99,7 @@ public:
 	virtual bool IsSearchAllAssets() const override;
 	virtual bool IsSearchAsync() const override;
 	virtual void WaitForCompletion() override;
+	virtual void WaitForPremadeAssetRegistry() override;
 #if WITH_EDITOR
 	virtual void FlushDirectoryWatcherAndWaitForCompletion() override;
 #endif
diff --git a/Engine/Source/Runtime/AssetRegistry/Public/AssetRegistry/IAssetRegistry.h b/Engine/Source/Runtime/AssetRegistry/Public/AssetRegistry/IAssetRegistry.h
--- a/Engine/Source/Runtime/AssetRegistry/Public/AssetRegistry/IAssetRegistry.h
+++ b/Engine/Source/Runtime/AssetRegistry/Public/AssetRegistry/IAssetRegistry.h
@@ -620,6 +620,12 @@ public:
 	UFUNCTION(BlueprintCallable, Category = "AssetRegistry")
 	virtual void WaitForCompletion() = 0;
 
+	/**
+	 * Wait for premade asset registry load to be completed, for use when querying the asset registry prior to flushing in ScanPathSynchronous.
+	 * Will no-op if no such asset registry exists.
+	 */
+	virtual void WaitForPremadeAssetRegistry() = 0;
+
 #if WITH_EDITOR
 	/**
 	 * Make directory watcher pick up new changes to files on disk, then wait for scan to be complete.
diff --git a/Engine/Source/Runtime/Engine/Private/AssetManager.cpp b/Engine/Source/Runtime/Engine/Private/AssetManager.cpp
--- a/Engine/Source/Runtime/Engine/Private/AssetManager.cpp
+++ b/Engine/Source/Runtime/Engine/Private/AssetManager.cpp
@@ -3614,6 +3614,9 @@ bool UAssetManager::GetPackageManagers(FName PackageName, bool bRecurseToParents
 
 void UAssetManager::StartInitialLoading()
 {
+	// The scan below queries asset registry, so we should make sure the premade registry is finished loading if it exists.
+	GetAssetRegistry().WaitForPremadeAssetRegistry();
+
 	GInitialBulkScan.StartOnce(this);
 
 	ScanPrimaryAssetTypesFromConfig();
