GetAssetRegistryTags changes for upcoming API change:
* Move the FArchiveCookContext enums into a namespace so they can be included from other types that need them without using FArchiveCookContext.
* Make construction of CookContext a TOptional to avoid constructing it when not cooking; the new constructor will be more difficult to call.
* Technically enforce the lifetime of ArchiveCookContext during SaveCookedPackage by making it a TOptional.
* Add a default constructor for FArchiveTag
* Remove a deprecated UObject virtual function that will be changing again in the upcoming change.
* Move AssetRegistry's bUpdateDiskCacheAfterLoad into a variable outside the critical section to make it cheaper to read; it is immutable after AssetRegistry is constructed.
* Remove use of the soft-deprecated FAssetData constructor that takes a bool rather than ECreationFlags.
* Reduce some duplication in FAssetRegistryState::FilterTags
* FAssetBundleData: use the native exporter in the case when defaults are provided but are pointing to the object itself; this combination of parameters means that no defaults are available and the entire struct should be exported. Use of the native exporter is necessary for correctness; the generic exporter incorrectly writes both the deprecated and the new variables.
ExportToText is used with defaults == object when called from GetAssetRegistryTagFromProperty called from GetAssetRegistryTags in obj.cpp.
* UObjectRedirector::GetAssetRegistryTags by contract should call the Super's version of the function.
* UBlueprintGeneratedClass::GetAssetRegistryTags does not need to append actor tags when saving in the cookcommandlet since GetAssetRegistryTags is not called when saving cooked packages and no other saves are executed in the cook commandlet.
* NiagaraEditorSettings GetAssetRegistryTags: sort the ClassPaths in the ClassUsageList to remove indeterminism.
* Whitespace fixes
#rnx

diff --git a/Engine/Plugins/FX/Niagara/Source/NiagaraEditor/Private/NiagaraEditorSettings.cpp b/Engine/Plugins/FX/Niagara/Source/NiagaraEditor/Private/NiagaraEditorSettings.cpp
--- a/Engine/Plugins/FX/Niagara/Source/NiagaraEditor/Private/NiagaraEditorSettings.cpp
+++ b/Engine/Plugins/FX/Niagara/Source/NiagaraEditor/Private/NiagaraEditorSettings.cpp
@@ -734,6 +734,7 @@ UObject::FAssetRegistryTag UNiagaraEditorSettings::CreateClassUsageAssetRegistry
 				ClassPaths.Add(ObjectInPackage->GetClass()->GetClassPathName().ToString());
 			}
 		}
+		ClassPaths.Sort([](const FString& A, const FString& B) { return A < B; });
 		ClassUsageList = FString::Join(ClassPaths, TEXT("\n"));
 	}
 	else
diff --git a/Engine/Source/Editor/UnrealEd/Private/Commandlets/SavePackageUtilitiesCommandlet.cpp b/Engine/Source/Editor/UnrealEd/Private/Commandlets/SavePackageUtilitiesCommandlet.cpp
--- a/Engine/Source/Editor/UnrealEd/Private/Commandlets/SavePackageUtilitiesCommandlet.cpp
+++ b/Engine/Source/Editor/UnrealEd/Private/Commandlets/SavePackageUtilitiesCommandlet.cpp
@@ -36,11 +36,13 @@ int32 USavePackageUtilitiesCommandlet::Main(const FString& Params)
 		FString Filename = FPaths::CreateTempFilename(*FPaths::ProjectSavedDir());
 
 		// CookData should only be nonzero if we are cooking.
+		TOptional<FArchiveCookContext> CookContext;
 		TOptional<FArchiveCookData> CookData;
-		FArchiveCookContext CookContext(Package, FArchiveCookContext::ECookTypeUnknown, FArchiveCookContext::ECookingDLCUnknown);
 		if (TargetPlatform != nullptr)
 		{
-			CookData.Emplace(*TargetPlatform, CookContext);
+			CookContext.Emplace(Package, UE::Cook::ECookType::Unknown,
+				UE::Cook::ECookingDLC::Unknown);
+			CookData.Emplace(*TargetPlatform, *CookContext);
 		}
 
 		FSavePackageArgs SaveArgs;
diff --git a/Engine/Source/Editor/UnrealEd/Private/CookOnTheFlyServer.cpp b/Engine/Source/Editor/UnrealEd/Private/CookOnTheFlyServer.cpp
--- a/Engine/Source/Editor/UnrealEd/Private/CookOnTheFlyServer.cpp
+++ b/Engine/Source/Editor/UnrealEd/Private/CookOnTheFlyServer.cpp
@@ -4279,13 +4279,11 @@ private: // Used only by UCookOnTheFlyServer, which has private access
 	bool bHasDelayLoaded = false;
 	bool bContainsMap = false;
 
-	// This holds cook data generated by Serialize() that isn't saved in the package - only valid in
-	// PerPlatform callbacks and is reset in SetupPlatform.
-	FArchiveCookContext ArchiveCookContext;
-
 	// Per-platform data, only valid in PerPlatform callbacks
 	const ITargetPlatform* TargetPlatform = nullptr;
 	FCookSavePackageContext* CookContext = nullptr;
+	// This holds cook data generated by Serialize() that isn't saved in the package
+	TOptional<FArchiveCookContext> ArchiveCookContext;
 	FSavePackageContext* SavePackageContext = nullptr;
 	ICookedPackageWriter* PackageWriter = nullptr;
 	FString PlatFilename;
@@ -5896,7 +5894,7 @@ void UCookOnTheFlyServer::SaveCookedPackage(UE::Cook::FSaveCookedPackageContext&
 			UE_SCOPED_HIERARCHICAL_COOKTIMER(GEditorSavePackage);
 			UE_TRACK_REFERENCING_PLATFORM_SCOPED(TargetPlatform);
 
-			FArchiveCookData CookData(*TargetPlatform, Context.ArchiveCookContext);
+			FArchiveCookData CookData(*TargetPlatform, *Context.ArchiveCookContext);
 			FSavePackageArgs SaveArgs;
 			SaveArgs.TopLevelFlags = Context.FlagsToCook;
 			SaveArgs.bForceByteSwapping = Context.bEndianSwap;
@@ -5987,9 +5985,6 @@ FSaveCookedPackageContext::FSaveCookedPackageContext(UCookOnTheFlyServer& InCOTF
 	, Package(PackageData.GetPackage())
 	, PackageName(Package ? Package->GetName() : FString())
 	, Filename(PackageData.GetFileName().ToString())
-	, ArchiveCookContext(Package,
-		InCOTFS.IsDirectorCookByTheBook() ? FArchiveCookContext::ECookByTheBook : FArchiveCookContext::ECookOnTheFly,
-		InCOTFS.IsCookingDLC() ? FArchiveCookContext::ECookingDLCYes : FArchiveCookContext::ECookingDLCNo)
 {
 }
 
@@ -6027,7 +6022,9 @@ void FSaveCookedPackageContext::SetupPlatform(const ITargetPlatform* InTargetPla
 	PlatFilename = Filename.Replace(TEXT("[Platform]"), *TargetPlatform->PlatformName());
 	bPlatformSetupSuccessful = false;
 
-	ArchiveCookContext.Reset();
+	ArchiveCookContext.Emplace(Package,
+		COTFS.IsDirectorCookByTheBook() ? UE::Cook::ECookType::ByTheBook : UE::Cook::ECookType::OnTheFly,
+		COTFS.IsCookingDLC() ? UE::Cook::ECookingDLC::Yes : UE::Cook::ECookingDLC::No);
 
 	// don't save Editor resources from the Engine if the target doesn't have editoronly data
 	if (COTFS.IsCookFlagSet(ECookInitializationFlags::SkipEditorContent) &&
@@ -6269,7 +6266,7 @@ void FSaveCookedPackageContext::FinishPlatform()
 			}
 		}
 		Reporter.UpdateAssetRegistryData(Package->GetFName(), Package, CookResult, &SavePackageResult,
-			MoveTemp(*ArchiveCookContext.GetCookTagList()), bIncludeOnlyDiskAssets,
+			MoveTemp(*ArchiveCookContext->GetCookTagList()), bIncludeOnlyDiskAssets,
 			MoveTemp(OverrideAssetPackageData), MoveTemp(OverridePackageDependencies));
 	}
 
@@ -6340,6 +6337,8 @@ void FSaveCookedPackageContext::FinishPlatform()
 
 	bHasRetryErrorCode |= bIsRetryErrorCode;
 	bHasFirstPlatformResults = true;
+
+	ArchiveCookContext.Reset();
 }
 
 void FSaveCookedPackageContext::FinishPackage()
diff --git a/Engine/Source/Editor/UnrealEd/Private/EditorEngine.cpp b/Engine/Source/Editor/UnrealEd/Private/EditorEngine.cpp
--- a/Engine/Source/Editor/UnrealEd/Private/EditorEngine.cpp
+++ b/Engine/Source/Editor/UnrealEd/Private/EditorEngine.cpp
@@ -4618,6 +4618,11 @@ FSavePackageResultStruct UEditorEngine::Save(UPackage* InOuter, UObject* InAsset
 	const bool bSavingConcurrent = !!(SaveArgs.SaveFlags & ESaveFlags::SAVE_Concurrent);
 
 	FObjectSaveContextData ObjectSaveContext(InOuter, SaveArgs.GetTargetPlatform(), Filename, SaveArgs.SaveFlags);
+	if (InSaveArgs.ArchiveCookData)
+	{
+		ObjectSaveContext.CookType = InSaveArgs.ArchiveCookData->CookContext.GetCookType();
+		ObjectSaveContext.CookingDLC = InSaveArgs.ArchiveCookData->CookContext.GetCookingDLC();
+	}
 	UWorld *OriginalOwningWorld = nullptr;
 	if ( World )
 	{
diff --git a/Engine/Source/Runtime/AssetRegistry/Private/AssetRegistry.cpp b/Engine/Source/Runtime/AssetRegistry/Private/AssetRegistry.cpp
--- a/Engine/Source/Runtime/AssetRegistry/Private/AssetRegistry.cpp
+++ b/Engine/Source/Runtime/AssetRegistry/Private/AssetRegistry.cpp
@@ -694,7 +694,7 @@ EChunkProgressReportingType::Type GetChunkAvailabilityProgressType(EAssetAvailab
 	return ChunkReportType;
 }
 
-ASSETREGISTRY_API const TCHAR* GetDevelopmentAssetRegistryFilename()
+const TCHAR* GetDevelopmentAssetRegistryFilename()
 {
 	return TEXT("DevelopmentAssetRegistry.bin");
 }
@@ -852,7 +852,15 @@ void FAssetRegistryImpl::Initialize(Impl::FInitializeContext& Context)
 	StoreGatherResultsTimeSeconds = 0.f;
 
 	// By default update the disk cache once on asset load, to incorporate changes made in PostLoad. This only happens in editor builds
-	bUpdateDiskCacheAfterLoad = true;
+#if !WITH_EDITOR
+	Context.bUpdateDiskCacheAfterLoad = false;
+#else
+	Context.bUpdateDiskCacheAfterLoad = true;
+	if (GConfig)
+	{
+		GConfig->GetBool(TEXT("AssetRegistry"), TEXT("bUpdateDiskCacheAfterLoad"), Context.bUpdateDiskCacheAfterLoad, GEngineIni);
+	}
+#endif
 
 	bIsTempCachingAlwaysEnabled = ASSETREGISTRY_CACHE_ALWAYS_ENABLED;
 	bIsTempCachingEnabled = bIsTempCachingAlwaysEnabled;
@@ -908,14 +916,6 @@ void FAssetRegistryImpl::Initialize(Impl::FInitializeContext& Context)
 		bStartedAsyncGather
 	});
 
-#if WITH_EDITOR
-	if (GConfig)
-	{
-		GConfig->GetBool(TEXT("AssetRegistry"), TEXT("bUpdateDiskCacheAfterLoad"), bUpdateDiskCacheAfterLoad, GEngineIni);
-	}
-#endif
-	Context.bUpdateDiskCacheAfterLoad = bUpdateDiskCacheAfterLoad;
-
 	// Content roots always exist; add them as paths
 	FPackageName::QueryRootContentPaths(Context.RootContentPaths, false, false, true);
 	for (const FString& AssetPath : Context.RootContentPaths)
@@ -1012,7 +1012,8 @@ void UAssetRegistryImpl::InitializeEvents(UE::AssetRegistry::Impl::FInitializeCo
 		}
 	}
 
-	if (Context.bUpdateDiskCacheAfterLoad)
+	bUpdateDiskCacheAfterLoad = Context.bUpdateDiskCacheAfterLoad;
+	if (bUpdateDiskCacheAfterLoad)
 	{
 		FCoreUObjectDelegates::OnAssetLoaded.AddUObject(this, &UAssetRegistryImpl::OnAssetLoaded);
 	}
@@ -1530,7 +1531,7 @@ void UAssetRegistryImpl::FinishDestroy()
 			}
 		}
 
-		if (GuardedData.IsUpdateDiskCacheAfterLoad())
+		if (bUpdateDiskCacheAfterLoad)
 		{
 			FCoreUObjectDelegates::OnAssetLoaded.RemoveAll(this);
 		}
@@ -2407,7 +2408,7 @@ FAssetData UAssetRegistryImpl::GetAssetByObjectPath(const FSoftObjectPath& Objec
 		{
 			if (!bSkipARFilteredAssets || !UE::AssetRegistry::FFiltering::ShouldSkipAsset(Asset))
 			{
-				return FAssetData(Asset, false /* bAllowBlueprintClass */);
+				return FAssetData(Asset, FAssetData::ECreationFlags::None /** Do not allow blueprint classes */);
 			}
 			else
 			{
@@ -3830,7 +3831,7 @@ void UAssetRegistryImpl::AssetCreated(UObject* NewAsset)
 		{
 			checkf(IsInGameThread(), TEXT("AssetCreated is not yet implemented as callable from other threads"));
 			// Let subscribers know that the new asset was added to the registry
-			FAssetData AssetData = FAssetData(NewAsset, true /* bAllowBlueprintClass */); 
+			FAssetData AssetData = FAssetData(NewAsset, FAssetData::ECreationFlags::AllowBlueprintClass);
 			AssetAddedEvent.Broadcast(AssetData);
 			OnAssetsAdded().Broadcast({ AssetData });
 
@@ -3883,7 +3884,7 @@ void UAssetRegistryImpl::AssetDeleted(UObject* DeletedAsset)
 
 		if (!bShouldSkipAsset)
 		{
-			FAssetData AssetDataDeleted = FAssetData(DeletedAsset, true /* bAllowBlueprintClass */);
+			FAssetData AssetDataDeleted = FAssetData(DeletedAsset, FAssetData::ECreationFlags::AllowBlueprintClass);
 
 			checkf(IsInGameThread(), TEXT("AssetDeleted is not yet implemented as callable from other threads"));
 			// Let subscribers know that the asset was removed from the registry
@@ -3941,7 +3942,10 @@ void UAssetRegistryImpl::AssetRenamed(const UObject* RenamedAsset, const FString
 		if (!bShouldSkipAsset)
 		{
 			checkf(IsInGameThread(), TEXT("AssetRenamed is not yet implemented as callable from other threads"));
-			AssetRenamedEvent.Broadcast(FAssetData(RenamedAsset, true /* bAllowBlueprintClass */), OldObjectPath);
+			AssetRenamedEvent.Broadcast(
+				FAssetData(RenamedAsset, FAssetData::ECreationFlags::AllowBlueprintClass),
+				OldObjectPath
+			);
 		}
 	}
 }
@@ -3966,7 +3970,7 @@ void UAssetRegistryImpl::AssetsSaved(TArray<FAssetData>&& Assets)
 void UAssetRegistryImpl::AssetFullyUpdateTags(UObject* Object)
 {
 #if WITH_EDITOR
-	FAssetData AssetData(Object);
+	FAssetData AssetData(Object, FAssetData::ECreationFlags::None);
 	TArray<FAssetData> Assets;
 	Assets.Add(MoveTemp(AssetData));
 
@@ -6208,7 +6212,7 @@ void UAssetRegistryImpl::ProcessLoadedAssetsToUpdateCache(UE::AssetRegistry::Imp
 	{
 		LLM_SCOPE(ELLMTag::AssetRegistry);
 		FWriteScopeLock InterfaceScopeLock(InterfaceLock);
-		GuardedData.GetProcessLoadedAssetsBatch(BatchObjects, BatchSize);
+		GuardedData.GetProcessLoadedAssetsBatch(BatchObjects, BatchSize, bUpdateDiskCacheAfterLoad);
 		if (BatchObjects.Num() == 0)
 		{
 			return;
@@ -6232,7 +6236,7 @@ void UAssetRegistryImpl::ProcessLoadedAssetsToUpdateCache(UE::AssetRegistry::Imp
 				// If the object has changed and is no longer an asset, ignore it. This can happen when an Actor is modified during cooking to no longer have an external package
 				continue;
 			}
-			BatchAssetDatas.Add(FAssetData(LoadedObject, true /* bAllowBlueprintClass */));
+			BatchAssetDatas.Add(FAssetData(LoadedObject, FAssetData::ECreationFlags::AllowBlueprintClass));
 
 			// Check to see if we have run out of time in this tick
 			if (!bFlushFullBuffer &&
@@ -6251,7 +6255,7 @@ void UAssetRegistryImpl::ProcessLoadedAssetsToUpdateCache(UE::AssetRegistry::Imp
 		{
 			break;
 		}
-		GuardedData.GetProcessLoadedAssetsBatch(BatchObjects, BatchSize);
+		GuardedData.GetProcessLoadedAssetsBatch(BatchObjects, BatchSize, bUpdateDiskCacheAfterLoad);
 	}
 }
 
@@ -6263,7 +6267,8 @@ void FAssetRegistryImpl::AddLoadedAssetToProcess(const UObject& AssetLoaded)
 	LoadedAssetsToProcess.Add(&AssetLoaded);
 }
 
-void FAssetRegistryImpl::GetProcessLoadedAssetsBatch(TArray<const UObject*>& OutLoadedAssets, uint32 BatchSize)
+void FAssetRegistryImpl::GetProcessLoadedAssetsBatch(TArray<const UObject*>& OutLoadedAssets, uint32 BatchSize,
+	bool bUpdateDiskCacheAfterLoad)
 {
 	if (!GlobalGatherer.IsValid() || !bUpdateDiskCacheAfterLoad)
 	{
@@ -6319,7 +6324,7 @@ void FAssetRegistryImpl::PushProcessLoadedAssetsBatch(Impl::FEventContext& Event
 		else
 		{
 			// When updating disk-based AssetData with the AssetData from a loaded UObject, we keep
-			// existing tags from disk even if they no are no longer returned from the GetAssetRegistryTags
+			// existing tags from disk even if they are no longer returned from the GetAssetRegistryTags
 			// function on the loaded UObject. We do this because they might come from GetAssetRegistryTagsExtended,
 			// which is only called during Save.
 			// Modified tag values on the other hand do overwrite the old values from disk.
diff --git a/Engine/Source/Runtime/AssetRegistry/Private/AssetRegistry.h b/Engine/Source/Runtime/AssetRegistry/Private/AssetRegistry.h
--- a/Engine/Source/Runtime/AssetRegistry/Private/AssetRegistry.h
+++ b/Engine/Source/Runtime/AssetRegistry/Private/AssetRegistry.h
@@ -193,6 +193,14 @@ public:
 	virtual FFileLoadProgressUpdatedEvent& OnFileLoadProgressUpdated() override;
 
 	virtual bool IsLoadingAssets() const override;
+	virtual bool IsUpdateDiskCacheAfterLoad() const override
+	{
+#if WITH_EDITORONLY_DATA
+		return bUpdateDiskCacheAfterLoad;
+#else
+		return false;
+#endif
+	}
 
 	virtual void Tick (float DeltaTime) override;
 
@@ -301,6 +309,9 @@ private:
 #if WITH_EDITORONLY_DATA
 	/** If true, the asset registry will inject missing tags from UMetaData for cooked assets only in GetAssetRegistryTags */
 	bool bAddMetaDataTagsToOnGetExtraObjectTags = true;
+
+	/** If true, the AssetRegistry updates its on-disk information for an Asset whenever that Asset loads. */
+	bool bUpdateDiskCacheAfterLoad = true;
 #endif
 
 	/** The delegate to execute when one or more files have been blocked from the registry */
diff --git a/Engine/Source/Runtime/AssetRegistry/Private/AssetRegistryImpl.h b/Engine/Source/Runtime/AssetRegistry/Private/AssetRegistryImpl.h
--- a/Engine/Source/Runtime/AssetRegistry/Private/AssetRegistryImpl.h
+++ b/Engine/Source/Runtime/AssetRegistry/Private/AssetRegistryImpl.h
@@ -204,7 +204,8 @@ public:
 	void ClearGathererCache();
 #if WITH_EDITOR
 	void AssetsSaved(UE::AssetRegistry::Impl::FEventContext& EventContext, TArray<FAssetData>&& Assets);
-	void GetProcessLoadedAssetsBatch(TArray<const UObject*>& OutLoadedAssets, uint32 BatchSize);
+	void GetProcessLoadedAssetsBatch(TArray<const UObject*>& OutLoadedAssets, uint32 BatchSize,
+		bool bUpdateDiskCacheAfterLoad);
 	void PushProcessLoadedAssetsBatch(Impl::FEventContext& EventContext,
 		TArrayView<FAssetData> LoadedAssetDatas, TArrayView<const UObject*> UnprocessedFromBatch);
 	/** Call LoadCalculatedDependencies on each Package updated after the last LoadCalculatedDependencies. */
@@ -245,7 +246,6 @@ public:
 	void GetSubClasses(Impl::FClassInheritanceContext& InheritanceContext, const TArray<FTopLevelAssetPath>& InClassNames,
 		const TSet<FTopLevelAssetPath>& ExcludedClassNames, TSet<FTopLevelAssetPath>& SubClassNames) const;
 
-	bool IsUpdateDiskCacheAfterLoad() const { return bUpdateDiskCacheAfterLoad; }
 	bool IsInitialSearchCompleted() const { return bInitialSearchCompleted; }
 	bool IsTempCachingEnabled() const { return bIsTempCachingEnabled; }
 	bool IsTempCachingAlwaysEnabled() const { return bIsTempCachingAlwaysEnabled; }
@@ -394,9 +394,6 @@ private:
 
 	uint64 ClassGeneratorNamesRegisteredClassesVersionNumber;
 
-	/** If true, will cache AssetData loaded from in memory assets back into the disk cache */
-	bool bUpdateDiskCacheAfterLoad;
-
 	/** The tree of known cached paths that assets may reside within */
 	FPathTree CachedPathTree;
 
diff --git a/Engine/Source/Runtime/AssetRegistry/Private/AssetRegistryState.cpp b/Engine/Source/Runtime/AssetRegistry/Private/AssetRegistryState.cpp
--- a/Engine/Source/Runtime/AssetRegistry/Private/AssetRegistryState.cpp
+++ b/Engine/Source/Runtime/AssetRegistry/Private/AssetRegistryState.cpp
@@ -154,35 +154,34 @@ void FAssetRegistryState::Reset()
 	CachedPackageData.Empty();
 }
 
-void FAssetRegistryState::FilterTags(const FAssetDataTagMapSharedView& InTagsAndValues, FAssetDataTagMap& OutTagsAndValues, const TSet<FName>* ClassSpecificFilterList, const FAssetRegistrySerializationOptions& Options)
+void FAssetRegistryState::FilterTags(const FAssetDataTagMapSharedView& InTagsAndValues, FAssetDataTagMap& OutTagsAndValues,
+	const TSet<FName>* ClassSpecificFilterList, const FAssetRegistrySerializationOptions& Options)
 {
 	const TSet<FName>* AllClassesFilterList = Options.CookFilterlistTagsByClass.Find(UE::AssetRegistry::WildcardPathName);
 
 	// Exclude denied tags or include only allowed tags, based on how we were configured in ini
 	for (const auto& TagPair : InTagsAndValues)
 	{
-			const bool bInAllClassesList = AllClassesFilterList && (AllClassesFilterList->Contains(TagPair.Key) || AllClassesFilterList->Contains(UE::AssetRegistry::WildcardFName));
-			const bool bInClassSpecificList = ClassSpecificFilterList && (ClassSpecificFilterList->Contains(TagPair.Key) || ClassSpecificFilterList->Contains(UE::AssetRegistry::WildcardFName));
-			if (Options.bUseAssetRegistryTagsAllowListInsteadOfDenyList)
-			{
-				// It's an allow list, only include it if it is in the all classes list or in the class specific list
-			if (bInAllClassesList || bInClassSpecificList)
-			{
-				// It is in the allow list. Keep it.
-				OutTagsAndValues.Add(TagPair.Key, TagPair.Value.ToLoose());
-			}
-			}
-			else
-			{
-				// It's a deny list, include it unless it is in the all classes list or in the class specific list
-			if (!bInAllClassesList && !bInClassSpecificList)
+		bool bKeep = false;
+
+		const bool bInAllClassesList = AllClassesFilterList && (AllClassesFilterList->Contains(TagPair.Key) || AllClassesFilterList->Contains(UE::AssetRegistry::WildcardFName));
+		const bool bInClassSpecificList = ClassSpecificFilterList && (ClassSpecificFilterList->Contains(TagPair.Key) || ClassSpecificFilterList->Contains(UE::AssetRegistry::WildcardFName));
+		if (Options.bUseAssetRegistryTagsAllowListInsteadOfDenyList)
+		{
+			// It's an allow list, only include it if it is in the all classes list or in the class specific list
+			bKeep = bInAllClassesList || bInClassSpecificList;
+		}
+		else
+		{
+			// It's a deny list, include it unless it is in the all classes list or in the class specific list
+			bKeep = !bInAllClassesList && !bInClassSpecificList;
+		}
+		if (bKeep)
 		{
-				// It isn't in the deny list. Keep it.
 			OutTagsAndValues.Add(TagPair.Key, TagPair.Value.ToLoose());
 		}
 	}
 }
-}
 
 void FAssetRegistryState::InitializeFromExistingAndPrune(const FAssetRegistryState & ExistingState, const TSet<FName>& RequiredPackages, const TSet<FName>& RemovePackages,
 	const TSet<int32> ChunksToKeep, const FAssetRegistrySerializationOptions& Options)
diff --git a/Engine/Source/Runtime/AssetRegistry/Public/AssetRegistry/IAssetRegistry.h b/Engine/Source/Runtime/AssetRegistry/Public/AssetRegistry/IAssetRegistry.h
--- a/Engine/Source/Runtime/AssetRegistry/Public/AssetRegistry/IAssetRegistry.h
+++ b/Engine/Source/Runtime/AssetRegistry/Public/AssetRegistry/IAssetRegistry.h
@@ -784,6 +784,9 @@ public:
 	UFUNCTION(BlueprintCallable, Category = "AssetRegistry")
 	virtual bool IsLoadingAssets() const = 0;
 
+	/** If true, the AssetRegistry updates its on-disk information for an Asset whenever that Asset loads. */
+	virtual bool IsUpdateDiskCacheAfterLoad() const = 0;
+
 	/** Tick the asset registry */
 	virtual void Tick (float DeltaTime) = 0;
 
diff --git a/Engine/Source/Runtime/CoreUObject/Private/AssetRegistry/AssetBundleData.cpp b/Engine/Source/Runtime/CoreUObject/Private/AssetRegistry/AssetBundleData.cpp
--- a/Engine/Source/Runtime/CoreUObject/Private/AssetRegistry/AssetBundleData.cpp
+++ b/Engine/Source/Runtime/CoreUObject/Private/AssetRegistry/AssetBundleData.cpp
@@ -33,9 +33,12 @@ namespace UE::AssetBundleEntry::Private
 
 bool FAssetBundleEntry::ExportTextItem(FString& ValueStr, const FAssetBundleEntry& DefaultValue, UObject* Parent, int32 PortFlags, UObject* ExportRootScope) const
 {	
-	if (DefaultValue.IsValid())
+	// If the DefaultValue points to this, that is the import/export system's way of specifying that there are no
+	// available defaults and the entire struct should be exported. If the DefautlValue is an empty version of the
+	// struct, then we also will export the entire struct. Otherwise there are actual defaults and we're supposed to
+	// export a delta. This path does not handle that delta, fall back to normal export path
+	if (&DefaultValue != this && DefaultValue.IsValid())
 	{
-		// This path does not handle default values, fall back to normal export path
 		return false;
 	}
 
@@ -485,9 +488,12 @@ bool FAssetBundleData::ExportTextItem(FString& ValueStr, FAssetBundleData const&
 		// Empty, don't write anything to avoid it cluttering the asset registry tags
 		return true;
 	}
-	else if (DefaultValue.Bundles.Num() != 0)
+	// If the DefaultValue points to this, that is the import/export system's way of specifying that there are no
+	// available defaults and the entire struct should be exported. If the DefautlValue is an empty version of the
+	// struct, then we also will export the entire struct. Otherwise there are actual defaults and we're supposed to
+	// export a delta. This path does not handle that delta, fall back to normal export path
+	else if (&DefaultValue != this && DefaultValue.Bundles.Num() != 0)
 	{
-		// This path does not handle default values, fall back to normal export path
 		return false;
 	}
 	
diff --git a/Engine/Source/Runtime/CoreUObject/Private/UObject/Obj.cpp b/Engine/Source/Runtime/CoreUObject/Private/UObject/Obj.cpp
--- a/Engine/Source/Runtime/CoreUObject/Private/UObject/Obj.cpp
+++ b/Engine/Source/Runtime/CoreUObject/Private/UObject/Obj.cpp
@@ -2046,7 +2046,8 @@ FString GetConfigFilename( UObject* SourceObject )
 	return SourceObject->GetClass()->GetConfigName();
 }
 
-namespace UE { namespace Object { namespace Private {
+namespace UE::Object::Private
+{
 
 // Thread local state to avoid UObject::GetAssetRegistryTags() API change
 thread_local FAssetBundleData const** TGetAssetRegistryTags_OutBundles = nullptr;
@@ -2134,7 +2135,7 @@ static void GetAssetRegistryTagsFromSearchableProperties(const UObject* Object,
 	}
 }
 
-}}} // end namespace UE::Object::Private
+} // end namespace UE::Object::Private
 
 const FName FPrimaryAssetId::PrimaryAssetTypeTag(TEXT("PrimaryAssetType"));
 const FName FPrimaryAssetId::PrimaryAssetNameTag(TEXT("PrimaryAssetName"));
@@ -2188,10 +2189,6 @@ void UObject::GetExtendedAssetRegistryTagsForSave(const ITargetPlatform* TargetP
 {
 	// Notify external sources that we need tags for save.
 	FAssetRegistryTag::OnGetExtendedAssetRegistryTagsForSave.Broadcast(this, TargetPlatform, OutTags);
-
-	PRAGMA_DISABLE_DEPRECATION_WARNINGS
-	GetExternalActorExtendedAssetRegistryTags(OutTags);
-	PRAGMA_ENABLE_DEPRECATION_WARNINGS
 }
 #endif // WITH_EDITOR
 
diff --git a/Engine/Source/Runtime/CoreUObject/Private/UObject/ObjectRedirector.cpp b/Engine/Source/Runtime/CoreUObject/Private/UObject/ObjectRedirector.cpp
--- a/Engine/Source/Runtime/CoreUObject/Private/UObject/ObjectRedirector.cpp
+++ b/Engine/Source/Runtime/CoreUObject/Private/UObject/ObjectRedirector.cpp
@@ -62,6 +62,8 @@ bool UObjectRedirector::NeedsLoadForEditorGame() const
 
 void UObjectRedirector::GetAssetRegistryTags(TArray<FAssetRegistryTag>& OutTags) const
 {
+	Super::GetAssetRegistryTags(OutTags);
+
 	FString DestVal;
 	if ( DestinationObject != nullptr )
 	{
diff --git a/Engine/Source/Runtime/CoreUObject/Private/UObject/SavePackage.cpp b/Engine/Source/Runtime/CoreUObject/Private/UObject/SavePackage.cpp
--- a/Engine/Source/Runtime/CoreUObject/Private/UObject/SavePackage.cpp
+++ b/Engine/Source/Runtime/CoreUObject/Private/UObject/SavePackage.cpp
@@ -25,11 +25,13 @@ bool UPackage::SavePackage(UPackage* InOuter, UObject* Base, EObjectFlags TopLev
 	const ITargetPlatform* TargetPlatform, const FDateTime& FinalTimeStamp, bool bSlowTask)
 {
 	// CookData should only be nonzero if we are cooking.
+	TOptional<FArchiveCookContext> CookContext;
 	TOptional<FArchiveCookData> CookData;
-	FArchiveCookContext CookContext(InOuter, FArchiveCookContext::ECookTypeUnknown, FArchiveCookContext::ECookingDLCUnknown);
 	if (TargetPlatform != nullptr)
 	{
-		CookData.Emplace(*TargetPlatform, CookContext);
+		CookContext.Emplace(InOuter, UE::Cook::ECookType::Unknown,
+			UE::Cook::ECookingDLC::Unknown);
+		CookData.Emplace(*TargetPlatform, *CookContext);
 	}
 	FSavePackageArgs SaveArgs = { nullptr /* deprecated target platform */, CookData.GetPtrOrNull(), TopLevelFlags, SaveFlags, bForceByteSwapping,
 		bWarnOfLongFilename, bSlowTask, FinalTimeStamp, Error };
diff --git a/Engine/Source/Runtime/CoreUObject/Private/UObject/SavePackage/SaveContext.h b/Engine/Source/Runtime/CoreUObject/Private/UObject/SavePackage/SaveContext.h
--- a/Engine/Source/Runtime/CoreUObject/Private/UObject/SavePackage/SaveContext.h
+++ b/Engine/Source/Runtime/CoreUObject/Private/UObject/SavePackage/SaveContext.h
@@ -15,6 +15,7 @@
 #include "Serialization/UnversionedPropertySerialization.h"
 #include "Templates/PimplPtr.h"
 #include "Templates/UniquePtr.h"
+#include "UObject/ArchiveCookContext.h"
 #include "UObject/LinkerSave.h"
 #include "UObject/NameTypes.h"
 #include "UObject/ObjectMacros.h"
@@ -465,6 +466,11 @@ public:
 		}
 
 		ObjectSaveContext.Set(InPackage, GetTargetPlatform(), TargetPackagePath, SaveArgs.SaveFlags);
+		if (SaveArgs.ArchiveCookData)
+		{
+			ObjectSaveContext.CookType = SaveArgs.ArchiveCookData->CookContext.GetCookType();
+			ObjectSaveContext.CookingDLC = SaveArgs.ArchiveCookData->CookContext.GetCookingDLC();
+		}
 
 		// Setup the harvesting flags and generate the context for harvesting the package
 		SetupHarvestingRealms();
diff --git a/Engine/Source/Runtime/CoreUObject/Public/AssetRegistry/CookTagList.h b/Engine/Source/Runtime/CoreUObject/Public/AssetRegistry/CookTagList.h
--- a/Engine/Source/Runtime/CoreUObject/Public/AssetRegistry/CookTagList.h
+++ b/Engine/Source/Runtime/CoreUObject/Public/AssetRegistry/CookTagList.h
@@ -47,7 +47,7 @@ struct FCookTagList
 	*/
 	void Add(UObject* InObject, FName InTagName, FString&& InTagValue)
 	{
-		if (!ensure(InObject->GetOutermost() == Package))
+		if (!Package || !ensure(InObject->GetOutermost() == Package))
 		{
 			return;
 		}
diff --git a/Engine/Source/Runtime/CoreUObject/Public/UObject/ArchiveCookContext.h b/Engine/Source/Runtime/CoreUObject/Public/UObject/ArchiveCookContext.h
--- a/Engine/Source/Runtime/CoreUObject/Public/UObject/ArchiveCookContext.h
+++ b/Engine/Source/Runtime/CoreUObject/Public/UObject/ArchiveCookContext.h
@@ -1,8 +1,9 @@
 // Copyright Epic Games, Inc. All Rights Reserved.
 
 #pragma once
-#include "AssetRegistry/CookTagList.h"
 
+#include "AssetRegistry/CookTagList.h"
+#include "UObject/CookEnums.h"
 
 /**
 *	Accessor for data about the package being cooked during UObject::Serialize calls.
@@ -10,38 +11,39 @@
 struct FArchiveCookContext
 {
 public:
-	enum ECookType
-	{
-		ECookTypeUnknown,
-		ECookOnTheFly,
-		ECookByTheBook
-	};
-
-	enum ECookingDLC
-	{
-		ECookingDLCUnknown,
-		ECookingDLCYes,
-		ECookingDLCNo,
-	};
+	UE_DEPRECATED(5.4, "Use UE::Cook::ECookType")
+	typedef UE::Cook::ECookType ECookType;
+	UE_DEPRECATED(5.4, "Use UE::Cook::ECookingDLC")
+	typedef UE::Cook::ECookingDLC ECookingDLC;
+	UE_DEPRECATED(5.4, "Use UE::Cook::ECookType::Unknown")
+	static const UE::Cook::ECookType ECookTypeUnknown = UE::Cook::ECookType::Unknown;
+	UE_DEPRECATED(5.4, "Use UE::Cook::ECookType::OnTheFly")
+	static const UE::Cook::ECookType ECookOnTheFly = UE::Cook::ECookType::OnTheFly;
+	UE_DEPRECATED(5.4, "Use UE::Cook::ECookType::ByTheBook")
+	static const UE::Cook::ECookType ECookByTheBook = UE::Cook::ECookType::ByTheBook;
+	UE_DEPRECATED(5.4, "Use UE::Cook::ECookingDLC::Unknown")
+	static const UE::Cook::ECookingDLC ECookingDLCUnknown = UE::Cook::ECookingDLC::Unknown;
+	UE_DEPRECATED(5.4, "Use UE::Cook::ECookingDLC::Yes")
+	static const UE::Cook::ECookingDLC ECookingDLCYes = UE::Cook::ECookingDLC::Yes;
+	UE_DEPRECATED(5.4, "Use UE::Cook::ECookingDLC::No")
+	static const UE::Cook::ECookingDLC ECookingDLCNo = UE::Cook::ECookingDLC::No;
 
 private:
 
-	// Only valid for cook by the book.
+	/** CookTagList is only valid for cook by the book; it is not publically accessible otherwise. */
 	FCookTagList CookTagList;
 	bool bCookTagListEnabled = false;
-	ECookType CookType;
-	ECookingDLC CookingDLC = ECookingDLCUnknown;
+	UE::Cook::ECookType CookType = UE::Cook::ECookType::Unknown;
+	UE::Cook::ECookingDLC CookingDLC = UE::Cook::ECookingDLC::Unknown;
 
 public:
 
-
-	FArchiveCookContext(UPackage* InPackage, ECookType InCookType, ECookingDLC CookingDLC) :
-		CookTagList(InPackage),
-		bCookTagListEnabled(InPackage && InCookType == ECookByTheBook),
-		CookType(InCookType),
-		CookingDLC(CookingDLC)
+	FArchiveCookContext(UPackage* InPackage, UE::Cook::ECookType InCookType, UE::Cook::ECookingDLC CookingDLC)
+		: CookTagList(InPackage)
+		, bCookTagListEnabled(InPackage && InCookType == UE::Cook::ECookType::ByTheBook)
+		, CookType(InCookType)
+		, CookingDLC(CookingDLC)
 	{
-		
 	}
 
 	void Reset()
@@ -51,9 +53,9 @@ public:
 
 	FCookTagList* GetCookTagList() { return bCookTagListEnabled ? &CookTagList : nullptr; }
 
-	bool IsCookByTheBook() const { return CookType == ECookType::ECookByTheBook; }
-	bool IsCookOnTheFly() const { return CookType == ECookType::ECookOnTheFly; }
-	bool IsCookTypeUnknown() const { return CookType == ECookType::ECookTypeUnknown; }
-	ECookType GetCookType() const { return CookType; }
-	ECookingDLC GetCookingDLC() const { return CookingDLC; }
+	bool IsCookByTheBook() const { return CookType == UE::Cook::ECookType::ByTheBook; }
+	bool IsCookOnTheFly() const { return CookType == UE::Cook::ECookType::OnTheFly; }
+	bool IsCookTypeUnknown() const { return CookType == UE::Cook::ECookType::Unknown; }
+	UE::Cook::ECookType GetCookType() const { return CookType; }
+	UE::Cook::ECookingDLC GetCookingDLC() const { return CookingDLC; }
 };
\ No newline at end of file
diff --git a/Engine/Source/Runtime/CoreUObject/Public/UObject/CookEnums.h b/Engine/Source/Runtime/CoreUObject/Public/UObject/CookEnums.h
new file mode 100644
--- /dev/null
+++ b/Engine/Source/Runtime/CoreUObject/Public/UObject/CookEnums.h
@@ -0,0 +1,23 @@
+// Copyright Epic Games, Inc. All Rights Reserved.
+
+#pragma once
+
+namespace UE::Cook
+{
+
+enum class ECookType
+{
+	Unknown,
+	OnTheFly,
+	ByTheBook
+};
+
+enum class ECookingDLC
+{
+	Unknown,
+	Yes,
+	No,
+};
+
+}
+
diff --git a/Engine/Source/Runtime/CoreUObject/Public/UObject/Object.h b/Engine/Source/Runtime/CoreUObject/Public/UObject/Object.h
--- a/Engine/Source/Runtime/CoreUObject/Public/UObject/Object.h
+++ b/Engine/Source/Runtime/CoreUObject/Public/UObject/Object.h
@@ -840,11 +840,12 @@ public:
 		FString Value;
 
 		/** Broad description of kind of data represented in Value */
-		ETagType Type;
+		ETagType Type = TT_Alphabetical;
 
 		/** Flags describing more detail for displaying in the UI */
-		uint32 DisplayFlags;
+		uint32 DisplayFlags = TD_None;
 
+		FAssetRegistryTag() = default;
 		FAssetRegistryTag(FName InName, const FString& InValue, ETagType InType, uint32 InDisplayFlags = TD_None)
 			: Name(InName), Value(InValue), Type(InType), DisplayFlags(InDisplayFlags) {}
 		FAssetRegistryTag(FName InName, FString&& InValue, ETagType InType, uint32 InDisplayFlags = TD_None)
@@ -869,9 +870,6 @@ public:
 	 */
 	COREUOBJECT_API virtual void GetAssetRegistryTags(TArray<FAssetRegistryTag>& OutTags) const;
 
-	UE_DEPRECATED(5.1, "Use the new GetExtendedAssetRegistryTagsForSave that takes a TargetPlatform")
-	virtual void GetExternalActorExtendedAssetRegistryTags(TArray<FAssetRegistryTag>& OutTags) const {}
-
 #if WITH_EDITOR
 	/**
 	 * Temporary interim solution to gather asset registry data at save time only.  Can depend on the target platform being saved for.
diff --git a/Engine/Source/Runtime/CoreUObject/Public/UObject/ObjectSaveContext.h b/Engine/Source/Runtime/CoreUObject/Public/UObject/ObjectSaveContext.h
--- a/Engine/Source/Runtime/CoreUObject/Public/UObject/ObjectSaveContext.h
+++ b/Engine/Source/Runtime/CoreUObject/Public/UObject/ObjectSaveContext.h
@@ -4,6 +4,7 @@
 
 #include "Containers/UnrealString.h"
 #include "HAL/Platform.h"
+#include "UObject/CookEnums.h"
 #include "UObject/ObjectSaveOverride.h"
 
 class FPackagePath;
@@ -41,6 +42,9 @@ struct FObjectSaveContextData
 	/** Package->GetPackageFlags before the save, or 0 if no package. */
 	uint32 OriginalPackageFlags = 0;
 
+	UE::Cook::ECookType CookType = UE::Cook::ECookType::Unknown;
+	UE::Cook::ECookingDLC CookingDLC = UE::Cook::ECookingDLC::Unknown;
+
 	/**
 	 * Set to true when the package is being saved due to a procedural save.
 	 * Any save without the possibility of user-generated edits to the package is a procedural save (Cooking, EditorDomain).
@@ -110,6 +114,12 @@ public:
 	/** Return the targetplatform of the save, if cooking. Null if not cooking. */
 	const ITargetPlatform* GetTargetPlatform() const { return Data.TargetPlatform; }
 
+	bool IsCookByTheBook() const { return GetCookType()  == UE::Cook::ECookType::ByTheBook; }
+	bool IsCookOnTheFly() const { return GetCookType() == UE::Cook::ECookType::OnTheFly; }
+	bool IsCookTypeUnknown() const { return GetCookType() == UE::Cook::ECookType::Unknown; }
+	UE::Cook::ECookType GetCookType() const { return Data.CookType; }
+	UE::Cook::ECookingDLC GetCookingDLC() const { return Data.CookingDLC; }
+
 	/**
 	 * Return whether the package is being saved due to a procedural save.
 	 * Any save without the possibility of user-generated edits to the package is a procedural save (Cooking, EditorDomain).
@@ -175,6 +185,12 @@ public:
 	/** Return the targetplatform of the save, if cooking. Null if not cooking. */
 	const ITargetPlatform* GetTargetPlatform() const { return Data.TargetPlatform; }
 
+	bool IsCookByTheBook() const { return GetCookType() == UE::Cook::ECookType::ByTheBook; }
+	bool IsCookOnTheFly() const { return GetCookType() == UE::Cook::ECookType::OnTheFly; }
+	bool IsCookTypeUnknown() const { return GetCookType() == UE::Cook::ECookType::Unknown; }
+	UE::Cook::ECookType GetCookType() const { return Data.CookType; }
+	UE::Cook::ECookingDLC GetCookingDLC() const { return Data.CookingDLC; }
+
 	/**
 	 * Return whether the package is being saved due to a procedural save.
 	 * Any save without the possibility of user-generated edits to the package is a procedural save (Cooking, EditorDomain).
diff --git a/Engine/Source/Runtime/Engine/Private/BlueprintGeneratedClass.cpp b/Engine/Source/Runtime/Engine/Private/BlueprintGeneratedClass.cpp
--- a/Engine/Source/Runtime/Engine/Private/BlueprintGeneratedClass.cpp
+++ b/Engine/Source/Runtime/Engine/Private/BlueprintGeneratedClass.cpp
@@ -399,7 +399,7 @@ void UBlueprintGeneratedClass::GetAssetRegistryTags(TArray<FAssetRegistryTag>& O
 	 *
 	 * See UBlueprint::GetExtendedAssetRegistryTagsForSave.
 	 */
-	if (GIsSavingPackage)
+	if (GIsSavingPackage && !IsRunningCookCommandlet())
 	{
 		if (AActor* BlueprintCDO = Cast<AActor>(ClassDefaultObject))
 		{
diff --git a/Engine/Source/Runtime/Engine/Private/Materials/MaterialInterface.cpp b/Engine/Source/Runtime/Engine/Private/Materials/MaterialInterface.cpp
--- a/Engine/Source/Runtime/Engine/Private/Materials/MaterialInterface.cpp
+++ b/Engine/Source/Runtime/Engine/Private/Materials/MaterialInterface.cpp
@@ -198,7 +198,7 @@ void UMaterialInterface::Serialize(FArchive& Ar)
 	{
 		// Mark whether this material is part of the base game. This provides information on whether it is safe to be
 		// used as a parent safely in child modules.
-		bIncludedInBaseGame = Ar.GetCookContext()->GetCookingDLC() == FArchiveCookContext::ECookingDLCNo;
+		bIncludedInBaseGame = Ar.GetCookContext()->GetCookingDLC() == UE::Cook::ECookingDLC::No;
 	}
 
 	Super::Serialize(Ar);
diff --git a/Engine/Source/Runtime/Engine/Private/ShaderCompiler/ShaderCompiler.cpp b/Engine/Source/Runtime/Engine/Private/ShaderCompiler/ShaderCompiler.cpp
--- a/Engine/Source/Runtime/Engine/Private/ShaderCompiler/ShaderCompiler.cpp
+++ b/Engine/Source/Runtime/Engine/Private/ShaderCompiler/ShaderCompiler.cpp
@@ -9124,17 +9124,20 @@ FString SaveGlobalShaderFile(EShaderPlatform Platform, FString SavePath, class I
 
 	TArray<uint8> GlobalShaderData;
 	{
+#if WITH_EDITOR
+		TOptional<FArchiveCookContext> CookContext;
+		TOptional<FArchiveCookData> CookData;
+#endif
 		FMemoryWriter MemoryWriter(GlobalShaderData, true);
 
 #if WITH_EDITOR
-		TOptional<FArchiveCookData> CookData;
-		FArchiveCookContext CookContext(nullptr /*InPackage*/, FArchiveCookContext::ECookTypeUnknown, FArchiveCookContext::ECookingDLCUnknown);
 		if (TargetPlatform != nullptr)
 		{
-			CookData.Emplace(*TargetPlatform, CookContext);
+			CookContext.Emplace(nullptr /*InPackage*/, UE::Cook::ECookType::Unknown,
+				UE::Cook::ECookingDLC::Unknown);
+			CookData.Emplace(*TargetPlatform, *CookContext);
+			MemoryWriter.SetCookData(CookData.GetPtrOrNull());
 		}
-		
-		MemoryWriter.SetCookData(CookData.GetPtrOrNull());
 #endif // WITH_EDITOR
 
 		GlobalShaderMap->SaveToGlobalArchive(MemoryWriter);
@@ -9749,16 +9752,18 @@ void CompileGlobalShaderMapForRemote(
 	FinishRecompileGlobalShaders();
 
 	// Write the shader compilation info to memory, converting FName to strings
+	TOptional<FArchiveCookContext> CookContext;
+	TOptional<FArchiveCookData> CookData;
 	FMemoryWriter MemWriter(*OutArray, true);
 	FNameAsStringProxyArchive Ar(MemWriter);
 
-	TOptional<FArchiveCookData> CookData;
-	FArchiveCookContext CookContext(nullptr /*InPackage*/, FArchiveCookContext::ECookTypeUnknown, FArchiveCookContext::ECookingDLCUnknown);
 	if (TargetPlatform != nullptr)
 	{
-		CookData.Emplace(*TargetPlatform, CookContext);
+		CookContext.Emplace(nullptr /*InPackage*/, UE::Cook::ECookType::Unknown,
+			UE::Cook::ECookingDLC::Unknown);
+		CookData.Emplace(*TargetPlatform, *CookContext);
+		Ar.SetCookData(CookData.GetPtrOrNull());
 	}
-	Ar.SetCookData(CookData.GetPtrOrNull());
 
 	// save out the global shader map to the byte array
 	SaveGlobalShadersForRemoteRecompile(Ar, ShaderPlatform);
@@ -9767,18 +9772,19 @@ void CompileGlobalShaderMapForRemote(
 void SaveShaderMapsForRemote(ITargetPlatform* TargetPlatform, const TMap<FString, TArray<TRefCountPtr<FMaterialShaderMap>>>& CompiledShaderMaps, TArray<uint8>* OutArray)
 {
 	// write the shader compilation info to memory, converting fnames to strings
+	TOptional<FArchiveCookContext> CookContext;
+	TOptional<FArchiveCookData> CookData;
 	FMemoryWriter MemWriter(*OutArray, true);
 	FNameAsStringProxyArchive Ar(MemWriter);
 
-	TOptional<FArchiveCookData> CookData;
-	FArchiveCookContext CookContext(nullptr /*InPackage*/, FArchiveCookContext::ECookTypeUnknown, FArchiveCookContext::ECookingDLCUnknown);
 	if (TargetPlatform != nullptr)
 	{
-		CookData.Emplace(*TargetPlatform, CookContext);
+		CookContext.Emplace(nullptr /*InPackage*/, UE::Cook::ECookType::Unknown,
+			UE::Cook::ECookingDLC::Unknown);
+		CookData.Emplace(*TargetPlatform, *CookContext);
+		Ar.SetCookData(CookData.GetPtrOrNull());
 	}
 
-	Ar.SetCookData(CookData.GetPtrOrNull());
-
 	// save out the shader map to the byte array
 	FMaterialShaderMap::SaveForRemoteRecompile(Ar, CompiledShaderMaps);
 }
