AssetRegistry: Optimize cooker's WaitForCompletion call by removing the busy-wait loop in WaitForCompletion that causes critical section contention with the gatherer's ResultsLock.
Remove the spurious logging of AssetRegistryGather time on every SearchAllAssets call; the logging should only occur for the initial SearchAllAssets, and it will be done for that one (synchronous or not) in TickGatherer.
Remove time spent sleeping waiting for discovery from the AssetRegistryGather time metric.
#rnx

diff --git a/Engine/Source/Runtime/AssetRegistry/Private/AssetDataGatherer.cpp b/Engine/Source/Runtime/AssetRegistry/Private/AssetDataGatherer.cpp
--- a/Engine/Source/Runtime/AssetRegistry/Private/AssetDataGatherer.cpp
+++ b/Engine/Source/Runtime/AssetRegistry/Private/AssetDataGatherer.cpp
@@ -1815,7 +1815,7 @@ void FAssetDataDiscovery::GetDiagnostics(float& OutCumulativeDiscoveryTime)
 }
 
 
-void FAssetDataDiscovery::WaitForIdle()
+void FAssetDataDiscovery::WaitForIdle(double EndTimeSeconds)
 {
 	if (bIsIdle)
 	{
@@ -1838,7 +1838,19 @@ void FAssetDataDiscovery::WaitForIdle()
 		}
 		else
 		{
-			FPlatformProcess::Sleep(IdleSleepTime);
+			float SleepTime = IdleSleepTime;
+			if (EndTimeSeconds > 0)
+			{
+				SleepTime = FMath::Min(SleepTime, static_cast<float>(EndTimeSeconds - FPlatformTime::Seconds()));
+			}
+			if (SleepTime > 0.f)
+			{
+				FPlatformProcess::Sleep(SleepTime);
+			}
+		}
+		if (EndTimeSeconds > 0 && FPlatformTime::Seconds() > EndTimeSeconds)
+		{
+			break;
 		}
 	}
 	if (bTickOwner)
@@ -1847,6 +1859,11 @@ void FAssetDataDiscovery::WaitForIdle()
 	}
 }
 
+bool FAssetDataDiscovery::IsIdle() const
+{
+	return bIsIdle;
+}
+
 FPathExistence::FPathExistence(FStringView InLocalAbsPath)
 	:LocalAbsPath(InLocalAbsPath)
 {
@@ -3211,6 +3228,8 @@ FAssetDataGatherer::FAssetDataGatherer(const TArray<FString>& InLongPackageNames
 {
 	using namespace UE::AssetDataGather::Private;
 
+	TickInternalBatchSize = FMath::Max(1, FTaskGraphInterface::Get().GetNumWorkerThreads()) * AssetDataGathererConstants::SingleThreadFilesPerBatch;
+
 	GPreloadSettings.Initialize();
 	bGatherAssetPackageData = GIsEditor || GPreloadSettings.IsForceDependsGathering();
 	bGatherDependsData = GPreloadSettings.IsGatherDependsData();
@@ -3297,30 +3316,33 @@ bool FAssetDataGatherer::Init()
 uint32 FAssetDataGatherer::Run()
 {
 	constexpr float IdleSleepTime = 0.1f;
+	constexpr float PausedSleepTime = 0.005f;
 	LLM_SCOPE(ELLMTag::AssetRegistry);
 
 	while (!IsStopped)
 	{
-		InnerTickLoop(false /* bInSynchronousTick */, true /* bContributeToCacheSave */);
+		InnerTickLoop(false /* bInSynchronousTick */, true /* bContributeToCacheSave */, -1. /* EndTimeSeconds */);
 
 		for (;;)
 		{
+			bool bLocalIdle = false;
 			{
 				FGathererScopeLock ResultsScopeLock(&ResultsLock); // bIsIdle requires the lock
 				if (IsStopped || bSaveAsyncCacheTriggered || (!IsPaused && !bIsIdle))
 				{
 					break;
 				}
+				bLocalIdle = bIsIdle;
 			}
 			// No work to do. Sleep for a little and try again later.
 			// TODO: Need IsPaused to be a condition variable so we avoid sleeping while waiting for it and then taking a long time to wake after it is unset.
-			FPlatformProcess::Sleep(IdleSleepTime);
+			FPlatformProcess::Sleep(bLocalIdle ? IdleSleepTime : PausedSleepTime);
 		}
 	}
 	return 0;
 }
 
-void FAssetDataGatherer::InnerTickLoop(bool bInSynchronousTick, bool bContributeToCacheSave)
+void FAssetDataGatherer::InnerTickLoop(bool bInSynchronousTick, bool bContributeToCacheSave, double EndTimeSeconds)
 {
 	using namespace UE::AssetDataGather::Private;
 
@@ -3335,11 +3357,61 @@ void FAssetDataGatherer::InnerTickLoop(bool bInSynchronousTick, bool bContribute
 		FGathererScopeLock RunScopeLock(&TickLock);
 		TGuardValue<bool> ScopeSynchronousTick(bSynchronousTick, bInSynchronousTick);
 		TRACE_CPUPROFILER_EVENT_SCOPE(FAssetDataGatherer::Tick);
-		bool bTickInterruptionEvent = false;
 		double TickStartTime = FPlatformTime::Seconds();
-		while (!IsStopped && (bInSynchronousTick || !IsPaused) && !bTickInterruptionEvent)
+		bool bPollDiscovery = true;
+		constexpr float PollDiscoveryPeriodSeconds = .005f;
+		double LastPollTimeSeconds = 0;
+
+		for (;;)
 		{
-			TickInternal(bTickInterruptionEvent, TickStartTime);
+			ETickResult TickResult = TickInternal(TickStartTime, bPollDiscovery);
+			if (EndTimeSeconds > 0. && FPlatformTime::Seconds() > EndTimeSeconds)
+			{
+				break;
+			}
+			if (IsStopped || (!bInSynchronousTick && IsPaused))
+			{
+				break;
+			}
+			if (TickResult != ETickResult::KeepTicking && TickResult != ETickResult::PollDiscovery)
+			{
+				break;
+			}
+			double CurrentTimeSeconds = FPlatformTime::Seconds();
+			if (bPollDiscovery)
+			{
+				LastPollTimeSeconds = CurrentTimeSeconds;
+			}
+			if (TickResult == ETickResult::KeepTicking)
+			{
+				// Poll discovery every so often to reduce super-linear costs
+				bPollDiscovery =
+					static_cast<float>(CurrentTimeSeconds - LastPollTimeSeconds) > PollDiscoveryPeriodSeconds;
+			}
+			else if (!bPollDiscovery)
+			{
+				bPollDiscovery = true;
+			}
+			else
+			{
+				// We just polled discovery; go to sleep to wait for the discovery thread rather than busy-spinning and
+				// causing contention on its critical section
+				float SleepTimeSeconds = PollDiscoveryPeriodSeconds;
+				if (EndTimeSeconds > 0.)
+				{
+					SleepTimeSeconds = FMath::Min(SleepTimeSeconds,
+						static_cast<float>(EndTimeSeconds - CurrentTimeSeconds));
+				}
+				if (SleepTimeSeconds > 0.f)
+				{
+					if (TickStartTime >= 0.)
+					{
+						CurrentSearchTime += CurrentTimeSeconds - TickStartTime;
+					}
+					FPlatformProcess::Sleep(SleepTimeSeconds);
+					TickStartTime = FPlatformTime::Seconds();
+				}
+			}
 		}
 		if (TickStartTime >= 0.)
 		{
@@ -3505,13 +3577,10 @@ void FAssetDataGatherer::EnsureCompletion()
 	}
 }
 
-void FAssetDataGatherer::TickInternal(bool& bOutIsTickInterrupt, double& TickStartTime)
+FAssetDataGatherer::ETickResult FAssetDataGatherer::TickInternal(double& TickStartTime, bool bPollDiscovery)
 {
-	LLM_SCOPE(ELLMTag::AssetRegistry);
-
 	using namespace UE::AssetDataGather::Private;
 
-	const int32 BatchSize = FMath::Max(1, FTaskGraphInterface::Get().GetNumWorkerThreads()) * AssetDataGathererConstants::SingleThreadFilesPerBatch;
 	typedef TInlineAllocator<AssetDataGathererConstants::ExpectedMaxBatchSize> FBatchInlineAllocator;
 
 	TArray<FGatheredPathData, FBatchInlineAllocator> LocalFilesToSearch;
@@ -3524,7 +3593,7 @@ void FAssetDataGatherer::TickInternal(bool& bOutIsTickInterrupt, double& TickSta
 	bool bLocalIsCacheWriteEnabled = false;
 	double LocalLastCacheWriteTime = 0.0;
 	bool bWaitBatchCountDecremented = false;
-	bOutIsTickInterrupt = false;
+	ETickResult TickResult = ETickResult::KeepTicking;
 
 	{
 		FGathererScopeLock ResultsScopeLock(&ResultsLock);
@@ -3535,7 +3604,10 @@ void FAssetDataGatherer::TickInternal(bool& bOutIsTickInterrupt, double& TickSta
 			LastCacheWriteTime = FPlatformTime::Seconds();
 		}
 
-		IngestDiscoveryResults();
+		if (bPollDiscovery)
+		{
+			IngestDiscoveryResults();
+		}
 		if (bInitialPluginsLoaded && !bFlushedRetryFiles)
 		{
 			bFlushedRetryFiles = true;
@@ -3543,19 +3615,22 @@ void FAssetDataGatherer::TickInternal(bool& bOutIsTickInterrupt, double& TickSta
 		}
 
 		// Take a batch off of the work list. If we're waiting only on the first WaitBatchCount results don't take more than that
-		int32 NumToProcess = FMath::Min<int32>(BatchSize-LocalFilesToSearch.Num(), FilesToSearch->GetNumAvailable());
+		int32 NumToProcess = FMath::Min<int32>(TickInternalBatchSize - LocalFilesToSearch.Num(), FilesToSearch->GetNumAvailable());
 		// If no work is available mark idle and exit
 		if (NumToProcess == 0)
 		{
 			if (WaitBatchCount != -1)
 			{
 				WaitBatchCount = -1; // WaitBatchCount was set higher than FilesToSearch->GetNumAvailable(), mark it completed
-				bOutIsTickInterrupt = true;
+				TickResult = ETickResult::Interrupt;
 			}
 
 			if (bDiscoveryIsComplete)
 			{
-				bOutIsTickInterrupt = true;
+				if (TickResult == ETickResult::KeepTicking)
+				{
+					TickResult = ETickResult::Idle;
+				}
 				const bool bWasInitialDiscoveryFinished = bFinishedInitialDiscovery;
 				SetIsIdle(true, TickStartTime);
 				if (!bWasInitialDiscoveryFinished && bFinishedInitialDiscovery)
@@ -3564,7 +3639,12 @@ void FAssetDataGatherer::TickInternal(bool& bOutIsTickInterrupt, double& TickSta
 					bSaveAsyncCacheTriggered = true;
 				}
 			}
-			return;
+			else if (TickResult == ETickResult::KeepTicking)
+			{
+				TickResult = ETickResult::PollDiscovery;
+			}
+
+			return TickResult;
 		}
 
 		if (WaitBatchCount >= 0)
@@ -3574,8 +3654,7 @@ void FAssetDataGatherer::TickInternal(bool& bOutIsTickInterrupt, double& TickSta
 				// We've finished executing the caller's requested batchcount (and we have restored idle if we
 				// are idle), so exit now without doing any further work.
 				WaitBatchCount = -1;
-				bOutIsTickInterrupt = true;
-				return;
+				return ETickResult::Interrupt;
 			}
 
 			// Otherwise we still have some work to do for the caller's requested batchcount, so do work up to
@@ -3587,7 +3666,7 @@ void FAssetDataGatherer::TickInternal(bool& bOutIsTickInterrupt, double& TickSta
 			{
 				// Mark that WaitBatchCount has been consumed and is no longer active
 				WaitBatchCount = -1;
-				bOutIsTickInterrupt = true;
+				TickResult = ETickResult::Interrupt;
 			}
 		}
 		DependencyResults.Reserve(FilesToSearch->GetNumAvailable() + DependencyResults.Num());
@@ -3844,9 +3923,10 @@ void FAssetDataGatherer::TickInternal(bool& bOutIsTickInterrupt, double& TickSta
 			&& NumAssetsReadSinceLastCacheWrite >= AssetDataGathererConstants::MinAssetReadsBeforeCacheWrite)
 		{
 			bSaveAsyncCacheTriggered = true;
-			bOutIsTickInterrupt = true;
+			TickResult = ETickResult::Interrupt;
 		}
 	}
+	return TickResult;
 }
 
 void FAssetDataGatherer::IngestDiscoveryResults()
@@ -4122,6 +4202,8 @@ void FAssetDataGatherer::ScanPathsSynchronous(const TArray<FString>& InLocalPath
 void FAssetDataGatherer::WaitOnPathsInternal(TArrayView<UE::AssetDataGather::Private::FPathExistence> QueryPaths,
 	const FString& SaveCacheFilename, const TArray<FString>& SaveCacheLongPackageNameDirs)
 {
+	LLM_SCOPE(ELLMTag::AssetRegistry);
+
 	// Request a halt to the async tick
 	FScopedPause ScopedPause(*this);
 	CHECK_IS_NOT_LOCKED_CURRENT_THREAD(ResultsLock);
@@ -4148,7 +4230,7 @@ void FAssetDataGatherer::WaitOnPathsInternal(TArrayView<UE::AssetDataGather::Pri
 	TRACE_CPUPROFILER_EVENT_SCOPE(FAssetDataGatherer::Tick);
 	for (;;)
 	{
-		InnerTickLoop(true /* bInSynchronousTick */, bContributeToCacheSave);
+		InnerTickLoop(true /* bInSynchronousTick */, bContributeToCacheSave, -1. /* EndTimeSeconds */);
 		FGathererScopeLock ResultsScopeLock(&ResultsLock); // WaitBatchCount requires the lock
 		if (WaitBatchCount < 0)
 		{
@@ -4173,7 +4255,7 @@ void FAssetDataGatherer::WaitOnPathsInternal(TArrayView<UE::AssetDataGather::Pri
 	}
 }
 
-void FAssetDataGatherer::WaitForIdle()
+void FAssetDataGatherer::WaitForIdle(float TimeoutSeconds)
 {
 	{
 		FGathererScopeLock ResultsScopeLock(&ResultsLock);
@@ -4183,7 +4265,22 @@ void FAssetDataGatherer::WaitForIdle()
 		}
 	}
 	TRACE_CPUPROFILER_EVENT_SCOPE_STR("FAssetDataGatherer::WaitForIdle");
-	Discovery->WaitForIdle();
+	LLM_SCOPE(ELLMTag::AssetRegistry);
+
+	double EndTimeSeconds = -1.;
+	if (TimeoutSeconds >= 0.0f)
+	{
+		EndTimeSeconds = FPlatformTime::Seconds() + TimeoutSeconds;
+	}
+	if (Discovery->IsSynchronous())
+	{
+		Discovery->WaitForIdle(EndTimeSeconds);
+		if (EndTimeSeconds > 0. && FPlatformTime::Seconds() > EndTimeSeconds)
+		{
+			return;
+		}
+	}
+
 	CHECK_IS_NOT_LOCKED_CURRENT_THREAD(ResultsLock);
 
 	// Request a halt to the async tick
@@ -4191,7 +4288,11 @@ void FAssetDataGatherer::WaitForIdle()
 	// Tick until idle
 	for (;;)
 	{
-		InnerTickLoop(true /* bInSynchronousTick */, true /* bContributeToCacheSave */);
+		InnerTickLoop(true /* bInSynchronousTick */, true /* bContributeToCacheSave */, EndTimeSeconds);
+		if (EndTimeSeconds > 0 && FPlatformTime::Seconds() > EndTimeSeconds)
+		{
+			break;
+		}
 		FGathererScopeLock ResultsScopeLock(&ResultsLock); // bIsIdle requires the lock
 		if (bIsIdle)
 		{
diff --git a/Engine/Source/Runtime/AssetRegistry/Private/AssetDataGatherer.h b/Engine/Source/Runtime/AssetRegistry/Private/AssetDataGatherer.h
--- a/Engine/Source/Runtime/AssetRegistry/Private/AssetDataGatherer.h
+++ b/Engine/Source/Runtime/AssetRegistry/Private/AssetDataGatherer.h
@@ -161,7 +161,7 @@ public:
 	void ScanPathsSynchronous(const TArray<FString>& InPaths, bool bForceRescan, bool bIgnoreDenyListScanFilters,
 		const FString& SaveCacheFilename, const TArray<FString>& SaveCacheLongPackageNameDirs);
 	/** Wait for all monitored assets to be added to search results. */
-	void WaitForIdle();
+	void WaitForIdle(float TimeoutSeconds = -1.0f);
 	/**
 	 * Report whether all monitored assets have been added to search results, AND these results have been gathered
 	 * through a GetAndTrimSearchResults call.
@@ -240,17 +240,23 @@ public:
 		FPackageReader::EReadOptions Options);
 
 private:
-
+	enum class ETickResult
+	{
+		KeepTicking,
+		PollDiscovery,
+		Idle,
+		Interrupt,
+	};
 	/**
 	 * Helper function to run the tick in a loop-within-a-loop to minimize critical section entry, and to move expensive
 	 * operations out of the critical section
 	 */
-	void InnerTickLoop(bool bInSynchronousTick, bool bContributeToCacheSave);
+	void InnerTickLoop(bool bInSynchronousTick, bool bContributeToCacheSave, double EndTimeSeconds);
 	/**
 	 * Tick function to pump scanning and push results into the search results structure. May be called from devoted
 	 * thread or inline from synchronous functions on other threads.
 	 */
-	void TickInternal(bool& bOutIsIdle, double& TickStartTime);
+	ETickResult TickInternal(double& TickStartTime, bool bPollDiscovery);
 	/** Add any new package files from the background directory scan to our work list **/
 	void IngestDiscoveryResults();
 
@@ -370,6 +376,8 @@ private:
 	 * If null, results will only be added when Wait functions are called. Constant during threading.
 	 */
 	FRunnableThread* Thread;
+	int32 TickInternalBatchSize = 1;
+
 	/**
 	 * True if async gathering is enabled, false if e.g. singlethreaded or disabled by commandline.
  	 * Even when enabled, gathering is still synchronous until StartAsync is called.
diff --git a/Engine/Source/Runtime/AssetRegistry/Private/AssetDataGathererPrivate.h b/Engine/Source/Runtime/AssetRegistry/Private/AssetDataGathererPrivate.h
--- a/Engine/Source/Runtime/AssetRegistry/Private/AssetDataGathererPrivate.h
+++ b/Engine/Source/Runtime/AssetRegistry/Private/AssetDataGathererPrivate.h
@@ -604,7 +604,9 @@ public:
 	/** Get diagnostics for telemetry or logging. */
 	void GetDiagnostics(float& OutCumulativeDiscoveryTime);
 	/** Wait (joining in on the tick) until all currently monitored paths have been scanned. */
-	void WaitForIdle();
+	void WaitForIdle(double EndTimeSeconds);
+	bool IsIdle() const;
+
 	/** Optionally set some scan properties for the given paths and then wait for their scans to finish. */
 	void SetPropertiesAndWait(TArrayView<FPathExistence> QueryPaths, bool bAddToAllowList, bool bForceRescan,
 		bool bIgnoreDenyListScanFilters);
diff --git a/Engine/Source/Runtime/AssetRegistry/Private/AssetRegistry.cpp b/Engine/Source/Runtime/AssetRegistry/Private/AssetRegistry.cpp
--- a/Engine/Source/Runtime/AssetRegistry/Private/AssetRegistry.cpp
+++ b/Engine/Source/Runtime/AssetRegistry/Private/AssetRegistry.cpp
@@ -1680,10 +1680,6 @@ void UAssetRegistryImpl::SearchAllAssets(bool bSynchronousSearch)
 			GuardedData.ConditionalLoadPremadeAssetRegistry(*this, EventContext, InterfaceScopeLock);
 		}
 		GuardedData.SearchAllAssets(EventContext, InheritanceContext, bSynchronousSearch);
-		if (bSynchronousSearch)
-		{
-			GuardedData.LogSearchDiagnostics(StartTime);
-		}
 	}
 #if WITH_EDITOR
 	if (bSynchronousSearch)
@@ -1763,6 +1759,19 @@ void UAssetRegistryImpl::WaitForCompletion()
 
 	using namespace UE::AssetRegistry::Impl;
 
+	// Try taking over the gather thread for a short time in case it is mostly done.
+	// But if it has more than a small amount of work to do, let the gather thread do that work
+	// while we consume the results in parallel.
+	{
+		LLM_SCOPE(ELLMTag::AssetRegistry);
+		FWriteScopeLock InterfaceScopeLock(InterfaceLock);
+		FClassInheritanceContext InheritanceContext;
+		FClassInheritanceBuffer InheritanceBuffer;
+		GetInheritanceContextWithRequiredLock(InterfaceScopeLock, InheritanceContext, InheritanceBuffer);
+		constexpr float TimeToJoinSeconds = 0.100f;
+		GuardedData.WaitForGathererIdle(TimeToJoinSeconds);
+	}
+
 	for (;;)
 	{
 		FEventContext EventContext;
@@ -1789,7 +1798,9 @@ void UAssetRegistryImpl::WaitForCompletion()
 		}
 
 		FThreadHeartBeat::Get().HeartBeat();
-		FPlatformProcess::SleepNoStats(0.0001f);
+		// Sleep long enough to avoid causing contention on the CriticalSection in GetAndTrimSearchResults
+		constexpr float SleepTimeSeconds = 0.010f;
+		FPlatformProcess::SleepNoStats(SleepTimeSeconds);
 	}
 }
 
@@ -4156,6 +4167,14 @@ void FAssetRegistryImpl::WaitForGathererIdleIfSynchronous()
 	}
 }
 
+void FAssetRegistryImpl::WaitForGathererIdle(float TimeoutSeconds)
+{
+	if (GlobalGatherer)
+	{
+		GlobalGatherer->WaitForIdle(TimeoutSeconds);
+	}
+}
+
 Impl::EGatherStatus FAssetRegistryImpl::TickGatherer(Impl::FEventContext& EventContext,
 	Impl::FClassInheritanceContext& InheritanceContext, const double TickStartTime, bool& bOutInterrupted,
 	TOptional<FAssetsFoundCallback> AssetsFoundCallback)
diff --git a/Engine/Source/Runtime/AssetRegistry/Private/AssetRegistryImpl.h b/Engine/Source/Runtime/AssetRegistry/Private/AssetRegistryImpl.h
--- a/Engine/Source/Runtime/AssetRegistry/Private/AssetRegistryImpl.h
+++ b/Engine/Source/Runtime/AssetRegistry/Private/AssetRegistryImpl.h
@@ -189,6 +189,8 @@ public:
 
 	/** Waits for the gatherer to be idle if it is operating synchronously. */
 	void WaitForGathererIdleIfSynchronous();
+	/** Waits for the gatherer to be idle. */
+	void WaitForGathererIdle(float TimeoutSeconds);
 	/** Callback type for TickGatherer */
 	typedef TFunctionRef<void(const TMultiMap<FName, FAssetData*>&)> FAssetsFoundCallback;
 	/** Consume any results from the gatherer and return its status */
