AssetRegistry: Respect value of LinkerLoad's s.EnforcePackageCompatibleVersionCheck and allow packages to be listed in the AssetRegistry even if they have invalid versions, if the check is disabled.
#rnx

diff --git a/Engine/Source/Runtime/AssetRegistry/Private/PackageReader.cpp b/Engine/Source/Runtime/AssetRegistry/Private/PackageReader.cpp
--- a/Engine/Source/Runtime/AssetRegistry/Private/PackageReader.cpp
+++ b/Engine/Source/Runtime/AssetRegistry/Private/PackageReader.cpp
@@ -15,6 +15,7 @@
 #include "Misc/ScopeExit.h"
 #include "UObject/Class.h"
 #include "UObject/Linker.h"
+#include "UObject/LinkerLoad.h"
 #include "UObject/PackageRelocation.h"
 #include "UObject/PackageTrailer.h"
 
@@ -115,37 +116,46 @@ bool FPackageReader::OpenPackageFile(EOpenPackageResult& OutErrorCode)
 		return false;
 	}
 
-	if (!PackageFileSummary.IsFileVersionValid())
+	// IsEnforcePackageCompatibleVersionCheck(): If LinkerLoad is not validating, PackageReader should not either.
+	// Optimize the IsEnforcePackageCompatibleVersionCheck==true but no errors case; only test
+	// IsEnforcePackageCompatibleVersionCheck after finding a version mismatch.
+	if (!PackageFileSummary.IsFileVersionValid() &&
+		IsEnforcePackageCompatibleVersionCheck())
 	{
 		// Log a warning rather than an error. Linkerload gracefully handles this case.
-		UE_LOG(LogAssetRegistry, Warning, TEXT("Package %s is unversioned which cannot be opened by the current process"), *PackageFilename);
+		UE_LOG(LogAssetRegistry, Warning,
+			TEXT("Package %s is unversioned which cannot be opened by the current process"), *PackageFilename);
 		OutErrorCode = EOpenPackageResult::Unversioned;
 		return false;
 	}
 
 	// Don't read packages that are too old
-	if (PackageFileSummary.IsFileVersionTooOld())
+	if (PackageFileSummary.IsFileVersionTooOld() &&
+		IsEnforcePackageCompatibleVersionCheck())
 	{
 		// Log a warning rather than an error. Linkerload gracefully handles this case.
-		UE_LOG(	LogAssetRegistry, Warning, TEXT("Package %s is too old. Min Version: %i  Package Version: %i"),
-				*PackageFilename, (int32)VER_UE4_OLDEST_LOADABLE_PACKAGE, PackageFileSummary.GetFileVersionUE().FileVersionUE4);
+		UE_LOG(LogAssetRegistry, Warning, TEXT("Package %s is too old. Min Version: %i  Package Version: %i"),
+			*PackageFilename, (int32)VER_UE4_OLDEST_LOADABLE_PACKAGE,
+			PackageFileSummary.GetFileVersionUE().FileVersionUE4);
 
-		OutErrorCode = EOpenPackageResult::Unversioned;
+		OutErrorCode = EOpenPackageResult::VersionTooOld;
 		return false;
 	}
 
 	// Don't read packages that were saved with a package version newer than the current one.
-	if (PackageFileSummary.IsFileVersionTooNew())
+	if (PackageFileSummary.IsFileVersionTooNew() &&
+		IsEnforcePackageCompatibleVersionCheck())
 	{
 		// Log a warning rather than an error. Linkerload gracefully handles this case.
-		UE_LOG(	LogAssetRegistry, Warning, TEXT("Package %s is too new. Engine Version: %i  Package Version: %i"),
-				*PackageFilename, GPackageFileUEVersion.ToValue(), PackageFileSummary.GetFileVersionUE().ToValue());
+		UE_LOG(LogAssetRegistry, Warning, TEXT("Package %s is too new. Engine Version: %i  Package Version: %i"),
+			*PackageFilename, GPackageFileUEVersion.ToValue(), PackageFileSummary.GetFileVersionUE().ToValue());
 
 		OutErrorCode = EOpenPackageResult::VersionTooNew;
 		return false;
 	}
 
-	if (PackageFileSummary.GetFileVersionLicenseeUE() > GPackageFileLicenseeUEVersion)
+	if (PackageFileSummary.GetFileVersionLicenseeUE() > GPackageFileLicenseeUEVersion &&
+		IsEnforcePackageCompatibleVersionCheck())
 	{
 		// Log a warning rather than an error. Linkerload gracefully handles this case.
 		UE_LOG(LogAssetRegistry, Warning, TEXT("Package %s is too new. Licensee Version: %i Package Licensee Version: %i"),
@@ -156,23 +166,34 @@ bool FPackageReader::OpenPackageFile(EOpenPackageResult& OutErrorCode)
 	}
 
 	// Check serialized custom versions against latest custom versions.
-	TArray<FCustomVersionDifference> Diffs = FCurrentCustomVersions::Compare(PackageFileSummary.GetCustomVersionContainer().GetAllVersions(), *PackageFilename);
+	TArray<FCustomVersionDifference> Diffs = FCurrentCustomVersions::Compare(
+		PackageFileSummary.GetCustomVersionContainer().GetAllVersions(), *PackageFilename);
 	for (FCustomVersionDifference Diff : Diffs)
 	{
 		if (Diff.Type == ECustomVersionDifference::Missing)
 		{
-			OutErrorCode = EOpenPackageResult::CustomVersionMissing;
+			if (IsEnforcePackageCompatibleVersionCheck())
+			{
+				OutErrorCode = EOpenPackageResult::CustomVersionMissing;
+			}
 		}
 		else if (Diff.Type == ECustomVersionDifference::Invalid)
 		{
-			OutErrorCode = EOpenPackageResult::CustomVersionInvalid;
+			if (IsEnforcePackageCompatibleVersionCheck())
+			{
+				OutErrorCode = EOpenPackageResult::CustomVersionInvalid;
+			}
 		}
 		else if (Diff.Type == ECustomVersionDifference::Newer)
 		{
-			UE_LOG(LogAssetRegistry, Error, TEXT("Package %s has newer custom version of %s"), *PackageFilename, *Diff.Version->GetFriendlyName().ToString());
-
-			OutErrorCode = EOpenPackageResult::VersionTooNew;
+			if (IsEnforcePackageCompatibleVersionCheck())
+			{
+				UE_LOG(LogAssetRegistry, Error, TEXT("Package %s has newer custom version of %s"),
+					*PackageFilename, *Diff.Version->GetFriendlyName().ToString());
+				OutErrorCode = EOpenPackageResult::VersionTooNew;
+			}
 		}
+		// else ECustomVersionDifference::Older, which is not a problem
 	}
 
 	//make sure the filereader gets the correct version number (it defaults to latest version)
diff --git a/Engine/Source/Runtime/CoreUObject/Private/UObject/LinkerLoad.cpp b/Engine/Source/Runtime/CoreUObject/Private/UObject/LinkerLoad.cpp
--- a/Engine/Source/Runtime/CoreUObject/Private/UObject/LinkerLoad.cpp
+++ b/Engine/Source/Runtime/CoreUObject/Private/UObject/LinkerLoad.cpp
@@ -181,6 +181,11 @@ static FAutoConsoleVariableRef CEnforcePackageCompatibleVersionCheck(
 	ECVF_Default
 );
 
+bool IsEnforcePackageCompatibleVersionCheck()
+{
+	return GEnforcePackageCompatibleVersionCheck != 0;
+}
+
 /** 
  * Required to load packages saved from the editor domain between UE 5.0 and 5.2, the cvar is only provided in case the fix causes
  * unintended problems so that it can be disabled quickly.
@@ -1373,7 +1378,7 @@ FLinkerLoad::ELinkerStatus FLinkerLoad::SerializePackageFileSummaryInternal()
 	}
 
 	// Don't load packages that are only compatible with an engine version newer than the current one.
-	if (bLoaderNeedsEngineVersionChecks && GEnforcePackageCompatibleVersionCheck && !FEngineVersion::Current().IsCompatibleWith(Summary.CompatibleWithEngineVersion))
+	if (bLoaderNeedsEngineVersionChecks && IsEnforcePackageCompatibleVersionCheck() && !FEngineVersion::Current().IsCompatibleWith(Summary.CompatibleWithEngineVersion))
 	{
 		FMessageLog("LoadErrors").Warning(FText::Format(NSLOCTEXT("Core", "LinkerLoad_EngineVersionIncompatible", "Package '{0}' has been saved with a newer engine version and can't be loaded. Current EngineVersion: {1} (Licensee={2}). Package EngineVersion: {3} (Licensee={4})"),
 			FText::FromString(GetDebugName()), 
diff --git a/Engine/Source/Runtime/CoreUObject/Public/UObject/LinkerLoad.h b/Engine/Source/Runtime/CoreUObject/Public/UObject/LinkerLoad.h
--- a/Engine/Source/Runtime/CoreUObject/Public/UObject/LinkerLoad.h
+++ b/Engine/Source/Runtime/CoreUObject/Public/UObject/LinkerLoad.h
@@ -1488,3 +1488,4 @@ enum class ENotifyRegistrationPhase
 
 COREUOBJECT_API void NotifyRegistrationEvent(const TCHAR* PackageName, const TCHAR* Name, ENotifyRegistrationType NotifyRegistrationType, ENotifyRegistrationPhase NotifyRegistrationPhase, UObject *(*InRegister)() = nullptr, bool InbDynamic = false, UObject* FinishedObject = nullptr);
 COREUOBJECT_API void NotifyRegistrationComplete();
+COREUOBJECT_API bool IsEnforcePackageCompatibleVersionCheck();
