Cooker: AssetRegistry: Optimize FilterTags
#rnx

diff --git a/Engine/Source/Runtime/AssetRegistry/Private/AssetRegistryState.cpp b/Engine/Source/Runtime/AssetRegistry/Private/AssetRegistryState.cpp
--- a/Engine/Source/Runtime/AssetRegistry/Private/AssetRegistryState.cpp
+++ b/Engine/Source/Runtime/AssetRegistry/Private/AssetRegistryState.cpp
@@ -2034,15 +2034,32 @@ void FAssetRegistryState::AddTagsToAssetData(const FSoftObjectPath& InObjectPath
 
 void FAssetRegistryState::FilterTags(const FAssetRegistrySerializationOptions& Options)
 {
+	// Calling SetTagsOnExistingAsset for any changed tags is slow because the elements of 
+	// CachedAssetsByTag are unsorted TArrays and removal of the AssetData from its old
+	// CachedAssetsByTag is slow. For cases where many Assets change it is therefore faster
+	// just to recreate CachedAssetsByTag rather than trying to update them.
+	for (TPair<FName, TArray<FAssetData*>>& Pair : CachedAssetsByTag)
+	{
+		Pair.Value.Reset();
+	}
+
 	for (FAssetData* AssetData : CachedAssets) 
 	{
 		check(AssetData);
 
 		FAssetDataTagMap LocalTagsAndValues;
-		FAssetRegistryState::FilterTags(AssetData->TagsAndValues, LocalTagsAndValues, Options.CookFilterlistTagsByClass.Find(AssetData->AssetClassPath), Options);
+		FAssetRegistryState::FilterTags(AssetData->TagsAndValues, LocalTagsAndValues,
+			Options.CookFilterlistTagsByClass.Find(AssetData->AssetClassPath), Options);
 		if (LocalTagsAndValues != AssetData->TagsAndValues)
 		{
-			SetTagsOnExistingAsset(AssetData, MoveTemp(LocalTagsAndValues));
+			AssetData->TagsAndValues = FAssetDataTagMapSharedView(MoveTemp(LocalTagsAndValues));
+		}
+
+		// Add the AssetData to all its CachedAssetsByTag keys even if nothing changed, because
+		// we are reconstructing all CachedAssetsByTag.
+		for (auto TagIt = AssetData->TagsAndValues.CreateConstIterator(); TagIt; ++TagIt)
+		{
+			CachedAssetsByTag.FindOrAdd(TagIt.Key()).Add(AssetData);
 		}
 	}
 }
@@ -2057,7 +2074,10 @@ void FAssetRegistryState::SetTagsOnExistingAsset(FAssetData* AssetData, FAssetDa
 		if (!NewTags.Contains(FNameKey))
 		{
 			TArray<FAssetData*>* OldTagAssets = CachedAssetsByTag.Find(FNameKey);
-			OldTagAssets->RemoveSingleSwap(AssetData);
+			if (OldTagAssets)
+			{
+				OldTagAssets->RemoveSingleSwap(AssetData);
+			}
 		}
 	}
 	// Update the tag cache map to add added tags
@@ -2067,8 +2087,7 @@ void FAssetRegistryState::SetTagsOnExistingAsset(FAssetData* AssetData, FAssetDa
 
 		if (!AssetData->TagsAndValues.Contains(FNameKey))
 		{
-			TArray<FAssetData*>& NewTagAssets = CachedAssetsByTag.FindOrAdd(FNameKey);
-			NewTagAssets.Add(AssetData);
+			CachedAssetsByTag.FindOrAdd(FNameKey).Add(AssetData);
 		}
 	}
 	AssetData->TagsAndValues = FAssetDataTagMapSharedView(MoveTemp(NewTags));
