AssetRegistry: Cooked platforms load their data from a premade assetregistry from the cook, possibly including some dynamically loadable and unloadable plugins. Mark the AssetDatas in those preloaded plugins as persistent and do not unload them when the plugin unloads.
#rnx

diff --git a/Engine/Source/Runtime/AssetRegistry/Private/AssetRegistry.cpp b/Engine/Source/Runtime/AssetRegistry/Private/AssetRegistry.cpp
--- a/Engine/Source/Runtime/AssetRegistry/Private/AssetRegistry.cpp
+++ b/Engine/Source/Runtime/AssetRegistry/Private/AssetRegistry.cpp
@@ -792,6 +792,7 @@ void FAssetRegistryImpl::LoadPremadeAssetRegistry(Impl::FEventContext& EventCont
 			{
 				AppendState(EventContext, ARState, FAssetRegistryState::EInitializationMode::OnlyUpdateNew);
 			}
+			UpdatePersistentMountPoints();
 		}
 		else
 		{
@@ -4623,6 +4624,7 @@ void FAssetRegistryImpl::Serialize(FArchive& Ar, Impl::FEventContext& EventConte
 	{
 		State.Load(Ar);
 		CachePathsFromState(EventContext, State);
+		UpdatePersistentMountPoints();
 	}
 	else if (Ar.IsSaving())
 	{
@@ -6602,6 +6604,13 @@ void FAssetRegistryImpl::OnContentPathDismounted(Impl::FEventContext& EventConte
 		GlobalGatherer->RemoveMountPoint(FileSystemPath);
 	}
 
+	FName MountPoint = FName(FStringView(AssetPathNoTrailingSlash));
+	if (PersistentMountPoints.Contains(MountPoint))
+	{
+		// This path is marked to never remove its AssetDatas. Skip the code below to remove it.
+		return;
+	}
+
 	// Remove all cached assets and Verse files found at this location
 	{
 		FName AssetPathNoTrailingSlashFName(*AssetPathNoTrailingSlash);
@@ -6644,6 +6653,22 @@ void FAssetRegistryImpl::OnContentPathDismounted(Impl::FEventContext& EventConte
 	}
 }
 
+void FAssetRegistryImpl::UpdatePersistentMountPoints()
+{
+	State.EnumerateAllPaths([this](FName Path)
+		{
+			TStringBuilder<256> PathString(InPlace, Path);
+			bool bHadClassesPrefix;
+			FStringView MountPoint = FPathViews::GetMountPointNameFromPath(PathString, &bHadClassesPrefix, false /* bInWithoutSlashes*/);
+			if (!MountPoint.IsEmpty() && !bHadClassesPrefix)
+			{
+				// Format returned by GetMountPointNameFromPath is e.g. /Engine, which is the format we need:
+				// LongPackageName with no trailing slash
+				PersistentMountPoints.Add(FName(MountPoint));
+			}
+		});
+}
+
 }
 
 void UAssetRegistryImpl::SetTemporaryCachingMode(bool bEnable)
diff --git a/Engine/Source/Runtime/AssetRegistry/Private/AssetRegistryImpl.h b/Engine/Source/Runtime/AssetRegistry/Private/AssetRegistryImpl.h
--- a/Engine/Source/Runtime/AssetRegistry/Private/AssetRegistryImpl.h
+++ b/Engine/Source/Runtime/AssetRegistry/Private/AssetRegistryImpl.h
@@ -365,6 +365,8 @@ private:
 	/** Moves a premade asset registry state into this AR */
 	void LoadPremadeAssetRegistry(Impl::FEventContext& Context,
 		Premade::ELoadResult LoadResult, FAssetRegistryState&& ARState);
+	/** Add MountPoints of all AssetDatas currently registered in this->State to the list of PersistentMountPoints. */
+	void UpdatePersistentMountPoints();
 
 private:
 
@@ -410,6 +412,12 @@ private:
 	TMap<FTopLevelAssetPath, int64> TagSizeByClass;
 #endif
 
+	/**
+	 * MountPoints, in the format of LongPackageName with no trailing slash, that should not have their
+	 * AssetDatas removed even if the MountPoint is dismounted.
+	 */
+	TSet<FName> PersistentMountPoints;
+
 	/** Time spent processing Gather results */
 	float StoreGatherResultsTimeSeconds;
 	/** The highest number of pending results observed during initial gathering */
diff --git a/Engine/Source/Runtime/AssetRegistry/Private/AssetRegistryState.cpp b/Engine/Source/Runtime/AssetRegistry/Private/AssetRegistryState.cpp
--- a/Engine/Source/Runtime/AssetRegistry/Private/AssetRegistryState.cpp
+++ b/Engine/Source/Runtime/AssetRegistry/Private/AssetRegistryState.cpp
@@ -1070,6 +1070,14 @@ void FAssetRegistryState::EnumerateAllAssets(TFunctionRef<void(const FAssetData&
 	}
 }
 
+void FAssetRegistryState::EnumerateAllPaths(TFunctionRef<void(FName PathName)> Callback) const
+{
+	for (const TPair<FName, TArray<FAssetData*>>& Pair : CachedAssetsByPath)
+	{
+		Callback(Pair.Key);
+	}
+}
+
 void FAssetRegistryState::GetPackagesByName(FStringView PackageName, TArray<FName>& OutPackageNames) const
 {
 	// Note that we use CachedAssetsByPackageName rather than CachedPackageData because CachedPackageData
diff --git a/Engine/Source/Runtime/AssetRegistry/Public/AssetRegistry/AssetRegistryState.h b/Engine/Source/Runtime/AssetRegistry/Public/AssetRegistry/AssetRegistryState.h
--- a/Engine/Source/Runtime/AssetRegistry/Public/AssetRegistry/AssetRegistryState.h
+++ b/Engine/Source/Runtime/AssetRegistry/Public/AssetRegistry/AssetRegistryState.h
@@ -265,6 +265,13 @@ public:
 	ASSETREGISTRY_API bool EnumerateAllAssets(const TSet<FName>& PackageNamesToSkip, TFunctionRef<bool(const FAssetData&)> Callback, bool bARFiltering = false) const;
 	ASSETREGISTRY_API void EnumerateAllAssets(TFunctionRef<void(const FAssetData&)> Callback) const;
 
+	/**
+	 * Calls the callback with the LongPackageName of each path that has assets as direct children.
+	 * Callback will not be called for parent paths that have childpaths with direct children but do not
+	 * have direct children of their own.
+	 */
+	ASSETREGISTRY_API void EnumerateAllPaths(TFunctionRef<void(FName PathName)> Callback) const;
+
 	/**
 	 * Gets the LongPackageNames for all packages with the given PackageName.
 	 * Call to check existence of a LongPackageName or find all packages with a ShortPackageName.
