AssetRegistry: AssetDataGather: When checking whether cached data is invalidated, a change in capitalization of the PackageName should cause invalidation as well, to avoid causing the package to save itself to the old capitalization of its path when saved. (PackageNames MatchCapitalizationOnDisk, but thie capitalization they match is the data that comes from the AssetRegistry).
#rnx
#10805

diff --git a/Engine/Source/Runtime/AssetRegistry/Private/AssetDataGatherer.cpp b/Engine/Source/Runtime/AssetRegistry/Private/AssetDataGatherer.cpp
--- a/Engine/Source/Runtime/AssetRegistry/Private/AssetDataGatherer.cpp
+++ b/Engine/Source/Runtime/AssetRegistry/Private/AssetDataGatherer.cpp
@@ -3771,7 +3771,7 @@ FAssetDataGatherer::ETickResult FAssetDataGatherer::TickInternal(double& TickSta
 			{
 				DiskCachedAssetData = nullptr;
 			}
-			else if ((DiskCachedAssetData->DependencyData.PackageName != PackageName && DiskCachedAssetData->DependencyData.PackageName != NAME_None) ||
+			else if ((!DiskCachedAssetData->DependencyData.PackageName.IsEqual(PackageName, ENameCase::CaseSensitive) && DiskCachedAssetData->DependencyData.PackageName != NAME_None) ||
 				DiskCachedAssetData->Extension != Extension)
 			{
 				UE_LOG(LogAssetRegistry, Display, TEXT("Cached dependency data for package '%s' is invalid. Discarding cached data."), *PackageName.ToString());
