AssetRegistry: DumpAssetRegistry:

Remove display of tag caches from cooker's dumpassetregistry because it is expensive to compute and is redundant with the more useful tags per asset display.

Sort the tags in the tags per asset display to remove spurious indeterminism from the dump.
#rnx

diff --git a/Engine/Source/Editor/UnrealEd/Private/CookOnTheFlyServer.cpp b/Engine/Source/Editor/UnrealEd/Private/CookOnTheFlyServer.cpp
--- a/Engine/Source/Editor/UnrealEd/Private/CookOnTheFlyServer.cpp
+++ b/Engine/Source/Editor/UnrealEd/Private/CookOnTheFlyServer.cpp
@@ -8442,7 +8442,7 @@ void DumpAssetRegistryForCooker(IAssetRegistry* AssetRegistry)
 	DeleteReportDir.DeleteDirectory(DumpDir);
 	PlatformFile.CreateDirectoryTree(*DumpDir);
 	TArray<FString> Pages;
-	TArray<FString> Arguments({ TEXT("ObjectPath"),TEXT("PackageName"),TEXT("Path"),TEXT("Class"),TEXT("Tag"),
+	TArray<FString> Arguments({ TEXT("ObjectPath"),TEXT("PackageName"),TEXT("Path"),TEXT("Class"),
 		TEXT("DependencyDetails"), TEXT("PackageData"), TEXT("LegacyDependencies"), TEXT("AssetTags") });
 	AssetRegistry->DumpState(Arguments, Pages, 10000 /* LinesPerPage */);
 	int PageIndex = 0;
diff --git a/Engine/Source/Runtime/AssetRegistry/Private/AssetRegistryState.cpp b/Engine/Source/Runtime/AssetRegistry/Private/AssetRegistryState.cpp
--- a/Engine/Source/Runtime/AssetRegistry/Private/AssetRegistryState.cpp
+++ b/Engine/Source/Runtime/AssetRegistry/Private/AssetRegistryState.cpp
@@ -2793,6 +2793,7 @@ void FAssetRegistryState::Dump(const TArray<FString>& Arguments, TArray<FString>
 		AddLine();
 
 		InitializeSortedAssets();
+		TArray<FName> SortedTagKeys;
 		for (const FAssetData* AssetData : SortedAssets)
 		{
 			if (AssetData->TagsAndValues.Num() == 0)
@@ -2804,11 +2805,18 @@ void FAssetRegistryState::Dump(const TArray<FString>& Arguments, TArray<FString>
 			PageBuffer << TEXT("  ") << FCachedAssetKey(AssetData);
 			AddLine();
 
-			AssetData->TagsAndValues.ForEach([&](const TPair<FName, FAssetTagValueRef>& TagPair) 
+			SortedTagKeys.Reset();
+			AssetData->TagsAndValues.ForEach([&SortedTagKeys](const TPair<FName, FAssetTagValueRef>& TagPair)
 			{
-				PageBuffer << TEXT("    ") << TagPair.Key << TEXT(" : ") << *TagPair.Value.AsString();
-				AddLine();	
+				SortedTagKeys.Add(TagPair.Key);
 			});
+			Algo::Sort(SortedTagKeys, FNameLexicalLess());
+			for (FName TagKey : SortedTagKeys)
+			{
+				FAssetTagValueRef Value = AssetData->TagsAndValues.FindTag(TagKey);
+				PageBuffer << TEXT("    ") << TagKey << TEXT(" : ") << *Value.AsString();
+				AddLine();
+			}
 		}
 
 		PageBuffer.Appendf(TEXT("--- End AssetTags : %d entries ---"), Counter);
