AssetRegistry: Move FPreloadSettings into AssetDataGathererPrivate.h so it can be used in other cpp files.
#rnx

diff --git a/Engine/Source/Runtime/AssetRegistry/Private/AssetDataGatherer.cpp b/Engine/Source/Runtime/AssetRegistry/Private/AssetDataGatherer.cpp
--- a/Engine/Source/Runtime/AssetRegistry/Private/AssetDataGatherer.cpp
+++ b/Engine/Source/Runtime/AssetRegistry/Private/AssetDataGatherer.cpp
@@ -77,6 +77,89 @@ static FAutoConsoleVariableRef CVarIgnoreEmptyDirectories(
 	bIgnoreEmptyDirectories,
 	TEXT("If true, completely empty leaf directories are ignored by the asset registry while scanning"));
 
+void FPreloadSettings::Initialize()
+{
+	if (bInitialized)
+	{
+		return;
+	}
+	bInitialized = true;
+	if (!FParse::Value(FCommandLine::Get(), TEXT("AssetRegistryCacheRootFolder="), AssetRegistryCacheRootFolder))
+	{
+		AssetRegistryCacheRootFolder = FPaths::ProjectIntermediateDir();
+	}
+	bForceDependsGathering = FParse::Param(FCommandLine::Get(), TEXT("ForceDependsGathering"));
+	bGatherDependsData = (GIsEditor && !FParse::Param(FCommandLine::Get(), TEXT("NoDependsGathering"))) || bForceDependsGathering;
+	bool bNoAssetRegistryCache = FParse::Param(FCommandLine::Get(), TEXT("NoAssetRegistryCache"));
+	bool bNoAssetRegistryCacheRead = FParse::Param(FCommandLine::Get(), TEXT("NoAssetRegistryCacheRead"));
+	bool bNoAssetRegistryCacheWrite = FParse::Param(FCommandLine::Get(), TEXT("NoAssetRegistryCacheWrite"));
+	uint32 MultiprocessId = UE::GetMultiprocessId();
+	bool bMultiprocess = MultiprocessId > 0 || FParse::Param(FCommandLine::Get(), TEXT("multiprocess"));
+	bCacheReadEnabled = !bNoAssetRegistryCache && !bNoAssetRegistryCacheRead;
+	bCacheWriteEnabled = !bNoAssetRegistryCache && !bNoAssetRegistryCacheWrite && !bMultiprocess;
+	bool bAsyncEnabled = FPlatformProcess::SupportsMultithreading() && FTaskGraphInterface::IsRunning();
+
+	MonolithicCacheBaseFilename = AssetRegistryCacheRootFolder / (bGatherDependsData ? TEXT("CachedAssetRegistry") : TEXT("CachedAssetRegistryNoDeps"));
+#if UE_EDITOR // See note on FPreloader for why we only allow preloading if UE_EDITOR
+	bMonolithicCacheActivatedDuringPreload = bAsyncEnabled && UE::AssetRegistry::ShouldSearchAllAssetsAtStart();
+#else
+	bMonolithicCacheActivatedDuringPreload = false;
+#endif
+}
+bool FPreloadSettings::IsCacheReadEnabled() const
+{
+	return bCacheReadEnabled;
+}
+bool FPreloadSettings::IsCacheWriteEnabled() const
+{
+	return bCacheWriteEnabled;
+}
+bool FPreloadSettings::IsMonolithicCacheActivatedDuringPreload() const
+{
+	return bMonolithicCacheActivatedDuringPreload;
+}
+bool FPreloadSettings::IsPreloadMonolithicCache() const
+{
+	return bCacheReadEnabled && bMonolithicCacheActivatedDuringPreload;
+}
+bool FPreloadSettings::IsGatherDependsData() const
+{
+	return bGatherDependsData;
+}
+bool FPreloadSettings::IsForceDependsGathering() const
+{
+	return bForceDependsGathering;
+}
+FString FPreloadSettings::GetLegacyMonolithicCacheFilename() const
+{
+	return MonolithicCacheBaseFilename + TEXT(".bin");
+}
+const FString& FPreloadSettings::GetMonolithicCacheBaseFilename() const
+{
+	return MonolithicCacheBaseFilename;
+}
+const FString& FPreloadSettings::GetAssetRegistryCacheRootFolder() const
+{
+	return AssetRegistryCacheRootFolder;
+}
+
+TArray<FString> FPreloadSettings::FindShardedMonolithicCacheFiles() const
+{
+	TArray<FString> CachePaths;
+	IFileManager::Get().FindFiles(CachePaths, *(GetMonolithicCacheBaseFilename() + TEXT("_*.bin")), /* Files */ true, /* Directories */ false);
+	if (CachePaths.Num())
+	{
+		FString Directory = FPaths::GetPath(GetMonolithicCacheBaseFilename());
+		for (FString& Path : CachePaths)
+		{
+			Path = Directory / Path;
+		}
+	}
+	return CachePaths;
+}
+
+FPreloadSettings GPreloadSettings;
+
 bool IsPackageBlocked(const FStringView& FilePath)
 {
 	return bBlockPackagesWithMarkOfTheWeb && IPlatformFile::GetPlatformPhysical().HasMarkOfTheWeb(*FString(FilePath));
@@ -3030,105 +3113,6 @@ private:
 	TUniquePtr<IMappedFileRegion> Region;
 };
 
-/**
- * Settings about whether to use cache data for the AssetDataGatherer; these settings are shared by
- * FPreloader and the FAssetDataGatherer.
- */
-struct FPreloadSettings
-{
-	void Initialize()
-	{
-		if (bInitialized)
-		{
-			return;
-		}
-		bInitialized = true;
-		if (!FParse::Value(FCommandLine::Get(), TEXT("AssetRegistryCacheRootFolder="), AssetRegistryCacheRootFolder))
-		{
-			AssetRegistryCacheRootFolder = FPaths::ProjectIntermediateDir();
-		}
-		bForceDependsGathering = FParse::Param(FCommandLine::Get(), TEXT("ForceDependsGathering"));
-		bGatherDependsData = (GIsEditor && !FParse::Param(FCommandLine::Get(), TEXT("NoDependsGathering"))) || bForceDependsGathering;
-		bool bNoAssetRegistryCache = FParse::Param(FCommandLine::Get(), TEXT("NoAssetRegistryCache"));
-		bool bNoAssetRegistryCacheRead = FParse::Param(FCommandLine::Get(), TEXT("NoAssetRegistryCacheRead"));
-		bool bNoAssetRegistryCacheWrite = FParse::Param(FCommandLine::Get(), TEXT("NoAssetRegistryCacheWrite"));
-		uint32 MultiprocessId = UE::GetMultiprocessId();
-		bool bMultiprocess = MultiprocessId > 0 || FParse::Param(FCommandLine::Get(), TEXT("multiprocess"));
-		bCacheReadEnabled = !bNoAssetRegistryCache && !bNoAssetRegistryCacheRead;
-		bCacheWriteEnabled = !bNoAssetRegistryCache && !bNoAssetRegistryCacheWrite && !bMultiprocess;
-		bool bAsyncEnabled = FPlatformProcess::SupportsMultithreading() && FTaskGraphInterface::IsRunning();
-
-		MonolithicCacheBaseFilename = AssetRegistryCacheRootFolder / (bGatherDependsData ? TEXT("CachedAssetRegistry") : TEXT("CachedAssetRegistryNoDeps"));
-#if UE_EDITOR // See note on FPreloader for why we only allow preloading if UE_EDITOR
-		bMonolithicCacheActivatedDuringPreload = bAsyncEnabled && UE::AssetRegistry::ShouldSearchAllAssetsAtStart();
-#else
-		bMonolithicCacheActivatedDuringPreload = false;
-#endif
-	}
-	bool IsCacheReadEnabled() const
-	{
-		return bCacheReadEnabled;
-	}
-	bool IsCacheWriteEnabled() const
-	{
-		return bCacheWriteEnabled;
-	}
-	bool IsMonolithicCacheActivatedDuringPreload() const
-	{
-		return bMonolithicCacheActivatedDuringPreload;
-	}
-	bool IsPreloadMonolithicCache() const
-	{
-		return bCacheReadEnabled && bMonolithicCacheActivatedDuringPreload;
-	}
-	bool IsGatherDependsData() const
-	{
-		return bGatherDependsData;
-	}
-	bool IsForceDependsGathering() const
-	{
-		return bForceDependsGathering;
-	}
-	FString GetLegacyMonolithicCacheFilename() const
-	{
-		return MonolithicCacheBaseFilename + TEXT(".bin");
-	}
-	const FString& GetMonolithicCacheBaseFilename() const
-	{
-		return MonolithicCacheBaseFilename;
-	}
-	const FString& GetAssetRegistryCacheRootFolder() const
-	{
-		return AssetRegistryCacheRootFolder;
-	}
-	
-	TArray<FString> FindShardedMonolithicCacheFiles() const
-	{
-		TArray<FString> CachePaths;
-		IFileManager::Get().FindFiles(CachePaths, *(GetMonolithicCacheBaseFilename() + TEXT("_*.bin")), /* Files */ true, /* Directories */ false);
-		if (CachePaths.Num())
-		{
-			FString Directory = FPaths::GetPath(GetMonolithicCacheBaseFilename());
-			for (FString& Path : CachePaths)
-			{
-				Path = Directory / Path;
-			}
-		}
-		return CachePaths;
-	}
-
-private:
-	FString MonolithicCacheBaseFilename;
-	FString AssetRegistryCacheRootFolder;
-	bool bForceDependsGathering = false;
-	bool bGatherDependsData = false;
-	bool bCacheReadEnabled = false;
-	bool bCacheWriteEnabled = false;
-	bool bMonolithicCacheActivatedDuringPreload = false;
-	bool bInitialized = false;
-};
-FPreloadSettings GPreloadSettings;
-
 #if UE_EDITOR
 /** A class to preload the monolithic cache used by FAssetDataGatherer. Preloading the cache allows us to
  * start very early in editor startup, so that we have time to finish the cache load before the engine starts
diff --git a/Engine/Source/Runtime/AssetRegistry/Private/AssetDataGathererPrivate.h b/Engine/Source/Runtime/AssetRegistry/Private/AssetDataGathererPrivate.h
--- a/Engine/Source/Runtime/AssetRegistry/Private/AssetDataGathererPrivate.h
+++ b/Engine/Source/Runtime/AssetRegistry/Private/AssetDataGathererPrivate.h
@@ -12,11 +12,7 @@
 
 #include <atomic>
 
-namespace UE
-{
-namespace AssetDataGather
-{
-namespace Private
+namespace UE::AssetDataGather::Private
 {
 
 class FMountDir;
@@ -840,6 +836,35 @@ private:
 	friend class FScanDir;
 };
 
-} // namespace Private
-} // namespace AssetDataGather
-} // namespace UE
\ No newline at end of file
+/**
+ * Settings about whether to use cache data for the AssetDataGatherer; these settings are shared by
+ * FPreloader and the FAssetDataGatherer.
+ */
+struct FPreloadSettings
+{
+public:
+	void Initialize();
+	bool IsCacheReadEnabled() const;
+	bool IsCacheWriteEnabled() const;
+	bool IsMonolithicCacheActivatedDuringPreload() const;
+	bool IsPreloadMonolithicCache() const;
+	bool IsGatherDependsData() const;
+	bool IsForceDependsGathering() const;
+	FString GetLegacyMonolithicCacheFilename() const;
+	const FString& GetMonolithicCacheBaseFilename() const;
+	const FString& GetAssetRegistryCacheRootFolder() const;
+	TArray<FString> FindShardedMonolithicCacheFiles() const;
+
+private:
+	FString MonolithicCacheBaseFilename;
+	FString AssetRegistryCacheRootFolder;
+	bool bForceDependsGathering = false;
+	bool bGatherDependsData = false;
+	bool bCacheReadEnabled = false;
+	bool bCacheWriteEnabled = false;
+	bool bMonolithicCacheActivatedDuringPreload = false;
+	bool bInitialized = false;
+};
+extern FPreloadSettings GPreloadSettings;
+
+} // namespace UE::AssetDataGather::Private
\ No newline at end of file
