AssetRegistry: Reduce wasted cputime in listeners to the OnAssetAdded event by not setting IsLoadingAssets to false until after all OnAssetAdded events have been broadcast.
This fixes a potentially large stall when loading a map near the end of asset registry scan.
Checked in on behalf of matt.peters

diff --git a/Engine/Source/Runtime/AssetRegistry/Private/AssetRegistry.cpp b/Engine/Source/Runtime/AssetRegistry/Private/AssetRegistry.cpp
--- a/Engine/Source/Runtime/AssetRegistry/Private/AssetRegistry.cpp
+++ b/Engine/Source/Runtime/AssetRegistry/Private/AssetRegistry.cpp
@@ -1793,7 +1793,7 @@ void UAssetRegistryImpl::WaitForCompletion()
 		ProcessLoadedAssetsToUpdateCache(EventContext, -1., Status);
 #endif
 		Broadcast(EventContext);
-		if (Status != EGatherStatus::Active)
+		if (Status != EGatherStatus::Active && Status != EGatherStatus::WaitingForEvents)
 		{
 			break;
 		}
@@ -4154,7 +4154,8 @@ void UAssetRegistryImpl::Tick(float DeltaTime)
 		}
 #endif
 		Broadcast(EventContext);
-	} while (bInterrupted && (TickStartTime < 0 || (FPlatformTime::Seconds() - TickStartTime) <= UE::AssetRegistry::Impl::MaxSecondsPerFrame));
+	} while ((bInterrupted || Status == UE::AssetRegistry::Impl::EGatherStatus::WaitingForEvents) &&
+		(TickStartTime < 0 || (FPlatformTime::Seconds() - TickStartTime) <= UE::AssetRegistry::Impl::MaxSecondsPerFrame));
 }
 
 namespace UE::AssetRegistry
@@ -4223,10 +4224,13 @@ Impl::EGatherStatus FAssetRegistryImpl::TickGatherer(Impl::FEventContext& EventC
 		return ResultContext.NumFilesToSearch + ResultContext.NumPathsToSearch + BackgroundResults.Paths.Num() + BackgroundResults.Assets.Num()
 			+ BackgroundResults.Dependencies.Num() + BackgroundResults.CookedPackageNamesWithoutAssetData.Num();
 	};
-	auto UpdateStatus = [bHadAssetsToProcess, &ResultContext, &EventContext, &OutStatus, this](int32 NumGatherPending, bool bInterrupted)
+	int32 NumPending = 0;
+	auto CalculateStatus =
+		[bHadAssetsToProcess, &NumPending, &ResultContext, &EventContext, &OutStatus, this]
+		(int32 NumGatherPending, bool bInterrupted)
 	{
 		// Compute total pending, plus highest pending for this run so we can show a good progress bar
-		int32 NumPending = NumGatherPending
+		NumPending = NumGatherPending
 #if WITH_EDITOR
 			+ (PackagesNeedingDependencyCalculation.Num() ? 1 : 0)
 #endif
@@ -4245,6 +4249,9 @@ Impl::EGatherStatus FAssetRegistryImpl::TickGatherer(Impl::FEventContext& EventC
 		{
 			OutStatus = EGatherStatus::Active;
 		}
+	};
+	auto UpdateStatus = [bHadAssetsToProcess, &NumPending, &ResultContext, &EventContext, &OutStatus, this]()
+	{
 		// Notify the status change, only when something changed, or when sending the final result before going idle
 		if (ResultContext.bIsSearching || bHadAssetsToProcess ||
 			(OutStatus == EGatherStatus::Complete && this->GatherStatus != EGatherStatus::Complete))
@@ -4293,7 +4300,8 @@ Impl::EGatherStatus FAssetRegistryImpl::TickGatherer(Impl::FEventContext& EventC
 		CookedPackageNamesWithoutAssetDataGathered(EventContext, TickStartTime, BackgroundResults.CookedPackageNamesWithoutAssetData, bOutInterrupted);
 		if (bOutInterrupted)
 		{
-			UpdateStatus(GetNumGatherFromDiskPending(), true /* bInterrupted */);
+			CalculateStatus(GetNumGatherFromDiskPending(), true /* bInterrupted */);
+			UpdateStatus();
 			return OutStatus;
 		}
 	}
@@ -4322,29 +4330,50 @@ Impl::EGatherStatus FAssetRegistryImpl::TickGatherer(Impl::FEventContext& EventC
 		LoadCalculatedDependencies(nullptr, TickStartTime, InheritanceContext, bOutInterrupted);
 		if (bOutInterrupted)
 		{
-			UpdateStatus(NumGatherFromDiskPending, true /* bInterrupted */);
+			CalculateStatus(NumGatherFromDiskPending, true /* bInterrupted */);
+			UpdateStatus();
 			return OutStatus;
 		}
 	}
 #endif
 
-	// If completing an initial search, refresh the content browser
-	UpdateStatus(NumGatherFromDiskPending, false /* bInterrupted */);
+	CalculateStatus(NumGatherFromDiskPending, false /* bInterrupted */);
 
 	if (OutStatus == EGatherStatus::Complete)
 	{
-		HighestPending = 0;
-		BackgroundResults.Shrink();
-
-		// Finishing the background search is blocked until preloading complete because plugins can be mounted during
-		// startup up until that point, and we need to wait for all the plugins to load before declaring completion.
-		if (!bInitialSearchCompleted && bPreloadingComplete && IsEngineStartupModuleLoadingComplete())
+		if (!bInitialSearchCompleted)
 		{
-			RecordTimer(); // OnInitialSearchComplete reads data set by RecordTimer
-			OnInitialSearchCompleted(EventContext);
+			// Finishing the background search is blocked until preloading complete because plugins can be mounted during
+			// startup up until that point, and we need to wait for all the plugins to load before declaring completion.
+			bool bCanCompleteInitialSearch = bPreloadingComplete && IsEngineStartupModuleLoadingComplete();
+
+			if (bCanCompleteInitialSearch)
+			{
+				if (!EventContext.AssetEvents.IsEmpty())
+				{
+					// Don't mark the InitialSearch completed until we have sent all the AssetDataAdded events
+					// that arose from the final tick of the gatherer. Some callers might do more expensive
+					// work for assets added after the initial search completed, and we don't want them to do
+					// that more expensive work on the last batch of assets before completion.
+					OutStatus = EGatherStatus::WaitingForEvents;
+					bCanCompleteInitialSearch = false;
+				}
+			}
+			if (bCanCompleteInitialSearch)
+			{
+				RecordTimer(); // OnInitialSearchComplete reads data set by RecordTimer
+				OnInitialSearchCompleted(EventContext);
+			}
 		}
 	}
 
+	UpdateStatus();
+	if (OutStatus == EGatherStatus::Complete)
+	{
+		HighestPending = 0;
+		BackgroundResults.Shrink();
+	}
+
 	return OutStatus;
 }
 
diff --git a/Engine/Source/Runtime/AssetRegistry/Private/AssetRegistryImpl.h b/Engine/Source/Runtime/AssetRegistry/Private/AssetRegistryImpl.h
--- a/Engine/Source/Runtime/AssetRegistry/Private/AssetRegistryImpl.h
+++ b/Engine/Source/Runtime/AssetRegistry/Private/AssetRegistryImpl.h
@@ -99,6 +99,7 @@ namespace Impl
 		Active,
 		Complete,
 		UnableToProgress,
+		WaitingForEvents,
 	};
 
 	/** Affects how rules are applied to improve loading/runtime performance */
