Turn off the NTFS journal cache used by the asset registry discovery cache while we investigate how it can get out of sync with the asset registry cache.
Fixed the config reading for the AssetRegistryDiscoveryCache value and made the never setting also skip the reading of the discovery cache.

diff --git a/Engine/Config/BaseEngine.ini b/Engine/Config/BaseEngine.ini
--- a/Engine/Config/BaseEngine.ini
+++ b/Engine/Config/BaseEngine.ini
@@ -2651,7 +2651,7 @@ bFilterSearchableNames=true
 ; be used on-demand without checking it for out-of-date via commandline -SkipAssetRegistryCacheInvalidate.
 ; To support that use case set AssetRegistryDiscoveryCache=Always to always create the cache.
 ; Options are Never,Default,Always. Default only creates the cache if the platform supports invalidation via FileJournal.
-AssetRegistryDiscoveryCache=Default
+AssetRegistryDiscoveryCache=Never
 
 [AutomationTesting]
 ImportTestPath=Content/EditorAutomation/
diff --git a/Engine/Source/Runtime/AssetRegistry/Private/AssetDataGatherer.cpp b/Engine/Source/Runtime/AssetRegistry/Private/AssetDataGatherer.cpp
--- a/Engine/Source/Runtime/AssetRegistry/Private/AssetDataGatherer.cpp
+++ b/Engine/Source/Runtime/AssetRegistry/Private/AssetDataGatherer.cpp
@@ -102,20 +102,22 @@ void FPreloadSettings::Initialize()
 	bool bPlatformSupportsDiscoveryCache = FPlatformFileManager::Get().GetPlatformFile().FileJournalIsAvailable();
 
 	bool bSkipInvalidate = FParse::Param(FCommandLine::Get(), TEXT("AssetRegistryCacheSkipInvalidate"));
-	FString AssetRegistryDiscoveryWriteCacheStr;
-	GConfig->GetString(TEXT("AssetRegistry"), TEXT("AssetRegistryDiscoveryWriteCache"),
-		AssetRegistryDiscoveryWriteCacheStr, GEngineIni);
-	FParse::Value(FCommandLine::Get(), TEXT("AssetRegistryDiscoveryWriteCache="), AssetRegistryDiscoveryWriteCacheStr);
+	FString AssetRegistryDiscoveryCacheStr;
+	GConfig->GetString(TEXT("AssetRegistry"), TEXT("AssetRegistryDiscoveryCache"),
+		AssetRegistryDiscoveryCacheStr, GEngineIni);
+	FParse::Value(FCommandLine::Get(), TEXT("AssetRegistryDiscoveryCache="), AssetRegistryDiscoveryCacheStr);
+
+	bNoAssetRegistryDiscoveryCache |= AssetRegistryDiscoveryCacheStr == TEXT("Never") || AssetRegistryDiscoveryCacheStr == TEXT("false")
+		|| AssetRegistryDiscoveryCacheStr == TEXT("0");
 
 	bDiscoveryCacheReadEnabled = (bPlatformSupportsDiscoveryCache || bSkipInvalidate) && !bNoAssetRegistryDiscoveryCache
 		&& !bNoAssetRegistryCacheRead;
-	if (bNoAssetRegistryDiscoveryCache || bNoAssetRegistryCacheWrite || bMultiprocess ||
-		AssetRegistryDiscoveryWriteCacheStr == TEXT("Never") || AssetRegistryDiscoveryWriteCacheStr == TEXT("false")
-		|| AssetRegistryDiscoveryWriteCacheStr == TEXT("0"))
+
+	if (bNoAssetRegistryDiscoveryCache || bNoAssetRegistryCacheWrite || bMultiprocess)
 	{
 		DiscoveryCacheWriteEnabled = EFeatureEnabled::Never;
 	}
-	else if ((AssetRegistryDiscoveryWriteCacheStr.IsEmpty() || AssetRegistryDiscoveryWriteCacheStr == TEXT("Default"))
+	else if ((AssetRegistryDiscoveryCacheStr.IsEmpty() || AssetRegistryDiscoveryCacheStr == TEXT("Default"))
 		&& !bSkipInvalidate)
 	{
 		// Precalculate IfPlatformSupported -> Never if we already know the platform doesn't support it
