AssetRegistry: Fix AssetDataGathererDiscoveryCache's failure to recognize new files when the project directory is on a different drive than UnrealEngine.exe. Refactor the settings parsing to allow values of Never,Default,AlwaysWrite.
#rnx

diff --git a/Engine/Config/BaseEngine.ini b/Engine/Config/BaseEngine.ini
--- a/Engine/Config/BaseEngine.ini
+++ b/Engine/Config/BaseEngine.ini
@@ -2646,12 +2646,12 @@ bFilterSearchableNames=true
 +CommandletsUsingAR=GenerateNaniteDisplacedMesh
 
 ; AssetDataGathering can be cached in the editor; there are two different caches, one for the directory scan to find all
-; packages and one for tags and other data read from the packages. The directory scan cache by default is only enabled if
-; the current machine supports FileJournal (IPlatformFile::FileJournalIsAvailable), but the cache can also
-; be used on-demand without checking it for out-of-date via commandline -SkipAssetRegistryCacheInvalidate.
-; To support that use case set AssetRegistryDiscoveryCache=Always to always create the cache.
-; Options are Never,Default,Always. Default only creates the cache if the platform supports invalidation via FileJournal.
-AssetRegistryDiscoveryCache=Never
+; packages and one for tags and other data read from the packages. The directory scan cache by default is only read
+; or written if the current machine supports FileJournal (IPlatformFile::FileJournalIsAvailable). It can be force-written
+; by setting AssetRegistryDiscoveryCache=AlwaysWrite, and it can be force-read with no test for invalidation due to modified
+; directories by commandline -AssetRegistryCacheSkipInvalidate.
+; Expected values for AssetRegistryDiscoveryCache are Never,Default,AlwaysWrite.
+AssetRegistryDiscoveryCache=Default
 
 [AutomationTesting]
 ImportTestPath=Content/EditorAutomation/
diff --git a/Engine/Source/Runtime/AssetRegistry/Private/AssetDataGatherer.cpp b/Engine/Source/Runtime/AssetRegistry/Private/AssetDataGatherer.cpp
--- a/Engine/Source/Runtime/AssetRegistry/Private/AssetDataGatherer.cpp
+++ b/Engine/Source/Runtime/AssetRegistry/Private/AssetDataGatherer.cpp
@@ -78,6 +78,86 @@ static FAutoConsoleVariableRef CVarIgnoreEmptyDirectories(
 	bIgnoreEmptyDirectories,
 	TEXT("If true, completely empty leaf directories are ignored by the asset registry while scanning"));
 
+void LexFromString(EFeatureEnabledReadWrite& OutValue, FStringView Text)
+{
+	Text.TrimStartAndEndInline();
+
+	if (Text.Equals(TEXTVIEW("NeverWriteNeverRead"), ESearchCase::IgnoreCase) ||
+		Text.Equals(TEXTVIEW("Never"), ESearchCase::IgnoreCase))
+	{
+		OutValue = EFeatureEnabledReadWrite::NeverWriteNeverRead;
+		return;
+	}
+	if (Text.Equals(TEXTVIEW("NeverWriteDefaultRead"), ESearchCase::IgnoreCase))
+	{
+		OutValue = EFeatureEnabledReadWrite::NeverWriteDefaultRead;
+		return;
+	}
+	if (Text.Equals(TEXTVIEW("NeverWriteAlwaysRead"), ESearchCase::IgnoreCase))
+	{
+		OutValue = EFeatureEnabledReadWrite::NeverWriteAlwaysRead;
+		return;
+	}
+
+	if (Text.Equals(TEXTVIEW("DefaultWriteNeverRead"), ESearchCase::IgnoreCase))
+	{
+		OutValue = EFeatureEnabledReadWrite::DefaultWriteNeverRead;
+		return;
+	}
+	if (Text.Equals(TEXTVIEW("DefaultWriteDefaultRead"), ESearchCase::IgnoreCase) ||
+		Text.Equals(TEXTVIEW("Default"), ESearchCase::IgnoreCase))
+	{
+		OutValue = EFeatureEnabledReadWrite::DefaultWriteDefaultRead;
+		return;
+	}
+	if (Text.Equals(TEXTVIEW("DefaultWriteAlwaysRead"), ESearchCase::IgnoreCase))
+	{
+		OutValue = EFeatureEnabledReadWrite::DefaultWriteAlwaysRead;
+		return;
+	}
+
+	if (Text.Equals(TEXTVIEW("AlwaysWriteNeverRead"), ESearchCase::IgnoreCase))
+	{
+		OutValue = EFeatureEnabledReadWrite::AlwaysWriteNeverRead;
+		return;
+	}
+	if (Text.Equals(TEXTVIEW("AlwaysWriteDefaultRead"), ESearchCase::IgnoreCase) ||
+		Text.Equals(TEXTVIEW("AlwaysWrite"), ESearchCase::IgnoreCase))
+	{
+		OutValue = EFeatureEnabledReadWrite::AlwaysWriteDefaultRead;
+		return;
+	}
+	if (Text.Equals(TEXTVIEW("AlwaysWriteAlwaysRead"), ESearchCase::IgnoreCase))
+	{
+		OutValue = EFeatureEnabledReadWrite::AlwaysWriteAlwaysRead;
+		return;
+	}
+	if (Text.Equals(TEXTVIEW("false"), ESearchCase::IgnoreCase) ||
+		Text.Equals(TEXTVIEW("f"), ESearchCase::IgnoreCase) ||
+		Text.Equals(TEXTVIEW("off"), ESearchCase::IgnoreCase) ||
+		Text.Equals(TEXTVIEW("0"), ESearchCase::IgnoreCase))
+	{
+		OutValue = EFeatureEnabledReadWrite::NeverWriteNeverRead;
+		return;
+	}
+	if (Text.Equals(TEXTVIEW("true"), ESearchCase::IgnoreCase) ||
+		Text.Equals(TEXTVIEW("t"), ESearchCase::IgnoreCase) ||
+		Text.Equals(TEXTVIEW("on"), ESearchCase::IgnoreCase))
+	{
+		OutValue = EFeatureEnabledReadWrite::DefaultWriteDefaultRead;
+		return;
+	}
+	uint32 IntValue = 0;
+	LexFromString(IntValue, Text);
+	if (IntValue != 0)
+	{
+		OutValue = EFeatureEnabledReadWrite::DefaultWriteDefaultRead;
+		return;
+	}
+	OutValue = EFeatureEnabledReadWrite::Invalid;
+	return;
+}
+
 void FPreloadSettings::Initialize()
 {
 	if (bInitialized)
@@ -94,41 +174,80 @@ void FPreloadSettings::Initialize()
 	bool bNoAssetRegistryCache = FParse::Param(FCommandLine::Get(), TEXT("NoAssetRegistryCache"));
 	bool bNoAssetRegistryDiscoveryCache = bNoAssetRegistryCache || FParse::Param(FCommandLine::Get(), TEXT("NoAssetRegistryDiscoveryCache"));
 	bool bNoAssetRegistryCacheRead = FParse::Param(FCommandLine::Get(), TEXT("NoAssetRegistryCacheRead"));
-	bool bNoAssetRegistryCacheWrite = FParse::Param(FCommandLine::Get(), TEXT("NoAssetRegistryCacheWrite"));
 	uint32 MultiprocessId = UE::GetMultiprocessId();
 	bool bMultiprocess = MultiprocessId > 0 || FParse::Param(FCommandLine::Get(), TEXT("multiprocess"));
+	bool bNoAssetRegistryCacheWrite = FParse::Param(FCommandLine::Get(), TEXT("NoAssetRegistryCacheWrite")) || bMultiprocess;
 	bGatherCacheReadEnabled = !bNoAssetRegistryCache && !bNoAssetRegistryCacheRead;
-	bGatherCacheWriteEnabled = !bNoAssetRegistryCache && !bNoAssetRegistryCacheWrite && !bMultiprocess;
-	bool bPlatformSupportsDiscoveryCache = FPlatformFileManager::Get().GetPlatformFile().FileJournalIsAvailable();
+	bGatherCacheWriteEnabled = !bNoAssetRegistryCache && !bNoAssetRegistryCacheWrite;
+	bool bPlatformSupportsDiscoveryCacheInvalidation = FPlatformFileManager::Get().GetPlatformFile().FileJournalIsAvailable();
 
 	bool bSkipInvalidate = FParse::Param(FCommandLine::Get(), TEXT("AssetRegistryCacheSkipInvalidate"));
-	FString AssetRegistryDiscoveryCacheStr;
+	EFeatureEnabledReadWrite DiscoverySetting = EFeatureEnabledReadWrite::Invalid;
+	FString AssetRegistryDiscoveryCacheStr = TEXT("Default");
 	GConfig->GetString(TEXT("AssetRegistry"), TEXT("AssetRegistryDiscoveryCache"),
 		AssetRegistryDiscoveryCacheStr, GEngineIni);
 	FParse::Value(FCommandLine::Get(), TEXT("AssetRegistryDiscoveryCache="), AssetRegistryDiscoveryCacheStr);
-
-	bNoAssetRegistryDiscoveryCache |= AssetRegistryDiscoveryCacheStr == TEXT("Never") || AssetRegistryDiscoveryCacheStr == TEXT("false")
-		|| AssetRegistryDiscoveryCacheStr == TEXT("0");
-
-	bDiscoveryCacheReadEnabled = (bPlatformSupportsDiscoveryCache || bSkipInvalidate) && !bNoAssetRegistryDiscoveryCache
-		&& !bNoAssetRegistryCacheRead;
-
-	if (bNoAssetRegistryDiscoveryCache || bNoAssetRegistryCacheWrite || bMultiprocess)
+	LexFromString(DiscoverySetting, AssetRegistryDiscoveryCacheStr);
+	if (DiscoverySetting == EFeatureEnabledReadWrite::Invalid)
 	{
-		DiscoveryCacheWriteEnabled = EFeatureEnabled::Never;
+		UE_LOG(LogAssetRegistry, Error,
+			TEXT("Invalid text \"%s\" for Engine.ini:[AssetRegistry]:AssetRegistryDiscoveryCache. Expected \"Never\", \"Default\", or \"AlwaysWrite\"."),
+			*AssetRegistryDiscoveryCacheStr);
+		DiscoverySetting = EFeatureEnabledReadWrite::DefaultWriteDefaultRead;
 	}
-	else if ((AssetRegistryDiscoveryCacheStr.IsEmpty() || AssetRegistryDiscoveryCacheStr == TEXT("Default"))
-		&& !bSkipInvalidate)
+	if (bNoAssetRegistryDiscoveryCache)
 	{
-		// Precalculate IfPlatformSupported -> Never if we already know the platform doesn't support it
-		DiscoveryCacheWriteEnabled = bPlatformSupportsDiscoveryCache ?
-			EFeatureEnabled::IfPlatformSupported : EFeatureEnabled::Never;
+		DiscoverySetting = EFeatureEnabledReadWrite::NeverWriteNeverRead;
 	}
-	else
+	if (bNoAssetRegistryCacheRead)
+	{
+		DiscoverySetting = (DiscoverySetting & ~EFeatureEnabledReadWrite::ReadMask) | EFeatureEnabledReadWrite::NeverRead;
+	}
+	if (bNoAssetRegistryCacheWrite)
 	{
+		DiscoverySetting = (DiscoverySetting & ~EFeatureEnabledReadWrite::WriteMask) | EFeatureEnabledReadWrite::NeverWrite;
+	}
+	if (bSkipInvalidate || MultiprocessId != 0)
+	{
+		if ((DiscoverySetting & EFeatureEnabledReadWrite::WriteMask) == EFeatureEnabledReadWrite::DefaultWrite)
+		{
+			DiscoverySetting = (DiscoverySetting & ~EFeatureEnabledReadWrite::WriteMask) | EFeatureEnabledReadWrite::AlwaysWrite;
+		}
+		if ((DiscoverySetting & EFeatureEnabledReadWrite::ReadMask) == EFeatureEnabledReadWrite::DefaultRead)
+		{
+			DiscoverySetting = (DiscoverySetting & ~EFeatureEnabledReadWrite::ReadMask) | EFeatureEnabledReadWrite::AlwaysRead;
+		}
+	}
+	else if (!bPlatformSupportsDiscoveryCacheInvalidation)
+	{
+		// Precalculate Default -> Never if we already know the platform doesn't support it
+		if ((DiscoverySetting & EFeatureEnabledReadWrite::WriteMask) == EFeatureEnabledReadWrite::DefaultWrite)
+		{
+			DiscoverySetting = (DiscoverySetting& ~EFeatureEnabledReadWrite::WriteMask) | EFeatureEnabledReadWrite::NeverWrite;
+		}
+		if ((DiscoverySetting & EFeatureEnabledReadWrite::ReadMask) == EFeatureEnabledReadWrite::DefaultRead)
+		{
+			DiscoverySetting = (DiscoverySetting& ~EFeatureEnabledReadWrite::ReadMask) | EFeatureEnabledReadWrite::NeverRead;
+		}
+	}
+	bDiscoveryCacheReadEnabled = (DiscoverySetting & EFeatureEnabledReadWrite::ReadMask) != EFeatureEnabledReadWrite::NeverRead;
+	bDiscoveryCacheInvalidateEnabled = (DiscoverySetting & EFeatureEnabledReadWrite::ReadMask) != EFeatureEnabledReadWrite::AlwaysRead;
+	switch (DiscoverySetting & EFeatureEnabledReadWrite::WriteMask)
+	{
+	case EFeatureEnabledReadWrite::NeverWrite:
+		DiscoveryCacheWriteEnabled = EFeatureEnabled::Never;
+		break;
+	case EFeatureEnabledReadWrite::DefaultWrite:
+		DiscoveryCacheWriteEnabled = EFeatureEnabled::IfPlatformSupported;
+		break;
+	case EFeatureEnabledReadWrite::AlwaysWrite:
 		DiscoveryCacheWriteEnabled = EFeatureEnabled::Always;
+		break;
+	default:
+		checkNoEntry();
+		DiscoveryCacheWriteEnabled = EFeatureEnabled::Never;
+		break;
 	}
-	bDiscoveryCacheInvalidateEnabled = !bSkipInvalidate && MultiprocessId == 0;
 
 	bool bAsyncEnabled = FPlatformProcess::SupportsMultithreading() && FTaskGraphInterface::IsRunning();
 
diff --git a/Engine/Source/Runtime/AssetRegistry/Private/AssetDataGathererDiscoveryCache.h b/Engine/Source/Runtime/AssetRegistry/Private/AssetDataGathererDiscoveryCache.h
--- a/Engine/Source/Runtime/AssetRegistry/Private/AssetDataGathererDiscoveryCache.h
+++ b/Engine/Source/Runtime/AssetRegistry/Private/AssetDataGathererDiscoveryCache.h
@@ -23,6 +23,31 @@ enum class EFeatureEnabled : uint8
 	Always
 };
 
+enum class EFeatureEnabledReadWrite : uint32
+{
+	NeverRead		= 0x00,
+	DefaultRead		= 0x01,
+	AlwaysRead		= 0x02,
+	ReadMask		= 0x0f,
+	NeverWrite		= 0x00,
+	DefaultWrite	= 0x10,
+	AlwaysWrite		= 0x20,
+	WriteMask		= 0xf0,
+	Invalid			= 0xffffffff,
+	NeverWriteNeverRead		= NeverWrite   | NeverRead,   // aka Never
+	NeverWriteDefaultRead	= NeverWrite   | DefaultRead,
+	NeverWriteAlwaysRead	= NeverWrite   | AlwaysRead,
+	DefaultWriteNeverRead	= DefaultWrite | NeverRead,   
+	DefaultWriteDefaultRead = DefaultWrite | DefaultRead, // aka Default
+	DefaultWriteAlwaysRead	= DefaultWrite | AlwaysRead,
+	AlwaysWriteNeverRead	= AlwaysWrite  | NeverRead,
+	AlwaysWriteDefaultRead	= AlwaysWrite  | DefaultRead, // aka AlwaysWrite
+	// NOT given the nickname "Always", because skipping invalidation is not what users would expect from "Always"
+	AlwaysWriteAlwaysRead	= AlwaysWrite  | AlwaysRead,
+};
+ENUM_CLASS_FLAGS(EFeatureEnabledReadWrite);
+void LexFromString(EFeatureEnabledReadWrite& OutValue, FStringView Text);
+
 /** Data about a file in a cached directory used by FAssetDataDiscoveryCache to avoid the IO cost of rescanning. */
 struct FCachedDirScanFile
 {
diff --git a/Engine/Source/Runtime/Core/Private/Windows/WindowsPlatformFileJournal.cpp b/Engine/Source/Runtime/Core/Private/Windows/WindowsPlatformFileJournal.cpp
--- a/Engine/Source/Runtime/Core/Private/Windows/WindowsPlatformFileJournal.cpp
+++ b/Engine/Source/Runtime/Core/Private/Windows/WindowsPlatformFileJournal.cpp
@@ -56,15 +56,24 @@
 
 FString FWindowsPlatformFile::FileJournalGetVolumeName(FStringView Path)
 {
-	FString FullPath = FPaths::ConvertRelativePathToFull(FString(Path));
-	FStringView VolumeName;
-	FStringView Remainder;
-	FPathViews::SplitVolumeSpecifier(FullPath, VolumeName, Remainder);
-	if (VolumeName.EndsWith(TEXT(":")))
+	FString FullPath;
+	if (FPathViews::IsDriveSpecifierWithoutRoot(Path))
 	{
-		return FString(VolumeName);
+		// Already a drive specifier, and ConvertRelativePathToFull does not handle a rootless drivespecifier
+		return FString(Path);
+	}
+	else
+	{
+		FullPath = FPaths::ConvertRelativePathToFull(FString(Path));
+		FStringView VolumeName;
+		FStringView Remainder;
+		FPathViews::SplitVolumeSpecifier(FullPath, VolumeName, Remainder);
+		if (VolumeName.EndsWith(TEXT(":")))
+		{
+			return FString(VolumeName);
+		}
+		return FString(); // Volume names of the form \\volumename are not supported by USNJournal
 	}
-	return FString(); // Volume names of the form \\volumename are not supported by USNJournal
 }
 
 #if UE_WINDOWS_PLATFORM_FILEJOURNAL_ENABLED 
@@ -519,7 +528,7 @@ EFileJournalResult FWindowsPlatformFile::FileJournalReadModified(const TCHAR* Vo
 		static_assert(sizeof(uint64) >= sizeof(USN), "We are storing USNs as uint64");
 		OutNextJournalEntry = static_cast<uint64>(NextJournalUSN);
 
-		// Keep going until we reach the end, which we detect by  Journal did not have enough results to fill the buffer
+		// Keep going until we reach the end, which we detect by Journal did not have enough results to fill the buffer
 		if (RecordBytesEnd <= ReadBuffer + ReadBufferSize / 2)
 		{
 			break;
@@ -646,19 +655,27 @@ EFileJournalResult CreateJournalWindowsHandle(const TCHAR* VolumeOrPath, HANDLE&
 	OutVolumeName.Reset();
 
 	TStringBuilder<256> AbsPath;
-	FPathViews::ToAbsolutePath(VolumeOrPath, AbsPath);
 	FStringView VolumeName;
-	FStringView Remainder;
-	FPathViews::SplitVolumeSpecifier(AbsPath, VolumeName, Remainder);
-	if (!VolumeName.EndsWith(':'))
+	if (FPathViews::IsDriveSpecifierWithoutRoot(VolumeOrPath))
 	{
-		if (OutError)
+		// Already a drive specifier, and ToAbsolutePath does not handle a rootless drivespecifier
+		VolumeName = VolumeOrPath;
+	}
+	else
+	{
+		FPathViews::ToAbsolutePath(VolumeOrPath, AbsPath);
+		FStringView Remainder;
+		FPathViews::SplitVolumeSpecifier(AbsPath, VolumeName, Remainder);
+		if (!VolumeName.EndsWith(':'))
 		{
-			*OutError = FString::Printf(
-				TEXT("Invalid drive name in path '%s'. FileJournal is only supported for drives assigned to a drive letter (e.g. C:, D:)."),
-				*AbsPath);
+			if (OutError)
+			{
+				*OutError = FString::Printf(
+					TEXT("Invalid drive name in path '%s'. FileJournal is only supported for drives assigned to a drive letter (e.g. C:, D:)."),
+					*AbsPath);
+			}
+			return EFileJournalResult::InvalidVolumeName;
 		}
-		return EFileJournalResult::InvalidVolumeName;
 	}
 	OutVolumeName << VolumeName;
 
