AssetRegistry: ScanPathsSynchronous should remove old assets that no longer exist.
#rnx

diff --git a/Engine/Source/Runtime/AssetRegistry/Private/AssetRegistry.cpp b/Engine/Source/Runtime/AssetRegistry/Private/AssetRegistry.cpp
--- a/Engine/Source/Runtime/AssetRegistry/Private/AssetRegistry.cpp
+++ b/Engine/Source/Runtime/AssetRegistry/Private/AssetRegistry.cpp
@@ -5062,17 +5062,54 @@ void FAssetRegistryImpl::ScanPathsSynchronous(Impl::FScanPathContext& Context)
 		Gatherer.LoadCacheFiles({CacheFilename});
 	}
 
+	// If we are forcing a rescan, then delete any old assets that no longer exist. If we are not forcing a rescan,
+	// then there should not be any old assets that no longer exist, so we skip the cost of searching for them.
+	TSet<FSoftObjectPath> OldAssetsToRemove;
+	if (Context.bForceRescan)
+	{
+		// Initialize OldAssetsToRemove to the list of all assets in the given paths.
+		if (!Context.PackageDirs.IsEmpty())
+		{
+			FARFilter Filter;
+			Filter.bIncludeOnlyOnDiskAssets = true;
+			Filter.bRecursivePaths = true;
+			for (const FString& PackageDir : Context.PackageDirs)
+			{
+				Filter.PackagePaths.Add(FName(*PackageDir));
+			}
+			FARCompiledFilter CompiledFilter;
+			CompileFilter(Context.InheritanceContext, Filter, CompiledFilter);
+			TArray<FAssetData> AssetsInPaths;
+			State.EnumerateAssets(CompiledFilter, TSet<FName>() /* PackageNamesToSkip */,
+				[&OldAssetsToRemove](const FAssetData& AssetData)
+				{
+					OldAssetsToRemove.Add(AssetData.ToSoftObjectPath());
+					return true;
+				});
+		}
+		for (const FString& PackageName : Context.PackageFiles)
+		{
+			for (const FAssetData* AssetData : State.GetAssetsByPackageName(FName(*PackageName)))
+			{
+				OldAssetsToRemove.Add(AssetData->ToSoftObjectPath());
+			}
+		}
+	}
+
 	Gatherer.ScanPathsSynchronous(Context.LocalPaths, Context.bForceRescan, Context.bIgnoreDenyListScanFilters, CacheFilename, Context.PackageDirs);
 	TArray<FName> FoundAssetPackageNames;
 
-	auto AssetsFoundCallback = [&Context, &FoundAssetPackageNames, this](const TMultiMap<FName, FAssetData*>& InFoundAssets)
+	auto AssetsFoundCallback =
+		[&Context, &FoundAssetPackageNames, &OldAssetsToRemove, this]
+		(const TMultiMap<FName, FAssetData*>& InFoundAssets)
 	{
 		Context.NumFoundAssets = InFoundAssets.Num();
 
 		FoundAssetPackageNames.Reset();
 		FoundAssetPackageNames.Reserve(Context.NumFoundAssets);
 
-		// The gatherer may have added other assets that were scanned as part of the ongoing background scan; remove any assets that were not in the requested paths
+		// The gatherer may have added other assets that were scanned as part of the ongoing background scan,
+		// so remove any assets that were not in the requested paths
 		for (const TPair<FName,FAssetData*>& Pair : InFoundAssets)
 		{
 			FAssetData* AssetData = Pair.Value;
@@ -5113,6 +5150,11 @@ void FAssetRegistryImpl::ScanPathsSynchronous(Impl::FScanPathContext& Context)
 				}
 				FoundAssetPackageNames.Add(AssetData->PackageName);
 			}
+
+			if (!OldAssetsToRemove.IsEmpty())
+			{
+				OldAssetsToRemove.Remove(AssetData->ToSoftObjectPath());
+			}
 		}
 	};
 
@@ -5122,6 +5164,13 @@ void FAssetRegistryImpl::ScanPathsSynchronous(Impl::FScanPathContext& Context)
 #if WITH_EDITOR
 	LoadCalculatedDependencies(&FoundAssetPackageNames, -1., Context.InheritanceContext, bUnusedInterrupted);
 #endif
+	for (FSoftObjectPath& OldAssetToRemove : OldAssetsToRemove)
+	{
+		bool bOutRemovedAssetData;
+		bool bOutRemovedPackageData;
+		State.RemoveAssetData(OldAssetToRemove, true /* bRemoveDependencyData */,
+			bOutRemovedAssetData, bOutRemovedPackageData);
+	}
 }
 
 namespace Utils
diff --git a/Engine/Source/Runtime/AssetRegistry/Private/AssetRegistryState.cpp b/Engine/Source/Runtime/AssetRegistry/Private/AssetRegistryState.cpp
--- a/Engine/Source/Runtime/AssetRegistry/Private/AssetRegistryState.cpp
+++ b/Engine/Source/Runtime/AssetRegistry/Private/AssetRegistryState.cpp
@@ -2281,22 +2281,46 @@ bool FAssetRegistryState::UpdateAssetDataPackageFlags(FName PackageName, uint32
 	return false;
 }
 
-void FAssetRegistryState::RemoveAssetData(FAssetData* AssetData, bool bRemoveDependencyData, bool& bOutRemovedAssetData, bool& bOutRemovedPackageData)
+void FAssetRegistryState::RemoveAssetData(FAssetData* AssetData, bool bRemoveDependencyData,
+	bool& bOutRemovedAssetData, bool& bOutRemovedPackageData)
 {
-	bOutRemovedAssetData = false;
-	bOutRemovedPackageData = false;
-
 	if (!ensure(AssetData))
 	{
+		bOutRemovedAssetData = false;
+		bOutRemovedPackageData = false;
 		return;
 	}
+	RemoveAssetData(AssetData, FCachedAssetKey(AssetData), bRemoveDependencyData, bOutRemovedAssetData, bOutRemovedPackageData);
+	if (!bOutRemovedAssetData)
+	{
+		UE_LOG(LogAssetRegistry, Error, TEXT("RemoveAssetData called on AssetData %s that is not present in the AssetRegistry."),
+			*FCachedAssetKey(*AssetData).ToString());
+	}
+}
 
-	int32 NumRemoved = CachedAssets.Remove(FCachedAssetKey(AssetData));
+void FAssetRegistryState::RemoveAssetData(const FSoftObjectPath& AssetPath, bool bRemoveDependencyData,
+	bool& bOutRemovedAssetData, bool& bOutRemovedPackageData)
+{
+	FCachedAssetKey Key(AssetPath);
+	FAssetData** AssetDataPtrPtr = CachedAssets.Find(Key);
+	if (!AssetDataPtrPtr)
+	{
+		bOutRemovedAssetData = false;
+		bOutRemovedPackageData = false;
+		return;
+	}
+	RemoveAssetData(*AssetDataPtrPtr, Key, bRemoveDependencyData, bOutRemovedAssetData, bOutRemovedPackageData);
+}
+
+void FAssetRegistryState::RemoveAssetData(FAssetData* AssetData, const FCachedAssetKey& Key,
+	bool bRemoveDependencyData, bool& bOutRemovedAssetData, bool& bOutRemovedPackageData)
+{
+	bOutRemovedAssetData = false;
+	bOutRemovedPackageData = false;
+	int32 NumRemoved = CachedAssets.Remove(Key);
 	check(NumRemoved <= 1);
 	if (NumRemoved == 0)
 	{
-		UE_LOG(LogAssetRegistry, Error, TEXT("RemoveAssetData called on AssetData %s that is not present in the AssetRegistry."),
-			*FCachedAssetKey(*AssetData).ToString());
 		return;
 	}
 
diff --git a/Engine/Source/Runtime/AssetRegistry/Public/AssetRegistry/AssetRegistryState.h b/Engine/Source/Runtime/AssetRegistry/Public/AssetRegistry/AssetRegistryState.h
--- a/Engine/Source/Runtime/AssetRegistry/Public/AssetRegistry/AssetRegistryState.h
+++ b/Engine/Source/Runtime/AssetRegistry/Public/AssetRegistry/AssetRegistryState.h
@@ -469,7 +469,10 @@ public:
 	ASSETREGISTRY_API bool UpdateAssetDataPackageFlags(FName PackageName, uint32 PackageFlags);
 
 	/** Removes the asset data from the lookup maps */
-	ASSETREGISTRY_API void RemoveAssetData(FAssetData* AssetData, bool bRemoveDependencyData, bool& bOutRemovedAssetData, bool& bOutRemovedPackageData);
+	ASSETREGISTRY_API void RemoveAssetData(FAssetData* AssetData, bool bRemoveDependencyData,
+		bool& bOutRemovedAssetData, bool& bOutRemovedPackageData);
+	ASSETREGISTRY_API void RemoveAssetData(const FSoftObjectPath& SoftObjectPath, bool bRemoveDependencyData,
+		bool& bOutRemovedAssetData, bool& bOutRemovedPackageData);
 
 	/**
 	 * Clear all dependencies of the given category from the given AssetIdentifier (e.g. package).
@@ -606,6 +609,9 @@ private:
 
 	void SetDependencyNodeSorting(bool bSortDependencies, bool bSortReferencers);
 
+	void RemoveAssetData(FAssetData* AssetData, const FCachedAssetKey& Key, bool bRemoveDependencyData,
+		bool& bOutRemovedAssetData, bool& bOutRemovedPackageData);
+
 	/** Set of asset data for assets saved to disk. Searched via path name types, implicitly converted to FCachedAssetKey. */
 	FAssetDataMap CachedAssets;
 
