AssetRegistry: Fix error "RemoveAssetData called on AssetData %s that is not present in the AssetRegistry" when any files are modified in editor.
#rnx

diff --git a/Engine/Source/Runtime/AssetRegistry/Private/AssetRegistry.cpp b/Engine/Source/Runtime/AssetRegistry/Private/AssetRegistry.cpp
--- a/Engine/Source/Runtime/AssetRegistry/Private/AssetRegistry.cpp
+++ b/Engine/Source/Runtime/AssetRegistry/Private/AssetRegistry.cpp
@@ -6496,18 +6496,18 @@ void FAssetRegistryImpl::ScanModifiedAssetFiles(Impl::FEventContext& EventContex
 		}
 
 		// Get the assets that are currently inside the package
-		TArray<TArray<FAssetData*, TInlineAllocator<1>>> ExistingFilesAssetData;
-		ExistingFilesAssetData.Reserve(InFilePaths.Num());
+		TArray<FSoftObjectPath> ExistingAssetDatas;
+		ExistingAssetDatas.Reserve(InFilePaths.Num());
 		for (const FString& PackageName : ModifiedPackageNames)
 		{
 			TArray<FAssetData*, TInlineAllocator<1>>* PackageAssetsPtr = State.CachedAssetsByPackageName.Find(*PackageName);
 			if (PackageAssetsPtr && PackageAssetsPtr->Num() > 0)
 			{
-				ExistingFilesAssetData.Add(*PackageAssetsPtr);
-			}
-			else
-			{
-				ExistingFilesAssetData.AddDefaulted();
+				ExistingAssetDatas.Reserve(ExistingAssetDatas.Num() + PackageAssetsPtr->Num());
+				for (FAssetData* AssetData : *PackageAssetsPtr)
+				{
+					ExistingAssetDatas.Add(AssetData->ToSoftObjectPath());
+				}
 			}
 		}
 
@@ -6518,13 +6518,14 @@ void FAssetRegistryImpl::ScanModifiedAssetFiles(Impl::FEventContext& EventContex
 		ScanPathsSynchronous(Context);
 
 		// Remove any assets that are no longer present in the package
-		for (const TArray<FAssetData*, TInlineAllocator<1>>& OldPackageAssets : ExistingFilesAssetData)
+		for (FSoftObjectPath& OldAssetPath : ExistingAssetDatas)
 		{
-			for (FAssetData* OldPackageAsset : OldPackageAssets)
+			if (!FoundAssets.Contains(OldAssetPath))
 			{
-				if (!FoundAssets.Contains(OldPackageAsset->GetSoftObjectPath()))
+				FAssetData* OldAssetData = const_cast<FAssetData*>(State.GetAssetByObjectPath(OldAssetPath));
+				if (OldAssetData)
 				{
-					RemoveAssetData(EventContext, OldPackageAsset);
+					RemoveAssetData(EventContext, OldAssetData);
 				}
 			}
 		}
