AssetRegistry: Change priority of FAssetDataDiscovery's parallelfor (used to scan directories) to BackgroundPriority, to reduce contention with other async work in editor startup. This makes time to complete AR scan slower, from 3.5 minutes to 3.75 minutes, but makes editor more responsive.
#rnx

diff --git a/Engine/Source/Runtime/AssetRegistry/Private/AssetDataGatherer.cpp b/Engine/Source/Runtime/AssetRegistry/Private/AssetDataGatherer.cpp
--- a/Engine/Source/Runtime/AssetRegistry/Private/AssetDataGatherer.cpp
+++ b/Engine/Source/Runtime/AssetRegistry/Private/AssetDataGatherer.cpp
@@ -2095,7 +2095,7 @@ void FAssetDataDiscovery::TickInternal(bool bTickAll)
 			}
 
 			Data.bScanned = true;
-		});
+		}, EParallelForFlags::BackgroundPriority);
 
 		if (Cache.IsWriteEnabled() != EFeatureEnabled::Never)
 		{
