Improve cook speed of cooked editors by increasing the thread priority of the asset data gatherer when running a commandlet.
Many commandlets are blocking on the asset registry so we should be running those with the fastest threads possible.

diff --git a/Engine/Source/Runtime/AssetRegistry/Private/AssetDataGatherer.cpp b/Engine/Source/Runtime/AssetRegistry/Private/AssetDataGatherer.cpp
--- a/Engine/Source/Runtime/AssetRegistry/Private/AssetDataGatherer.cpp
+++ b/Engine/Source/Runtime/AssetRegistry/Private/AssetDataGatherer.cpp
@@ -1983,6 +1983,18 @@ void FAssetDataDiscovery::TickInternal(bool bTickAll)
 			ScanBuffer.Reset();
 		}
 
+		static const EParallelForFlags ParallelFlags = []()
+			{
+				if (IsRunningCommandlet())
+				{
+					return EParallelForFlags::None;
+				}
+				else
+				{
+					// commandlets or cooks will be blocking on the asset registry finish.
+					return EParallelForFlags::BackgroundPriority;
+				}
+			}();
 		ParallelForWithExistingTaskContext(LocalScanBuffers, DirToScanDatasNum, AssetDataGathererConstants::GARDiscoverMinBatchSize,
 		[this](FDirToScanBuffer& ScanBuffer, int32 DirToScanDatasIndex)
 		{
@@ -2095,7 +2107,7 @@ void FAssetDataDiscovery::TickInternal(bool bTickAll)
 			}
 
 			Data.bScanned = true;
-		}, EParallelForFlags::BackgroundPriority);
+		}, ParallelFlags);
 
 		if (Cache.IsWriteEnabled() != EFeatureEnabled::Never)
 		{
@@ -4109,6 +4121,15 @@ FAssetDataGatherer::ETickResult FAssetDataGatherer::TickInternal(double& TickSta
 	}
 
 	// For all the files not found in the cache, read them from their package files on disk; the file reads are done in parallel
+	static const EParallelForFlags ParallelFlags = []()
+		{
+			EParallelForFlags ReturnFlags = EParallelForFlags::Unbalanced;
+			if (!IsRunningCommandlet())
+			{
+				ReturnFlags |= EParallelForFlags::BackgroundPriority;
+			}
+			return ReturnFlags;
+		}();
 	ParallelFor(ReadContexts.Num(),
 		[this, &ReadContexts](int32 Index)
 		{
@@ -4121,7 +4142,7 @@ FAssetDataGatherer::ETickResult FAssetDataGatherer::TickInternal(double& TickSta
 			UE_SCOPED_IO_ACTIVITY(*WriteToString<512>(TEXT("Loading Asset"), ReadContext.PackageName.ToString()));
 			ReadContext.bResult = ReadAssetFile(ReadContext.AssetFileData.LongPackageName, ReadContext.AssetFileData.LocalAbsPath, ReadContext.AssetDataFromFile, ReadContext.DependencyData, ReadContext.CookedPackageNamesWithoutAssetData, ReadContext.bCanAttemptAssetRetry);
 		},
-		EParallelForFlags::Unbalanced | EParallelForFlags::BackgroundPriority
+		ParallelFlags
 	);
 
 	// Accumulate the results
@@ -4926,6 +4947,12 @@ TArray<FCachePayload> LoadCacheFiles(TConstArrayView<FString> InCacheFilenames,
 	TRACE_CPUPROFILER_EVENT_SCOPE(LoadCacheFiles);
 	TArray<FCachePayload> Results;
 	Results.AddDefaulted(InCacheFilenames.Num());
+
+	EParallelForFlags ParallelFlags = EParallelForFlags::None;
+	if (!IsRunningCommandlet())
+	{
+		ParallelFlags |= EParallelForFlags::BackgroundPriority;
+	}
 	ParallelFor(InCacheFilenames.Num(), [&InCacheFilenames, &Results, bIsMonolithicCache](int32 Index)
 	{
 		const FString& CacheFilename = InCacheFilenames[Index];
@@ -4980,7 +5007,7 @@ TArray<FCachePayload> LoadCacheFiles(TConstArrayView<FString> InCacheFilenames,
 		}
 	
 		Results[Index] = MoveTemp(Payload);
-	}, EParallelForFlags::BackgroundPriority);
+	}, ParallelFlags);
 
 	return MoveTemp(Results);
 }
