backout CL: 29991575

diff --git a/Engine/Source/Runtime/AssetRegistry/Private/AssetDataGatherer.cpp b/Engine/Source/Runtime/AssetRegistry/Private/AssetDataGatherer.cpp
--- a/Engine/Source/Runtime/AssetRegistry/Private/AssetDataGatherer.cpp
+++ b/Engine/Source/Runtime/AssetRegistry/Private/AssetDataGatherer.cpp
@@ -1983,18 +1983,6 @@ void FAssetDataDiscovery::TickInternal(bool bTickAll)
 			ScanBuffer.Reset();
 		}
 
-		static const EParallelForFlags ParallelFlags = []()
-			{
-				if (IsRunningCommandlet())
-				{
-					return EParallelForFlags::None;
-				}
-				else
-				{
-					// commandlets or cooks will be blocking on the asset registry finish.
-					return EParallelForFlags::BackgroundPriority;
-				}
-			}();
 		ParallelForWithExistingTaskContext(LocalScanBuffers, DirToScanDatasNum, AssetDataGathererConstants::GARDiscoverMinBatchSize,
 		[this](FDirToScanBuffer& ScanBuffer, int32 DirToScanDatasIndex)
 		{
@@ -2107,7 +2095,7 @@ void FAssetDataDiscovery::TickInternal(bool bTickAll)
 			}
 
 			Data.bScanned = true;
-		}, ParallelFlags);
+		}, EParallelForFlags::BackgroundPriority);
 
 		if (Cache.IsWriteEnabled() != EFeatureEnabled::Never)
 		{
@@ -4121,15 +4109,6 @@ FAssetDataGatherer::ETickResult FAssetDataGatherer::TickInternal(double& TickSta
 	}
 
 	// For all the files not found in the cache, read them from their package files on disk; the file reads are done in parallel
-	static const EParallelForFlags ParallelFlags = []()
-		{
-			EParallelForFlags ReturnFlags = EParallelForFlags::Unbalanced;
-			if (!IsRunningCommandlet())
-			{
-				ReturnFlags |= EParallelForFlags::BackgroundPriority;
-			}
-			return ReturnFlags;
-		}();
 	ParallelFor(ReadContexts.Num(),
 		[this, &ReadContexts](int32 Index)
 		{
@@ -4142,7 +4121,7 @@ FAssetDataGatherer::ETickResult FAssetDataGatherer::TickInternal(double& TickSta
 			UE_SCOPED_IO_ACTIVITY(*WriteToString<512>(TEXT("Loading Asset"), ReadContext.PackageName.ToString()));
 			ReadContext.bResult = ReadAssetFile(ReadContext.AssetFileData.LongPackageName, ReadContext.AssetFileData.LocalAbsPath, ReadContext.AssetDataFromFile, ReadContext.DependencyData, ReadContext.CookedPackageNamesWithoutAssetData, ReadContext.bCanAttemptAssetRetry);
 		},
-		ParallelFlags
+		EParallelForFlags::Unbalanced | EParallelForFlags::BackgroundPriority
 	);
 
 	// Accumulate the results
@@ -4947,12 +4926,6 @@ TArray<FCachePayload> LoadCacheFiles(TConstArrayView<FString> InCacheFilenames,
 	TRACE_CPUPROFILER_EVENT_SCOPE(LoadCacheFiles);
 	TArray<FCachePayload> Results;
 	Results.AddDefaulted(InCacheFilenames.Num());
-
-	EParallelForFlags ParallelFlags = EParallelForFlags::None;
-	if (!IsRunningCommandlet())
-	{
-		ParallelFlags |= EParallelForFlags::BackgroundPriority;
-	}
 	ParallelFor(InCacheFilenames.Num(), [&InCacheFilenames, &Results, bIsMonolithicCache](int32 Index)
 	{
 		const FString& CacheFilename = InCacheFilenames[Index];
@@ -5007,7 +4980,7 @@ TArray<FCachePayload> LoadCacheFiles(TConstArrayView<FString> InCacheFilenames,
 		}
 	
 		Results[Index] = MoveTemp(Payload);
-	}, ParallelFlags);
+	}, EParallelForFlags::BackgroundPriority);
 
 	return MoveTemp(Results);
 }
