Improve cook speed of cooked editors by increasing the thread priority of the asset data gatherer when running a commandlet.
Many commandlets are blocking on the asset registry so we should be running those with the fastest threads possible.

diff --git a/Engine/Source/Runtime/AssetRegistry/Private/AssetDataGatherer.cpp b/Engine/Source/Runtime/AssetRegistry/Private/AssetDataGatherer.cpp
--- a/Engine/Source/Runtime/AssetRegistry/Private/AssetDataGatherer.cpp
+++ b/Engine/Source/Runtime/AssetRegistry/Private/AssetDataGatherer.cpp
@@ -4109,6 +4109,15 @@ FAssetDataGatherer::ETickResult FAssetDataGatherer::TickInternal(double& TickSta
 	}
 
 	// For all the files not found in the cache, read them from their package files on disk; the file reads are done in parallel
+	static const EParallelForFlags ParallelFlags = []()
+		{
+			EParallelForFlags ReturnFlags = EParallelForFlags::Unbalanced;
+			if (!IsRunningCommandlet())
+			{
+				ReturnFlags |= EParallelForFlags::BackgroundPriority;
+			}
+			return ReturnFlags;
+		}();
 	ParallelFor(ReadContexts.Num(),
 		[this, &ReadContexts](int32 Index)
 		{
@@ -4121,7 +4130,7 @@ FAssetDataGatherer::ETickResult FAssetDataGatherer::TickInternal(double& TickSta
 			UE_SCOPED_IO_ACTIVITY(*WriteToString<512>(TEXT("Loading Asset"), ReadContext.PackageName.ToString()));
 			ReadContext.bResult = ReadAssetFile(ReadContext.AssetFileData.LongPackageName, ReadContext.AssetFileData.LocalAbsPath, ReadContext.AssetDataFromFile, ReadContext.DependencyData, ReadContext.CookedPackageNamesWithoutAssetData, ReadContext.bCanAttemptAssetRetry);
 		},
-		EParallelForFlags::Unbalanced | EParallelForFlags::BackgroundPriority
+		ParallelFlags
 	);
 
 	// Accumulate the results
