AssetRegistry: Fix fatal assert during GC without introducing deadlock. Add FGCScopeGuard in CollectCodeGeneratorClasses to prevent the fatal assert.
Change CollectCodeGeneratorClasses to be invalidated only when native classes change to prevent the deadlock possibility caused by the FGCScopeGuard. It is safe to invalidate on the reduced schedule because only native classes can be code generators (e.g. UBlueprint).
#rnx

diff --git a/Engine/Source/Runtime/AssetRegistry/Private/AssetRegistry.cpp b/Engine/Source/Runtime/AssetRegistry/Private/AssetRegistry.cpp
--- a/Engine/Source/Runtime/AssetRegistry/Private/AssetRegistry.cpp
+++ b/Engine/Source/Runtime/AssetRegistry/Private/AssetRegistry.cpp
@@ -867,7 +867,8 @@ void FAssetRegistryImpl::Initialize(Impl::FInitializeContext& Context)
 	bIsTempCachingEnabled = bIsTempCachingAlwaysEnabled;
 	TempCachedInheritanceBuffer.bDirty = true;
 
-	ClassGeneratorNamesRegisteredClassesVersionNumber = MAX_uint64;
+	SavedGeneratorClassesVersionNumber = MAX_uint64;
+	SavedAllClassesVersionNumber = MAX_uint64;
 
 	// By default do not double check mount points are still valid when gathering new assets
 	bVerifyMountPointAfterGather = false;
@@ -1392,28 +1393,52 @@ void InitializeSerializationOptionsFromIni(FAssetRegistrySerializationOptions& O
 
 }
 
+uint64 FAssetRegistryImpl::GetCurrentGeneratorClassesVersionNumber()
+{
+	// Generator classes can only be native, so we can use the less-frequently-updated
+	// RegisteredNativeClassesVersionNumber. In monolithic configurations, this will only be
+	// updated at program start and when enabling DLC modules.
+	return GetRegisteredNativeClassesVersionNumber();
+}
+
+uint64 FAssetRegistryImpl::GetCurrentAllClassesVersionNumber()
+{
+	return GetRegisteredClassesVersionNumber();
+}
+
 void FAssetRegistryImpl::CollectCodeGeneratorClasses()
 {
 	LLM_SCOPE(ELLMTag::AssetRegistry); // Tagged here instead of a higher level because it can occur even when reading
 	// Only refresh the list if our registered classes have changed
-	if (ClassGeneratorNamesRegisteredClassesVersionNumber == GetRegisteredClassesVersionNumber())
+	uint64 CurrentGeneratorClassesVersionNumber = GetCurrentGeneratorClassesVersionNumber();
+	if (SavedGeneratorClassesVersionNumber == CurrentGeneratorClassesVersionNumber)
 	{
 		return;
 	}
-	ClassGeneratorNamesRegisteredClassesVersionNumber = GetRegisteredClassesVersionNumber();
+	SavedGeneratorClassesVersionNumber = CurrentGeneratorClassesVersionNumber;
 
-	// Work around the fact we don't reference Engine module directly
+	TArray<UClass*> BlueprintCoreDerivedClasses;
 	FTopLevelAssetPath BlueprintCorePathName(TEXT("/Script/Engine"), TEXT("BlueprintCore"));
-	UClass* BlueprintCoreClass = FindObject<UClass>(BlueprintCorePathName);
-	if (!BlueprintCoreClass)
+	UClass* BlueprintCoreClass = nullptr;
+
 	{
-		return;
+		// FindObject and GetDerivedClasses are not legal during GarbageCollection. Note that we might be called from
+		// an async thread, in which case we might lock this thread until GC completes. This could cause a deadlock if
+		// there aren't enough async threads. But CollectCodeGeneratorClasses is not called on runtime or cooked
+		// editor because they are monolithic, and so this lock should only occur on uncooked editor platforms, which
+		// should have a high enough number of threads to not block garbage collection.
+		FGCScopeGuard NoGCScopeGuard;
+
+		// Work around the fact we don't reference Engine module directly
+		BlueprintCoreClass = FindObject<UClass>(BlueprintCorePathName);
+		if (!BlueprintCoreClass)
+		{
+			return;
+		}
+		GetDerivedClasses(BlueprintCoreClass, BlueprintCoreDerivedClasses);
 	}
 
 	ClassGeneratorNames.Add(BlueprintCoreClass->GetClassPathName());
-
-	TArray<UClass*> BlueprintCoreDerivedClasses;
-	GetDerivedClasses(BlueprintCoreClass, BlueprintCoreDerivedClasses);
 	for (UClass* BPCoreClass : BlueprintCoreDerivedClasses)
 	{
 		bool bAlreadyRecorded;
@@ -6911,7 +6936,7 @@ void FAssetRegistryImpl::UpdateInheritanceBuffer(Impl::FClassInheritanceBuffer&
 
 	}
 
-	OutBuffer.RegisteredClassesVersionNumber = GetRegisteredClassesVersionNumber();
+	OutBuffer.SavedAllClassesVersionNumber = GetCurrentAllClassesVersionNumber();
 	OutBuffer.bDirty = false;
 }
 
@@ -6921,14 +6946,18 @@ void UAssetRegistryImpl::GetInheritanceContextWithRequiredLock(FRWScopeLock& InO
 	UE::AssetRegistry::Impl::FClassInheritanceContext& InheritanceContext,
 	UE::AssetRegistry::Impl::FClassInheritanceBuffer& StackBuffer)
 {
-	uint64 CurrentClassesVersionNumber = GetRegisteredClassesVersionNumber();
+	using namespace UE::AssetRegistry;
+
+	uint64 CurrentGeneratorClassesVersionNumber = FAssetRegistryImpl::GetCurrentGeneratorClassesVersionNumber();
+	uint64 CurrentAllClassesVersionNumber = FAssetRegistryImpl::GetCurrentAllClassesVersionNumber();
 	bool bNeedsWriteLock = false;
-	if (GuardedData.GetClassGeneratorNamesRegisteredClassesVersionNumber() != CurrentClassesVersionNumber)
+	if (GuardedData.GetSavedGeneratorClassesVersionNumber() != CurrentGeneratorClassesVersionNumber)
 	{
 		// ConditionalUpdate writes to protected data in CollectCodeGeneratorClasses, so we cannot proceed under a read lock
 		bNeedsWriteLock = true;
 	}
-	if (GuardedData.IsTempCachingEnabled() && !GuardedData.GetTempCachedInheritanceBuffer().IsUpToDate(CurrentClassesVersionNumber))
+	if (GuardedData.IsTempCachingEnabled() &&
+		!GuardedData.GetTempCachedInheritanceBuffer().IsUpToDate(CurrentAllClassesVersionNumber))
 	{
 		// Temp caching is enabled, so we will be reading the protected data in TempCachedInheritanceBuffer
 		// It's out of date, so we need to write to it first, so we cannot proceed under a read lock
@@ -6940,18 +6969,24 @@ void UAssetRegistryImpl::GetInheritanceContextWithRequiredLock(FRWScopeLock& InO
 	}
 
 	// Note that we have to reread all data since we may have dropped the lock
-	GetInheritanceContextAfterVerifyingLock(CurrentClassesVersionNumber, InheritanceContext, StackBuffer);
+	GetInheritanceContextAfterVerifyingLock(CurrentGeneratorClassesVersionNumber, CurrentAllClassesVersionNumber,
+		InheritanceContext, StackBuffer);
 }
 
 void UAssetRegistryImpl::GetInheritanceContextWithRequiredLock(FWriteScopeLock& InOutScopeLock,
 	UE::AssetRegistry::Impl::FClassInheritanceContext& InheritanceContext,
 	UE::AssetRegistry::Impl::FClassInheritanceBuffer& StackBuffer)
 {
-	uint64 CurrentClassesVersionNumber = GetRegisteredClassesVersionNumber();
-	GetInheritanceContextAfterVerifyingLock(CurrentClassesVersionNumber, InheritanceContext, StackBuffer);
+	using namespace UE::AssetRegistry;
+
+	uint64 CurrentGeneratorClassesVersionNumber = FAssetRegistryImpl::GetCurrentGeneratorClassesVersionNumber();
+	uint64 CurrentAllClassesVersionNumber = FAssetRegistryImpl::GetCurrentAllClassesVersionNumber();
+	GetInheritanceContextAfterVerifyingLock(CurrentGeneratorClassesVersionNumber, CurrentAllClassesVersionNumber,
+		InheritanceContext, StackBuffer);
 }
 
-void UAssetRegistryImpl::GetInheritanceContextAfterVerifyingLock(uint64 CurrentClassesVersionNumber,
+void UAssetRegistryImpl::GetInheritanceContextAfterVerifyingLock(uint64 CurrentGeneratorClassesVersionNumber,
+	uint64 CurrentAllClassesVersionNumber,
 	UE::AssetRegistry::Impl::FClassInheritanceContext& InheritanceContext,
 	UE::AssetRegistry::Impl::FClassInheritanceBuffer& StackBuffer)
 {
@@ -6959,12 +6994,12 @@ void UAssetRegistryImpl::GetInheritanceContextAfterVerifyingLock(uint64 CurrentC
 	// We rely on this to simplify logic and only check bIsTempCachingEnabled
 	check(!GuardedData.IsTempCachingAlwaysEnabled() || GuardedData.IsTempCachingEnabled());
 
-	bool bCodeGeneratorClassesUpToDate = GuardedData.GetClassGeneratorNamesRegisteredClassesVersionNumber() == CurrentClassesVersionNumber;
+	bool bCodeGeneratorClassesUpToDate = GuardedData.GetSavedGeneratorClassesVersionNumber() == CurrentGeneratorClassesVersionNumber;
 	if (GuardedData.IsTempCachingEnabled())
 	{
 		// Use the persistent buffer
 		UE::AssetRegistry::Impl::FClassInheritanceBuffer& TempCachedInheritanceBuffer = GuardedData.GetTempCachedInheritanceBuffer();
-		bool bInheritanceMapUpToDate = TempCachedInheritanceBuffer.IsUpToDate(CurrentClassesVersionNumber);
+		bool bInheritanceMapUpToDate = TempCachedInheritanceBuffer.IsUpToDate(CurrentAllClassesVersionNumber);
 		InheritanceContext.BindToBuffer(TempCachedInheritanceBuffer, GuardedData, bInheritanceMapUpToDate, bCodeGeneratorClassesUpToDate);
 	}
 	else
@@ -7016,9 +7051,9 @@ void FClassInheritanceBuffer::Clear()
 	ReverseInheritanceMap.Empty();
 }
 
-bool FClassInheritanceBuffer::IsUpToDate(uint64 CurrentClassesVersionNumber) const
+bool FClassInheritanceBuffer::IsUpToDate(uint64 CurrentAllClassesVersionNumber) const
 {
-	return !bDirty && RegisteredClassesVersionNumber == CurrentClassesVersionNumber;
+	return !bDirty && SavedAllClassesVersionNumber == CurrentAllClassesVersionNumber;
 }
 
 SIZE_T FClassInheritanceBuffer::GetAllocatedSize() const
diff --git a/Engine/Source/Runtime/AssetRegistry/Private/AssetRegistry.h b/Engine/Source/Runtime/AssetRegistry/Private/AssetRegistry.h
--- a/Engine/Source/Runtime/AssetRegistry/Private/AssetRegistry.h
+++ b/Engine/Source/Runtime/AssetRegistry/Private/AssetRegistry.h
@@ -278,7 +278,8 @@ private:
 	void GetInheritanceContextWithRequiredLock(FWriteScopeLock& InOutScopeLock,
 		UE::AssetRegistry::Impl::FClassInheritanceContext& InheritanceContext,
 		UE::AssetRegistry::Impl::FClassInheritanceBuffer& StackBuffer);
-	void GetInheritanceContextAfterVerifyingLock(uint64 CurrentClassesVersionNumber,
+	void GetInheritanceContextAfterVerifyingLock(uint64 CurrentGeneratorClassesVersionNumber,
+		uint64 CurrentAllClassesVersionNumber,
 		UE::AssetRegistry::Impl::FClassInheritanceContext& InheritanceContext,
 		UE::AssetRegistry::Impl::FClassInheritanceBuffer& StackBuffer);
 
diff --git a/Engine/Source/Runtime/AssetRegistry/Private/AssetRegistryImpl.h b/Engine/Source/Runtime/AssetRegistry/Private/AssetRegistryImpl.h
--- a/Engine/Source/Runtime/AssetRegistry/Private/AssetRegistryImpl.h
+++ b/Engine/Source/Runtime/AssetRegistry/Private/AssetRegistryImpl.h
@@ -80,13 +80,13 @@ namespace Impl
 		TMap<FTopLevelAssetPath, FTopLevelAssetPath> InheritanceMap;
 		/** Map from Class->(All subclasses) for all classes including native classes and blueprint classes. Updated on demand. */
 		TMap<FTopLevelAssetPath, TArray<FTopLevelAssetPath>> ReverseInheritanceMap;
-		/** Snapshot of GetRegisteredClassesVersionNumber() at the time of the last update, to invalidate on changes to classes. */
-		uint64 RegisteredClassesVersionNumber = MAX_uint64;
+		/** Snapshot of GetCurrentAllClassesVersionNumber() at the time of the last update, to invalidate on changes to classes. */
+		uint64 SavedAllClassesVersionNumber = MAX_uint64;
 		/** Dirty flag to invalidate on other changes requiring a recompute. */
 		bool bDirty = true;
 
 		/** Report whether the dirty flag and other checks indicate this buffer does not need to be updated. */
-		bool IsUpToDate(uint64 CurrentClassesVersionNumber) const;
+		bool IsUpToDate(uint64 CurrentAllClassesVersionNumber) const;
 		/** Delete data and free allocations. */
 		void Clear();
 		/** Report size of dynamic allocations. */
@@ -255,7 +255,11 @@ public:
 	bool IsInitialSearchStarted() const { return bInitialSearchStarted; }
 	bool IsSearchAllAssets() const { return bSearchAllAssets; }
 	Impl::FClassInheritanceBuffer& GetTempCachedInheritanceBuffer() { return TempCachedInheritanceBuffer; }
-	uint64 GetClassGeneratorNamesRegisteredClassesVersionNumber() const { return ClassGeneratorNamesRegisteredClassesVersionNumber; }
+	uint64 GetSavedGeneratorClassesVersionNumber() const { return SavedGeneratorClassesVersionNumber; }
+	uint64 GetSavedAllClassesVersionNumber() const { return SavedAllClassesVersionNumber; }
+	static uint64 GetCurrentGeneratorClassesVersionNumber();
+	static uint64 GetCurrentAllClassesVersionNumber();
+
 	/** Get a copy of the cached serialization options that were parsed from ini */
 	void CopySerializationOptions(FAssetRegistrySerializationOptions& OutOptions, ESerializationTarget Target) const;
 	
@@ -398,7 +402,8 @@ private:
 	/** Persistent InheritanceBuffer used when SetTemporaryCachingMode is on */
 	Impl::FClassInheritanceBuffer TempCachedInheritanceBuffer;
 
-	uint64 ClassGeneratorNamesRegisteredClassesVersionNumber;
+	uint64 SavedGeneratorClassesVersionNumber;
+	uint64 SavedAllClassesVersionNumber;
 
 	/** The tree of known cached paths that assets may reside within */
 	FPathTree CachedPathTree;
diff --git a/Engine/Source/Runtime/CoreUObject/Private/UObject/UObjectHash.cpp b/Engine/Source/Runtime/CoreUObject/Private/UObject/UObjectHash.cpp
--- a/Engine/Source/Runtime/CoreUObject/Private/UObject/UObjectHash.cpp
+++ b/Engine/Source/Runtime/CoreUObject/Private/UObject/UObjectHash.cpp
@@ -15,6 +15,8 @@
 #include "HAL/LowLevelMemStats.h"
 #include "UObject/AnyPackagePrivate.h"
 
+#include <atomic>
+
 #if UE_STORE_OBJECT_LIST_INTERNAL_INDEX
 static_assert(sizeof(FName) == 4, "Internal object index optimization depends exploits 4 bytes padding after the FName");
 static_assert(sizeof(UObjectBase) == 40,
@@ -387,7 +389,8 @@ public:
 	TBucketMap<UClass*> ClassToObjectListMap;
 #endif
 	TMap<UClass*, TSet<UClass*> > ClassToChildListMap;
-	TAtomic<uint64> ClassToChildListMapVersion;
+	std::atomic<uint64> AllClassesVersion;
+	std::atomic<uint64> NativeClassesVersion;
 
 	/** Map of package to the object their contain. */
 	TBucketMap<UPackage*> PackageToObjectListMap;
@@ -395,7 +398,8 @@ public:
 	TMap<UObjectBase*, UPackage*> ObjectToPackageMap;
 
 	FUObjectHashTables()
-		: ClassToChildListMapVersion(0)
+		: AllClassesVersion(0)
+		, NativeClassesVersion(0)
 	{
 	}
 
@@ -989,9 +993,13 @@ FORCEINLINE void AddToClassMap(FUObjectHashTables& ThreadHash, UObjectBase* Obje
 			TSet<UClass*>& ChildList = ThreadHash.ClassToChildListMap.FindOrAdd(SuperClass);
 			bool bIsAlreadyInSetPtr = false;
 			ChildList.Add(Class, &bIsAlreadyInSetPtr);
-			ThreadHash.ClassToChildListMapVersion++;
 			check(!bIsAlreadyInSetPtr); // if it already exists, something is wrong with the external code
 		}
+		ThreadHash.AllClassesVersion++;
+		if (Class->IsNative())
+		{
+			ThreadHash.NativeClassesVersion++;
+		}
 	}
 }
 
@@ -1076,7 +1084,11 @@ FORCEINLINE void RemoveFromClassMap(FUObjectHashTables& ThreadHash, UObjectBase*
 			{
 				ThreadHash.ClassToChildListMap.Remove(SuperClass);
 			}
-			ThreadHash.ClassToChildListMapVersion++;
+		}
+		ThreadHash.AllClassesVersion++;
+		if (Class->IsNative())
+		{
+			ThreadHash.NativeClassesVersion++;
 		}
 	}
 }
@@ -1119,7 +1131,13 @@ void ShrinkUObjectHashTables()
 uint64 GetRegisteredClassesVersionNumber()
 {
 	FUObjectHashTables& ThreadHash = FUObjectHashTables::Get();
-	return ThreadHash.ClassToChildListMapVersion;
+	return ThreadHash.AllClassesVersion;
+}
+
+uint64 GetRegisteredNativeClassesVersionNumber()
+{
+	FUObjectHashTables& ThreadHash = FUObjectHashTables::Get();
+	return ThreadHash.NativeClassesVersion;
 }
 
 static void ShrinkUObjectHashTablesDel(const TArray<FString>& Args)
diff --git a/Engine/Source/Runtime/CoreUObject/Public/UObject/UObjectHash.h b/Engine/Source/Runtime/CoreUObject/Public/UObject/UObjectHash.h
--- a/Engine/Source/Runtime/CoreUObject/Public/UObject/UObjectHash.h
+++ b/Engine/Source/Runtime/CoreUObject/Public/UObject/UObjectHash.h
@@ -273,6 +273,13 @@ COREUOBJECT_API void ShrinkUObjectHashTables();
 */
 COREUOBJECT_API uint64 GetRegisteredClassesVersionNumber();
 
+/**
+* Get a version number representing the current state of registered native classes.
+*
+* Can be stored and then compared to invalidate external caching of native classes hierarchy whenever it changes.
+*/
+COREUOBJECT_API uint64 GetRegisteredNativeClassesVersionNumber();
+
 /**
  * Logs out information about the object hash for debug purposes
  *
