Fix issue where a ScanPathsSynchronous could broadcast the FileLoadedEvent before processing the deferred asset events
This could result in slow loading when specifying the map on the command line

diff --git a/Engine/Source/Runtime/AssetRegistry/Private/AssetRegistry.cpp b/Engine/Source/Runtime/AssetRegistry/Private/AssetRegistry.cpp
--- a/Engine/Source/Runtime/AssetRegistry/Private/AssetRegistry.cpp
+++ b/Engine/Source/Runtime/AssetRegistry/Private/AssetRegistry.cpp
@@ -7757,6 +7757,18 @@ void UAssetRegistryImpl::Broadcast(UE::AssetRegistry::Impl::FEventContext& Event
 
 	if (EventContext.bFileLoadedEventBroadcast)
 	{
+		if (&EventContext != &DeferredEvents)
+		{
+			FWriteScopeLock InterfaceScopeLock(InterfaceLock);
+			// Do not send the file loaded event yet if there are still deferred events and pass the flag on instead
+			if (!DeferredEvents.IsEmpty())
+			{
+				EventContext.bFileLoadedEventBroadcast = false;
+				DeferredEvents.bFileLoadedEventBroadcast = true;
+				return;
+			}
+		}
+
 		FileLoadedEvent.Broadcast();
 		EventContext.bFileLoadedEventBroadcast = false;
 	}
