World Partition: Don't rely solely on GRedirectCollector when trying to redirect class descriptors, as it is only filled at the end of the Asset Registry scan.

diff --git a/Engine/Source/Runtime/AssetRegistry/Private/AssetRegistryHelpers.cpp b/Engine/Source/Runtime/AssetRegistry/Private/AssetRegistryHelpers.cpp
--- a/Engine/Source/Runtime/AssetRegistry/Private/AssetRegistryHelpers.cpp
+++ b/Engine/Source/Runtime/AssetRegistry/Private/AssetRegistryHelpers.cpp
@@ -232,97 +232,27 @@ UAssetRegistryHelpers::FTemporaryCachingModeScope::~FTemporaryCachingModeScope()
 	UAssetRegistryHelpers::GetAssetRegistry()->SetTemporaryCachingMode(PreviousCachingMode);
 }
 
-bool UAssetRegistryHelpers::FixupRedirectedAssetPath(FSoftObjectPath& InOutSoftObjectPath)
+void UAssetRegistryHelpers::FixupRedirectedAssetPath(FSoftObjectPath& InOutSoftObjectPath)
 {
-	if (InOutSoftObjectPath.IsNull())
-	{
-		// Empty path, no redirect
-		return true;
-	}
+	FSoftObjectPath FoundRedirection;
 
 #if WITH_EDITOR
-	// Check GRedirectCollector first for faster fixup
-	FSoftObjectPath FoundRedirection = GRedirectCollector.GetAssetPathRedirection(InOutSoftObjectPath.GetWithoutSubPath());
-	if (!FoundRedirection.IsNull())
+	FoundRedirection = GRedirectCollector.GetAssetPathRedirection(InOutSoftObjectPath);
+	if (FoundRedirection.IsValid())
 	{
-		InOutSoftObjectPath.SetPath(FoundRedirection.GetAssetPath(), InOutSoftObjectPath.GetSubPathString());
-		return true;
+		InOutSoftObjectPath = FoundRedirection;
+		return;
 	}
 #endif
 
-	if (InOutSoftObjectPath.GetAssetName().IsEmpty())
-	{
-		// A package name. No need to ask the asset registry for assets, it wont find any
-		return true;
-	}
-
-	const FAssetData* AssetData;
 	IAssetRegistry& AssetRegistry = IAssetRegistry::GetChecked();
-
-	FTopLevelAssetPath AssetPath = InOutSoftObjectPath.GetAssetPath();
-	TSet<FTopLevelAssetPath> Visited;
-	Visited.Add(AssetPath);
-
-	for (;;)
-	{
-		TArray<FAssetData> Assets;
-		AssetRegistry.ScanFilesSynchronous({ AssetPath.GetPackageName().ToString() }, /*bForceRescan*/false);
-		AssetRegistry.GetAssetsByPackageName(AssetPath.GetPackageName(), Assets, /*bIncludeOnlyOnDiskAssets*/true);
-
-		if (!Assets.Num())
-		{
-			UE_LOG(LogCore, Warning, TEXT("Failed to find assets for asset path '%s'"), *AssetPath.ToString());
-			return false;
-		}
-
-		AssetData = Assets.FindByPredicate([&AssetPath](const FAssetData& AssetData)
-		{
-			return (AssetData.ToSoftObjectPath().GetAssetPath() == AssetPath);
-		});
-
-		if (!AssetData)
-		{
-			UE_LOG(LogCore, Warning, TEXT("Failed to find asset for asset path '%s'"), *AssetPath.ToString());
-			return false;
-		}
-
-		if (!AssetData->IsRedirector())
-		{
-			break;
-		}
-
-		FString DestinationObjectPath;
-		if (!AssetData->GetTagValue(TEXT("DestinationObject"), DestinationObjectPath))
-		{
-			UE_LOG(LogCore, Warning, TEXT("Failed to follow redirector for '%s'"), *AssetPath.ToString());
-			return false;
-		}
-
-		// Update asset path
-		AssetPath = FTopLevelAssetPath(DestinationObjectPath);
-
-		bool bAlreadyVisited = false;
-		Visited.Add(AssetPath, &bAlreadyVisited);
-		if (bAlreadyVisited)
-		{
-			UE_LOG(LogCore, Warning, TEXT("Asset path was already visited '%s'"), *AssetPath.ToString());
-			return false;
-		}
-	}
-
-	InOutSoftObjectPath.SetPath(AssetPath, InOutSoftObjectPath.GetSubPathString());
-
-	return true;
+	FoundRedirection = AssetRegistry.GetRedirectedObjectPath(InOutSoftObjectPath.GetWithoutSubPath());
+	InOutSoftObjectPath = FSoftObjectPath(FoundRedirection.GetAssetPath(), InOutSoftObjectPath.GetSubPathString());
 }
 
-bool UAssetRegistryHelpers::FixupRedirectedAssetPath(FName& InOutAssetPath)
+void UAssetRegistryHelpers::FixupRedirectedAssetPath(FName& InOutAssetPath)
 {
 	FSoftObjectPath SoftObjectPath(InOutAssetPath.ToString());
-	if (FixupRedirectedAssetPath(SoftObjectPath))
-	{
-		InOutAssetPath = FName(*SoftObjectPath.ToString());
-		return true;
-	}
-
-	return false;
+	FixupRedirectedAssetPath(SoftObjectPath);
+	InOutAssetPath = FName(*SoftObjectPath.ToString());
 }
diff --git a/Engine/Source/Runtime/AssetRegistry/Public/AssetRegistry/AssetRegistryHelpers.h b/Engine/Source/Runtime/AssetRegistry/Public/AssetRegistry/AssetRegistryHelpers.h
--- a/Engine/Source/Runtime/AssetRegistry/Public/AssetRegistry/AssetRegistryHelpers.h
+++ b/Engine/Source/Runtime/AssetRegistry/Public/AssetRegistry/AssetRegistryHelpers.h
@@ -140,14 +140,12 @@ public:
 	/**
 	 * Resolves the provided asset path using asset redirectors.
 	 * @param	InOutAssetPath	Asset path to resolve [In/Out].
-	 * @return	False if an error occurred while trying to resolve, else true.
 	 */
-	static ASSETREGISTRY_API bool FixupRedirectedAssetPath(FName& InOutAssetPath);
+	static ASSETREGISTRY_API void FixupRedirectedAssetPath(FName& InOutAssetPath);
 	
 	/**
 	 * Resolves the provided soft object asset path using asset redirectors.
 	 * @param	InOutSoftObjectPath	Asset path to resolve [In/Out].
-	 * @return	False if an error occurred while trying to resolve, else true.
 	 */
-	static ASSETREGISTRY_API bool FixupRedirectedAssetPath(FSoftObjectPath& InOutSoftObjectPath);
+	static ASSETREGISTRY_API void FixupRedirectedAssetPath(FSoftObjectPath& InOutSoftObjectPath);
 };
diff --git a/Engine/Source/Runtime/Engine/Private/WorldPartition/WorldPartitionActorDescArchive.cpp b/Engine/Source/Runtime/Engine/Private/WorldPartition/WorldPartitionActorDescArchive.cpp
--- a/Engine/Source/Runtime/Engine/Private/WorldPartition/WorldPartitionActorDescArchive.cpp
+++ b/Engine/Source/Runtime/Engine/Private/WorldPartition/WorldPartitionActorDescArchive.cpp
@@ -79,11 +79,9 @@ PRAGMA_ENABLE_DEPRECATION_WARNINGS
 					InOutClassPath = FTopLevelAssetPath(RedirectedClassRedirect.ToString());
 				}
 
-				FSoftObjectPath FoundRedirection = GRedirectCollector.GetAssetPathRedirection(FSoftObjectPath(InOutClassPath.ToString()));
-				if (FoundRedirection.IsValid())
-				{
-					InOutClassPath = FoundRedirection.GetAssetPath();
-				}
+				FSoftObjectPath RedirectedClassPath(InOutClassPath.ToString());
+				UAssetRegistryHelpers::FixupRedirectedAssetPath(RedirectedClassPath);
+				InOutClassPath = RedirectedClassPath.GetAssetPath();
 			}
 		};
 
diff --git a/Engine/Source/Runtime/Engine/Private/WorldPartition/WorldPartitionActorDescUtils.cpp b/Engine/Source/Runtime/Engine/Private/WorldPartition/WorldPartitionActorDescUtils.cpp
--- a/Engine/Source/Runtime/Engine/Private/WorldPartition/WorldPartitionActorDescUtils.cpp
+++ b/Engine/Source/Runtime/Engine/Private/WorldPartition/WorldPartitionActorDescUtils.cpp
@@ -209,6 +209,7 @@ void FWorldPartitionActorDescUtils::ReplaceActorDescriptorPointerFromActor(const
 
 bool FWorldPartitionActorDescUtils::FixupRedirectedAssetPath(FName& InOutAssetPath)
 {
-	return UAssetRegistryHelpers::FixupRedirectedAssetPath(InOutAssetPath);
+	UAssetRegistryHelpers::FixupRedirectedAssetPath(InOutAssetPath);
+	return true;
 }
 #endif
diff --git a/Engine/Source/Runtime/Engine/Private/WorldPartition/WorldPartitionHelpers.cpp b/Engine/Source/Runtime/Engine/Private/WorldPartition/WorldPartitionHelpers.cpp
--- a/Engine/Source/Runtime/Engine/Private/WorldPartition/WorldPartitionHelpers.cpp
+++ b/Engine/Source/Runtime/Engine/Private/WorldPartition/WorldPartitionHelpers.cpp
@@ -312,12 +312,14 @@ bool FWorldPartitionHelpers::ConvertRuntimePathToEditorPath(const FSoftObjectPat
 
 bool FWorldPartitionHelpers::FixupRedirectedAssetPath(FSoftObjectPath& InOutSoftObjectPath)
 {
-	return UAssetRegistryHelpers::FixupRedirectedAssetPath(InOutSoftObjectPath);
+	UAssetRegistryHelpers::FixupRedirectedAssetPath(InOutSoftObjectPath);
+	return true;
 }
 
 bool FWorldPartitionHelpers::FixupRedirectedAssetPath(FName& InOutAssetPath)
 {
-	return UAssetRegistryHelpers::FixupRedirectedAssetPath(InOutAssetPath);
+	UAssetRegistryHelpers::FixupRedirectedAssetPath(InOutAssetPath);
+	return true;
 }
 
 #endif // #if WITH_EDITOR
