AssetRegistry: Extend the workaround that blocks invalid externalactors from entering the AssetRegistry to include blocking them if they are somehow loaded.
#rnx

diff --git a/Engine/Source/Runtime/AssetRegistry/Private/AssetRegistry.cpp b/Engine/Source/Runtime/AssetRegistry/Private/AssetRegistry.cpp
--- a/Engine/Source/Runtime/AssetRegistry/Private/AssetRegistry.cpp
+++ b/Engine/Source/Runtime/AssetRegistry/Private/AssetRegistry.cpp
@@ -6449,6 +6449,10 @@ void FAssetRegistryImpl::PushProcessLoadedAssetsBatch(Impl::FEventContext& Event
 	// Add or update existing for all of the AssetDatas created by the batch
 	for (FAssetData& NewAssetData : LoadedAssetDatas)
 	{
+		if (ShouldSkipGatheredAsset(NewAssetData))
+		{
+			continue;
+		}
 		FCachedAssetKey Key(NewAssetData);
 		FAssetData** DataFromGather = State.CachedAssets.Find(Key);
 
