AssetRegistry: Make sure WaitForCompletion sends out the FileLoadedEvent since doing so before return is part of the contract. AssetManager has an ensure that it has been sent.
This change fixes a race condition ensure failure that can occur when starting PIE for a level before the initial assetregistry scan finishes.
#rnx

diff --git a/Engine/Source/Runtime/AssetRegistry/Private/AssetRegistry.cpp b/Engine/Source/Runtime/AssetRegistry/Private/AssetRegistry.cpp
--- a/Engine/Source/Runtime/AssetRegistry/Private/AssetRegistry.cpp
+++ b/Engine/Source/Runtime/AssetRegistry/Private/AssetRegistry.cpp
@@ -1816,6 +1816,12 @@ void UAssetRegistryImpl::WaitForCompletion()
 			FClassInheritanceContext InheritanceContext;
 			FClassInheritanceBuffer InheritanceBuffer;
 			GetInheritanceContextWithRequiredLock(InterfaceScopeLock, InheritanceContext, InheritanceBuffer);
+			if (IsInGameThread())
+			{
+				// Process any deferred events. Required since deferred events would block sending the FileLoadedEvent
+				EventContext = MoveTemp(DeferredEvents);
+				DeferredEvents.Clear();
+			}
 
 			GuardedData.WaitForGathererIdleIfSynchronous();
 
