Add directory watchers for plugin root directories in Asset Registry

The directory watchers in the asset registry take up just above 100MB of memory when running on Windows. We are trying to reduce the memory footprint of the Content Worker to mitigate OOM scenarios. The main reason for this is the sheer amount of directory watches that we are setting up (around 418, as of writing). Most of these share a common prefix (either Engine/Plugins or <Game>/Plugin), so we may instead set up watchers for these directories and filter all incoming changes appropriately. This results in only 10 directory watchers, bringing memory usage from 104.5MB to just about 2.5MB.

diff --git a/Engine/Source/Runtime/AssetRegistry/Private/AssetRegistry.cpp b/Engine/Source/Runtime/AssetRegistry/Private/AssetRegistry.cpp
--- a/Engine/Source/Runtime/AssetRegistry/Private/AssetRegistry.cpp
+++ b/Engine/Source/Runtime/AssetRegistry/Private/AssetRegistry.cpp
@@ -1000,11 +1000,31 @@ void UAssetRegistryImpl::InitializeEvents(UE::AssetRegistry::Impl::FInitializeCo
 
 		if (DirectoryWatcher)
 		{
+			// The vast majority of directories we are watching are below the Plugin directories. The memory cost per watch
+			// is sufficiently high to want to avoid setting up many granular watches when we can also setup two coarse ones.
+			DirectoryWatchRoots.Add(FPaths::CreateStandardFilename(FPaths::EnginePluginsDir()));
+			DirectoryWatchRoots.Add(FPaths::CreateStandardFilename(FPaths::ProjectPluginsDir()));
+			for (FString& WatchRoot : DirectoryWatchRoots)
+			{
+				FDelegateHandle NewHandle;
+				DirectoryWatcher->RegisterDirectoryChangedCallback_Handle(
+					WatchRoot,
+					IDirectoryWatcher::FDirectoryChanged::CreateUObject(this, &UAssetRegistryImpl::OnDirectoryChanged),
+					NewHandle,
+					IDirectoryWatcher::WatchOptions::IncludeDirectoryChanges);
+
+				OnDirectoryChangedDelegateHandles.Add(WatchRoot, NewHandle);
+			}
+
 			FString ContentFolder;
 			for (TArray<FString>::TConstIterator RootPathIt(Context.RootContentPaths); RootPathIt; ++RootPathIt)
 			{
 				const FString& RootPath = *RootPathIt;
-				ContentFolder = FPackageName::LongPackageNameToFilename(RootPath);
+				ContentFolder = FPaths::CreateStandardFilename(FPackageName::LongPackageNameToFilename(RootPath));
+				if (IsDirAlreadyWatchedByRootWatchers(ContentFolder))
+				{
+					continue;
+				}
 
 				// A missing directory here could be due to a plugin that specifies it contains content, yet has no content yet.
 				// PluginManager mounts these folders anyway which results in them being returned from QueryRootContentPaths.
@@ -1558,9 +1578,19 @@ void UAssetRegistryImpl::FinishDestroy()
 					for (TArray<FString>::TConstIterator RootPathIt(RootContentPaths); RootPathIt; ++RootPathIt)
 					{
 						const FString& RootPath = *RootPathIt;
-						const FString& ContentFolder = FPackageName::LongPackageNameToFilename(RootPath);
-						DirectoryWatcher->UnregisterDirectoryChangedCallback_Handle(ContentFolder, OnDirectoryChangedDelegateHandles.FindRef(RootPath));
+						const FString ContentFolder = FPaths::CreateStandardFilename(FPackageName::LongPackageNameToFilename(RootPath));
+						if (!IsDirAlreadyWatchedByRootWatchers(ContentFolder))
+						{
+							DirectoryWatcher->UnregisterDirectoryChangedCallback_Handle(ContentFolder, OnDirectoryChangedDelegateHandles.FindRef(RootPath));
+						}
 					}
+
+					for (TArray<FString>::TConstIterator RootPathIt(DirectoryWatchRoots); RootPathIt; ++RootPathIt)
+					{
+						const FString& RootPath = *RootPathIt;
+						DirectoryWatcher->UnregisterDirectoryChangedCallback_Handle(RootPath, OnDirectoryChangedDelegateHandles.FindRef(RootPath));
+					}
+					DirectoryWatchRoots.Empty();
 				}
 			}
 		}
@@ -6042,6 +6072,21 @@ void UAssetRegistryImpl::OnDirectoryChanged(const TArray<FFileChangeData>& FileC
 		}
 	}
 
+	{
+		// Check that the change is related to a directory that has actually been mounted.
+		FStringBuilderBase MountPointPackageName;
+		FStringBuilderBase MountPointFilePath;
+		FStringBuilderBase RelativePath;
+		for (int32 FileEntryIndex = FileChangesProcessed.Num() - 1; FileEntryIndex >= 0; FileEntryIndex--)
+		{
+			const FString& Filename = FileChangesProcessed[FileEntryIndex].Filename;
+			if (!FPackageName::TryGetMountPointForPath(Filename, MountPointPackageName, MountPointFilePath, RelativePath))
+			{
+				FileChangesProcessed.RemoveAt(FileEntryIndex);
+			}
+		}
+	}
+
 	UE::AssetRegistry::Impl::FEventContext EventContext;
 	bool bInitialSearchStarted;
 	bool bInitialSearchCompleted;
@@ -6642,14 +6687,15 @@ void UAssetRegistryImpl::OnContentPathMounted(const FString& InAssetPath, const
 
 		// Listen for directory changes in this content path
 #if WITH_EDITOR
+		const FString StandardFileSystemPath = FPaths::CreateStandardFilename(FileSystemPath);
 		// In-game doesn't listen for directory changes
-		if (DirectoryWatcher)
+		if (DirectoryWatcher && !IsDirAlreadyWatchedByRootWatchers(StandardFileSystemPath))
 		{
 			if (!OnDirectoryChangedDelegateHandles.Contains(AssetPathWithTrailingSlash))
 			{
 				FDelegateHandle NewHandle;
 				DirectoryWatcher->RegisterDirectoryChangedCallback_Handle(
-					FileSystemPath, 
+					StandardFileSystemPath,
 					IDirectoryWatcher::FDirectoryChanged::CreateUObject(this, &UAssetRegistryImpl::OnDirectoryChanged),
 					NewHandle, 
 					IDirectoryWatcher::WatchOptions::IncludeDirectoryChanges);
@@ -6717,7 +6763,8 @@ void UAssetRegistryImpl::OnContentPathDismounted(const FString& InAssetPath, con
 
 		// Stop listening for directory changes in this content path
 #if WITH_EDITOR
-		if (DirectoryWatcher)
+		const FString StandardFileSystemPath = FPaths::CreateStandardFilename(FileSystemPath);
+		if (DirectoryWatcher && !IsDirAlreadyWatchedByRootWatchers(StandardFileSystemPath))
 		{
 			// Make sure OnDirectoryChangedDelegateHandles key is symmetrical with the one used in OnContentPathMounted
 			FString AssetPathWithTrailingSlash;
@@ -6733,7 +6780,7 @@ void UAssetRegistryImpl::OnContentPathDismounted(const FString& InAssetPath, con
 			FDelegateHandle DirectoryChangedHandle;
 			if (ensure(OnDirectoryChangedDelegateHandles.RemoveAndCopyValue(AssetPathWithTrailingSlash, DirectoryChangedHandle)))
 			{
-				DirectoryWatcher->UnregisterDirectoryChangedCallback_Handle(FileSystemPath, DirectoryChangedHandle);
+				DirectoryWatcher->UnregisterDirectoryChangedCallback_Handle(StandardFileSystemPath, DirectoryChangedHandle);
 			}
 		}
 #endif // WITH_EDITOR
@@ -7068,6 +7115,14 @@ void UAssetRegistryImpl::OnGetExtraObjectTags(FAssetRegistryTagsContext Context)
 		}
 	}
 }
+
+bool UAssetRegistryImpl::IsDirAlreadyWatchedByRootWatchers(const FString& Directory) const
+{
+	return DirectoryWatchRoots.ContainsByPredicate([&Directory](const FString& WatchRoot) -> bool {
+		return FPaths::IsUnderDirectory(Directory, WatchRoot);
+	});
+}
+
 #endif
 
 namespace UE::AssetRegistry
diff --git a/Engine/Source/Runtime/AssetRegistry/Private/AssetRegistry.h b/Engine/Source/Runtime/AssetRegistry/Private/AssetRegistry.h
--- a/Engine/Source/Runtime/AssetRegistry/Private/AssetRegistry.h
+++ b/Engine/Source/Runtime/AssetRegistry/Private/AssetRegistry.h
@@ -293,6 +293,12 @@ private:
 	 * If bAddMetaDataTagsToOnGetExtraObjectTags is true, this function will add missing UMetaData tags to cooked assets
 	 */
 	void OnGetExtraObjectTags(FAssetRegistryTagsContext Context);
+
+	/**
+	 * Checks whether the given path is already covered by the general directory watches, or whether we need to setup a
+	 * new directory watcher. The caller must ensure that the Directory parameter is in FPaths::CreateStandardFilename format.
+	 */
+	bool IsDirAlreadyWatchedByRootWatchers(const FString& Directory) const;
 #endif
 
 private:
@@ -305,6 +311,7 @@ private:
 #if WITH_EDITOR
 	/** Handles to all registered OnDirectoryChanged delegates */
 	TMap<FString, FDelegateHandle> OnDirectoryChangedDelegateHandles;
+	TArray<FString> DirectoryWatchRoots;
 #endif
 
 #if WITH_EDITORONLY_DATA
