Check whether plugin directories exist before adding a directory watcher to them

Those directories might not exist, in which case the directory watcher issues a warning. There is no point in watching these directories if they do not exist as there likely won't be any mount points below them (that is different from the per-plugin paths because there the plugin might exist and it's just missing the last part of the path). In a worst case scenario, the session will set up a new directory watcher if a new mount point appears at some point during the lifetime of the program.

diff --git a/Engine/Source/Runtime/AssetRegistry/Private/AssetRegistry.cpp b/Engine/Source/Runtime/AssetRegistry/Private/AssetRegistry.cpp
--- a/Engine/Source/Runtime/AssetRegistry/Private/AssetRegistry.cpp
+++ b/Engine/Source/Runtime/AssetRegistry/Private/AssetRegistry.cpp
@@ -1002,8 +1002,18 @@ void UAssetRegistryImpl::InitializeEvents(UE::AssetRegistry::Impl::FInitializeCo
 		{
 			// The vast majority of directories we are watching are below the Plugin directories. The memory cost per watch
 			// is sufficiently high to want to avoid setting up many granular watches when we can also setup two coarse ones.
-			DirectoryWatchRoots.Add(FPaths::CreateStandardFilename(FPaths::EnginePluginsDir()));
-			DirectoryWatchRoots.Add(FPaths::CreateStandardFilename(FPaths::ProjectPluginsDir()));
+
+			const FString ProjectPluginDir = FPaths::CreateStandardFilename(FPaths::ProjectPluginsDir());
+			if (IPlatformFile::GetPlatformPhysical().DirectoryExists(*ProjectPluginDir))
+			{
+				DirectoryWatchRoots.Add(ProjectPluginDir);
+			}
+			const FString EnginePluginDir = FPaths::CreateStandardFilename(FPaths::EnginePluginsDir());
+			if (IPlatformFile::GetPlatformPhysical().DirectoryExists(*EnginePluginDir))
+			{
+				DirectoryWatchRoots.Add(EnginePluginDir);
+			}
+
 			for (FString& WatchRoot : DirectoryWatchRoots)
 			{
 				FDelegateHandle NewHandle;
