AssetRegistry: Fix new WatcherRoots feature to handle FCA_RescanRequired event from the DirectoryWatcher.
When WatcherRoots are enabled, we have a single watch for the entire plugins directory. If the buffer on our DirectoryWatcher is too small (16k events) for a rapid set of changes, then FCA_RescanRequired is sent, and we have to rescan the entire watched directory. Fixed the AssetRegistry to correctly handle this case for a WatcherRoot, and rescan all content directories under the WatcherRoot.
This rescan is not accelerated and could be prohibitively slow. If that comes up as a problem there's not much planned to handle it, we may need to disable the WatcherRoots feature for large projects.
#rnx

diff --git a/Engine/Source/Runtime/AssetRegistry/Private/AssetRegistry.cpp b/Engine/Source/Runtime/AssetRegistry/Private/AssetRegistry.cpp
--- a/Engine/Source/Runtime/AssetRegistry/Private/AssetRegistry.cpp
+++ b/Engine/Source/Runtime/AssetRegistry/Private/AssetRegistry.cpp
@@ -10,6 +10,7 @@
 #include "AssetRegistryConsoleCommands.h"
 #include "AssetRegistryPrivate.h"
 #include "Async/Async.h"
+#include "Async/ParallelFor.h"
 #include "Blueprint/BlueprintSupport.h"
 #include "DependsNode.h"
 #include "GenericPlatform/GenericPlatformChunkInstall.h"
@@ -6089,8 +6090,9 @@ void UAssetRegistryImpl::OnDirectoryChanged(const TArray<FFileChangeData>& FileC
 		FStringBuilderBase RelativePath;
 		for (int32 FileEntryIndex = FileChangesProcessed.Num() - 1; FileEntryIndex >= 0; FileEntryIndex--)
 		{
-			const FString& Filename = FileChangesProcessed[FileEntryIndex].Filename;
-			if (!FPackageName::TryGetMountPointForPath(Filename, MountPointPackageName, MountPointFilePath, RelativePath))
+			FFileChangeData& Data = FileChangesProcessed[FileEntryIndex];
+			if (Data.Action != FFileChangeData::FCA_RescanRequired && !FPackageName::TryGetMountPointForPath(
+				Data.Filename, MountPointPackageName, MountPointFilePath, RelativePath))
 			{
 				FileChangesProcessed.RemoveAt(FileEntryIndex);
 			}
@@ -6286,70 +6288,124 @@ void FAssetRegistryImpl::OnDirectoryChanged(Impl::FEventContext& EventContext,
 void FAssetRegistryImpl::OnDirectoryRescanRequired(Impl::FEventContext& EventContext,
 	Impl::FClassInheritanceContext& InheritanceContext, FString& DirPath, int64 BeforeTimeStamp)
 {
-	FString PackageNamePath;
+	TArray<TPair<FString,FString>> DirPathsAndPackageNames;
+	FString DirPathAsPackageName;
 	FString NormalizedDirPath = FPaths::CreateStandardFilename(DirPath);
-	if (!FPackageName::TryConvertFilenameToLongPackageName(NormalizedDirPath, PackageNamePath))
+	if (FPackageName::TryConvertFilenameToLongPackageName(NormalizedDirPath, DirPathAsPackageName))
 	{
-		return;
+		DirPathsAndPackageNames.Emplace(DirPath, MoveTemp(DirPathAsPackageName));
 	}
-
-	TArray<FString> NewFiles;
-	TArray<FString> NewDirs;
-	TArray<FString> ModifiedFiles;
-	TSet<FName> RemovedLongPackageNames;
-
-	FDateTime BeforeDateTime = FDateTime::FromUnixTimestamp(BeforeTimeStamp);
-	EnumerateAssetsByPathNoTags(*PackageNamePath, [&RemovedLongPackageNames](const FAssetData& AssetData)
-		{
-			RemovedLongPackageNames.Add(AssetData.PackageName);
-			return true;
-		}, true /* bRecursive */, true /* bIncludeOnlyOnDiskAssets */);
-
-	FPackageName::IteratePackagesInDirectory(DirPath,
-		[&DirPath, &PackageNamePath, &BeforeDateTime, &NewFiles, &ModifiedFiles, &RemovedLongPackageNames, &NormalizedDirPath]
-		(const TCHAR* Filename, const FFileStatData& StatData)
+	else
+	{
+		TArray<FString> ContentRoots;
+		FPackageName::QueryRootContentPaths(ContentRoots);
+		TStringBuilder<64> UnusedPackageName;
+		TStringBuilder<256> MountedFilePath;
+		TStringBuilder<16> UnusedRelPath;
+		for (FString& MountedLongPackageName : ContentRoots)
 		{
-			// Convert Filename to a PackagePath. We know the base dir so its faster to use that than FPackageName
-			// which has to scan all mount dirs
-			FStringView RelPath;
-			FString NormalizedFilename = FPaths::CreateStandardFilename(Filename);
-			if (!FPathViews::TryMakeChildPathRelativeTo(NormalizedFilename, NormalizedDirPath, RelPath))
+			if (FPackageName::TryGetMountPointForPath(MountedLongPackageName, UnusedPackageName, MountedFilePath, UnusedRelPath))
 			{
-				return true;
-			}
-			const bool bIsPackageFile = FPackageName::IsPackageExtension(*FString(FPathViews::GetExtension(RelPath, true /* bIncludeDot */)));
-			RelPath = FPathViews::GetBaseFilenameWithPath(RelPath);
-			TStringBuilder<256> FilePackagePath;
-			FilePackagePath << PackageNamePath;
-			FPathViews::AppendPath(FilePackagePath, RelPath);
-			for (int32 Index = 0; Index < FilePackagePath.Len(); ++Index)
-			{
-				TCHAR& Char = FilePackagePath.GetData()[Index];
-				if (Char == '\\')
+				FString NormalizeMountedFilePath = FPaths::CreateStandardFilename(FString(MountedFilePath));
+				if (FPaths::IsUnderDirectory(NormalizeMountedFilePath, NormalizedDirPath))
 				{
-					Char = '/';
+					DirPathsAndPackageNames.Emplace(MoveTemp(NormalizeMountedFilePath), MoveTemp(MountedLongPackageName));
 				}
 			}
-			const bool bIsValidPackageName = FPackageName::IsValidTextForLongPackageName(FilePackagePath);
-			if (!bIsPackageFile || !bIsValidPackageName)
+		}
+	}
+	if (DirPathsAndPackageNames.IsEmpty())
+	{
+		return;
+	}
+
+	struct FDirectoryResults
+	{
+		TArray<FString> NewFiles;
+		TArray<FString> ModifiedFiles;
+		TSet<FName> RemovedLongPackageNames;
+	};
+	int32 NumDirs = DirPathsAndPackageNames.Num();
+	TArray<FDirectoryResults> Results;
+	Results.SetNum(NumDirs);
+	FDateTime BeforeDateTime = FDateTime::FromUnixTimestamp(BeforeTimeStamp);
+
+	for (int32 DirIndex = 0; DirIndex < NumDirs; ++DirIndex)
+	{
+		FString& PackageNamePath = DirPathsAndPackageNames[DirIndex].Value;
+		FDirectoryResults& Result = Results[DirIndex];
+		EnumerateAssetsByPathNoTags(*PackageNamePath, [&Result](const FAssetData& AssetData)
 			{
+				Result.RemovedLongPackageNames.Add(AssetData.PackageName);
 				return true;
-			}
+			}, true /* bRecursive */, true /* bIncludeOnlyOnDiskAssets */);
+	}
 
-			if (StatData.CreationTime > BeforeDateTime)
-			{
-				NewFiles.Add(NormalizedFilename);
-			}
-			else if (StatData.ModificationTime > BeforeDateTime)
-			{
-				ModifiedFiles.Add(NormalizedFilename);
-			}
-			RemovedLongPackageNames.Remove(FName(FilePackagePath.ToView()));
+	ParallelFor(NumDirs, [this, &DirPathsAndPackageNames, &Results, &BeforeDateTime](int32 DirIndex)
+		{
+			FDirectoryResults& Result = Results[DirIndex];
+			TPair<FString, FString>& Pair = DirPathsAndPackageNames[DirIndex];
+			FString& LocalPath = Pair.Key;
+			FString& PackageNamePath = Pair.Value;
 
-			return true;
+			FPackageName::IteratePackagesInDirectory(LocalPath,
+				[&LocalPath, &PackageNamePath, &BeforeDateTime, &Result]
+				(const TCHAR* Filename, const FFileStatData& StatData)
+				{
+					// Convert Filename to a PackagePath. We know the base dir so its faster to use that than FPackageName
+					// which has to scan all mount dirs
+					FStringView RelPath;
+					FString NormalizedFilename = FPaths::CreateStandardFilename(Filename);
+					if (!FPathViews::TryMakeChildPathRelativeTo(NormalizedFilename, LocalPath, RelPath))
+					{
+						return true;
+					}
+					const bool bIsPackageFile = FPackageName::IsPackageExtension(
+						*FString(FPathViews::GetExtension(RelPath, true /* bIncludeDot */)));
+					RelPath = FPathViews::GetBaseFilenameWithPath(RelPath);
+					TStringBuilder<256> FilePackagePath;
+					FilePackagePath << PackageNamePath;
+					FPathViews::AppendPath(FilePackagePath, RelPath);
+					for (int32 Index = 0; Index < FilePackagePath.Len(); ++Index)
+					{
+						TCHAR& Char = FilePackagePath.GetData()[Index];
+						if (Char == '\\')
+						{
+							Char = '/';
+						}
+					}
+					const bool bIsValidPackageName = FPackageName::IsValidTextForLongPackageName(FilePackagePath);
+					if (!bIsPackageFile || !bIsValidPackageName)
+					{
+						return true;
+					}
+
+					if (StatData.CreationTime > BeforeDateTime)
+					{
+						Result.NewFiles.Add(NormalizedFilename);
+					}
+					else if (StatData.ModificationTime > BeforeDateTime)
+					{
+						Result.ModifiedFiles.Add(NormalizedFilename);
+					}
+					Result.RemovedLongPackageNames.Remove(FName(FilePackagePath.ToView()));
+
+					return true;
+				});
 		});
 
-	for (FName LongPackageName : RemovedLongPackageNames)
+	TArray<FName> FinalRemovedLongPackageNames;
+	FDirectoryResults& FinalResult = Results[0];
+	FinalRemovedLongPackageNames.Append(FinalResult.RemovedLongPackageNames.Array());
+	for (int32 DirIndex = 1; DirIndex < NumDirs; ++DirIndex)
+	{
+		FDirectoryResults& ResultToMerge = Results[DirIndex];
+		FinalResult.NewFiles.Append(MoveTemp(ResultToMerge.NewFiles));
+		FinalResult.ModifiedFiles.Append(MoveTemp(ResultToMerge.ModifiedFiles));
+		FinalRemovedLongPackageNames.Append(ResultToMerge.RemovedLongPackageNames.Array());
+	}
+
+	for (FName LongPackageName : FinalRemovedLongPackageNames)
 	{
 		// This file was deleted. Remove all assets in the package from the registry.
 		RemovePackageData(EventContext, LongPackageName);
@@ -6357,24 +6413,21 @@ void FAssetRegistryImpl::OnDirectoryRescanRequired(Impl::FEventContext& EventCon
 		// Disk now matches editor
 		RemoveEmptyPackage(LongPackageName);
 	}
-	if (NewFiles.Num() || NewDirs.Num())
+	if (FinalResult.NewFiles.Num())
 	{
 		if (GlobalGatherer.IsValid())
 		{
-			for (FString& NewDir : NewDirs)
-			{
-				GlobalGatherer->OnDirectoryCreated(NewDir);
-			}
-			GlobalGatherer->OnFilesCreated(NewFiles);
+			GlobalGatherer->OnFilesCreated(FinalResult.NewFiles);
 			if (GlobalGatherer->IsSynchronous())
 			{
-				Impl::FScanPathContext Context(EventContext, InheritanceContext, NewDirs, NewFiles,
+				TArray<FString> UnusedNewDirs;
+				Impl::FScanPathContext Context(EventContext, InheritanceContext, UnusedNewDirs, FinalResult.NewFiles,
 					false /* bForceRescan */, false /* bIgnoreDenyListScanFilters */, nullptr /* OutFoundAssets */);
 				ScanPathsSynchronous(Context);
 			}
 		}
 	}
-	ScanModifiedAssetFiles(EventContext, InheritanceContext, ModifiedFiles);
+	ScanModifiedAssetFiles(EventContext, InheritanceContext, FinalResult.ModifiedFiles);
 }
 
 }
