Add missing LLM Tag for AssetRegistry loading

On a ContentWorker run with a pretty much empty project I am seeing roughly ~350MB of untagged allocations which should actually go to the AssetRegistry. They come from the preload codepath in CookOnTheFlyServer.cpp, which is missing any LLM Tags.

diff --git a/Engine/Source/Runtime/AssetRegistry/Private/AssetRegistryState.cpp b/Engine/Source/Runtime/AssetRegistry/Private/AssetRegistryState.cpp
--- a/Engine/Source/Runtime/AssetRegistry/Private/AssetRegistryState.cpp
+++ b/Engine/Source/Runtime/AssetRegistry/Private/AssetRegistryState.cpp
@@ -1482,6 +1482,7 @@ bool FAssetRegistryState::Save(FArchive& OriginalAr, const FAssetRegistrySeriali
 
 bool FAssetRegistryState::Load(FArchive& OriginalAr, const FAssetRegistryLoadOptions& Options, FAssetRegistryVersion::Type* OutVersion)
 {
+	LLM_SCOPE(ELLMTag::AssetRegistry);
 	FAssetRegistryHeader Header;
 	Header.SerializeHeader(OriginalAr);
 	if (OutVersion != nullptr)
