AssetRegistry: Disable DirectoryWatchRoots to fix performance while we diagnose why the FCA_RescanRequiredis occurring.
#rnx

diff --git a/Engine/Source/Runtime/AssetRegistry/Private/AssetRegistry.cpp b/Engine/Source/Runtime/AssetRegistry/Private/AssetRegistry.cpp
--- a/Engine/Source/Runtime/AssetRegistry/Private/AssetRegistry.cpp
+++ b/Engine/Source/Runtime/AssetRegistry/Private/AssetRegistry.cpp
@@ -1001,6 +1001,9 @@ void UAssetRegistryImpl::InitializeEvents(UE::AssetRegistry::Impl::FInitializeCo
 
 		if (DirectoryWatcher)
 		{
+			// Temporarily disabling the DirectoryWatchRoots until we diagnose why FCA_RescanRequired is being sent on
+			// editor startup of projects with a large number of plugins
+#if 0
 			// The vast majority of directories we are watching are below the Plugin directories. The memory cost per watch
 			// is sufficiently high to want to avoid setting up many granular watches when we can also setup two coarse ones.
 
@@ -1014,6 +1017,7 @@ void UAssetRegistryImpl::InitializeEvents(UE::AssetRegistry::Impl::FInitializeCo
 			{
 				DirectoryWatchRoots.Add(EnginePluginDir);
 			}
+#endif
 
 			for (FString& WatchRoot : DirectoryWatchRoots)
 			{
