AssetRegistry: Restore the DirectoryWatchRoots. We have fixed the expensive rescan caused by DirectoryWatcher's buffer overflowing when launching into PIE immediately with UnrealEditor.exe <ProjectName> <MapName> -PIE.
#rnx

diff --git a/Engine/Source/Editor/UnrealEd/Private/UnrealEdEngine.cpp b/Engine/Source/Editor/UnrealEd/Private/UnrealEdEngine.cpp
--- a/Engine/Source/Editor/UnrealEd/Private/UnrealEdEngine.cpp
+++ b/Engine/Source/Editor/UnrealEd/Private/UnrealEdEngine.cpp
@@ -83,6 +83,8 @@
 #include "ObjectTools.h"
 #include "Cooker/ExternalCookOnTheFlyServer.h"
 #include "ISettingsSection.h"
+#include "DirectoryWatcherModule.h"
+#include "IDirectoryWatcher.h"
 
 DEFINE_LOG_CATEGORY_STATIC(LogUnrealEdEngine, Log, All);
 
@@ -118,9 +120,19 @@ void UUnrealEdEngine::Init(IEngineLoop* InEngineLoop)
 	// Set-up the initial set of content mount paths to check for write permision
 	TArray<FString> RootPaths;
 	FPackageName::QueryRootContentPaths(RootPaths);
+	int32 RootPathIndex = 0;
+	FDirectoryWatcherModule* DirectoryWatcherModule = FModuleManager::GetModulePtr<FDirectoryWatcherModule>(TEXT("DirectoryWatcher"));
+	IDirectoryWatcher* DirectoryWatcher = DirectoryWatcherModule ? DirectoryWatcherModule->Get() : nullptr;
+
 	for (const FString& RootPath : RootPaths)
 	{
 		VerifyMountPointWritePermission(*RootPath);
+		// We're writing files to test the write permission, and this will send file events to the directory watcher. Tick
+		// every few writes so that we don't overflow the directory watcher's buffer and cause a FCA_RescanRequired event.
+		if (DirectoryWatcher && ((++RootPathIndex) % 10) == 0)
+		{
+			DirectoryWatcher->Tick(-1.0f);
+		}
 	}
 	// Watch for new content mount paths
 	FPackageName::OnContentPathMounted().AddUObject(this, &UUnrealEdEngine::OnContentPathMounted);
diff --git a/Engine/Source/Runtime/AssetRegistry/Private/AssetRegistry.cpp b/Engine/Source/Runtime/AssetRegistry/Private/AssetRegistry.cpp
--- a/Engine/Source/Runtime/AssetRegistry/Private/AssetRegistry.cpp
+++ b/Engine/Source/Runtime/AssetRegistry/Private/AssetRegistry.cpp
@@ -1003,7 +1003,6 @@ void UAssetRegistryImpl::InitializeEvents(UE::AssetRegistry::Impl::FInitializeCo
 		{
 			// Temporarily disabling the DirectoryWatchRoots until we diagnose why FCA_RescanRequired is being sent on
 			// editor startup of projects with a large number of plugins
-#if 0
 			// The vast majority of directories we are watching are below the Plugin directories. The memory cost per watch
 			// is sufficiently high to want to avoid setting up many granular watches when we can also setup two coarse ones.
 
@@ -1017,7 +1016,6 @@ void UAssetRegistryImpl::InitializeEvents(UE::AssetRegistry::Impl::FInitializeCo
 			{
 				DirectoryWatchRoots.Add(EnginePluginDir);
 			}
-#endif
 
 			for (FString& WatchRoot : DirectoryWatchRoots)
 			{
@@ -1835,28 +1833,45 @@ void FAssetRegistryImpl::SearchAllAssets(Impl::FEventContext& EventContext,
 void UAssetRegistryImpl::WaitForCompletion()
 {
 	TRACE_CPUPROFILER_EVENT_SCOPE(UAssetRegistryImpl::WaitForCompletion);
+	LLM_SCOPE(ELLMTag::AssetRegistry);
 
 	using namespace UE::AssetRegistry::Impl;
 
+	bool bInitialSearchStarted = false;
+	bool bInitialSearchCompleted = false;
 	// Try taking over the gather thread for a short time in case it is mostly done.
 	// But if it has more than a small amount of work to do, let the gather thread do that work
 	// while we consume the results in parallel.
 	{
-		LLM_SCOPE(ELLMTag::AssetRegistry);
 		FWriteScopeLock InterfaceScopeLock(InterfaceLock);
 		FClassInheritanceContext InheritanceContext;
 		FClassInheritanceBuffer InheritanceBuffer;
 		GetInheritanceContextWithRequiredLock(InterfaceScopeLock, InheritanceContext, InheritanceBuffer);
 		constexpr float TimeToJoinSeconds = 0.100f;
 		GuardedData.WaitForGathererIdle(TimeToJoinSeconds);
+		bInitialSearchStarted = GuardedData.IsInitialSearchStarted();
+		bInitialSearchCompleted = GuardedData.IsInitialSearchCompleted();
 	}
 
+#if WITH_EDITOR
+	if (bInitialSearchStarted && !bInitialSearchCompleted)
+	{
+		// If we do need to wait, then tick the DirectoryWatcher so we have the most up to date information.
+		// This is also important because we ignore rescan events from the directory watcher if they are sent
+		// during startup, so if there is a rescan event pending we want to trigger it now and ignore it
+		if (GIsEditor) 	// In-game doesn't listen for directory changes
+		{
+			FDirectoryWatcherModule& DirectoryWatcherModule = FModuleManager::LoadModuleChecked<FDirectoryWatcherModule>(TEXT("DirectoryWatcher"));
+			DirectoryWatcherModule.Get()->Tick(-1.f);
+		}
+	}
+#endif
+
 	for (;;)
 	{
 		FEventContext EventContext;
 		EGatherStatus Status;
 		{
-			LLM_SCOPE(ELLMTag::AssetRegistry);
 			FWriteScopeLock InterfaceScopeLock(InterfaceLock);
 			FClassInheritanceContext InheritanceContext;
 			FClassInheritanceBuffer InheritanceBuffer;
