AssetRegistry: Remove accidentally added comment, and remove poorly thought out shift of LLM_SCOPE.
#rnx

diff --git a/Engine/Source/Runtime/AssetRegistry/Private/AssetRegistry.cpp b/Engine/Source/Runtime/AssetRegistry/Private/AssetRegistry.cpp
--- a/Engine/Source/Runtime/AssetRegistry/Private/AssetRegistry.cpp
+++ b/Engine/Source/Runtime/AssetRegistry/Private/AssetRegistry.cpp
@@ -1001,8 +1001,6 @@ void UAssetRegistryImpl::InitializeEvents(UE::AssetRegistry::Impl::FInitializeCo
 
 		if (DirectoryWatcher)
 		{
-			// Temporarily disabling the DirectoryWatchRoots until we diagnose why FCA_RescanRequired is being sent on
-			// editor startup of projects with a large number of plugins
 			// The vast majority of directories we are watching are below the Plugin directories. The memory cost per watch
 			// is sufficiently high to want to avoid setting up many granular watches when we can also setup two coarse ones.
 
@@ -1833,7 +1831,6 @@ void FAssetRegistryImpl::SearchAllAssets(Impl::FEventContext& EventContext,
 void UAssetRegistryImpl::WaitForCompletion()
 {
 	TRACE_CPUPROFILER_EVENT_SCOPE(UAssetRegistryImpl::WaitForCompletion);
-	LLM_SCOPE(ELLMTag::AssetRegistry);
 
 	using namespace UE::AssetRegistry::Impl;
 
@@ -1843,6 +1840,7 @@ void UAssetRegistryImpl::WaitForCompletion()
 	// But if it has more than a small amount of work to do, let the gather thread do that work
 	// while we consume the results in parallel.
 	{
+		LLM_SCOPE(ELLMTag::AssetRegistry);
 		FWriteScopeLock InterfaceScopeLock(InterfaceLock);
 		FClassInheritanceContext InheritanceContext;
 		FClassInheritanceBuffer InheritanceBuffer;
@@ -1872,6 +1870,8 @@ void UAssetRegistryImpl::WaitForCompletion()
 		FEventContext EventContext;
 		EGatherStatus Status;
 		{
+			// Keep the LLM scope limited so it does not surround the broadcast which calls external code
+			LLM_SCOPE(ELLMTag::AssetRegistry);
 			FWriteScopeLock InterfaceScopeLock(InterfaceLock);
 			FClassInheritanceContext InheritanceContext;
 			FClassInheritanceBuffer InheritanceBuffer;
