AssetRegistry: Fix IAssetRegistry::GetRedirectedObjectPath  to handle subobjects that have had their outer redirected.
#rnx

#virtualized

diff --git a/Engine/Source/Runtime/AssetRegistry/Private/AssetRegistry.cpp b/Engine/Source/Runtime/AssetRegistry/Private/AssetRegistry.cpp
--- a/Engine/Source/Runtime/AssetRegistry/Private/AssetRegistry.cpp
+++ b/Engine/Source/Runtime/AssetRegistry/Private/AssetRegistry.cpp
@@ -3003,7 +3003,16 @@ FSoftObjectPath FAssetRegistryImpl::GetRedirectedObjectPath(const FSoftObjectPat
 		// For legacy behavior, for all redirects after the initial request, we only check on-disk assets
 		RedirectedPath = FSoftObjectPath(Redirector->DestinationObject);
 	}
+
+	FString SubPathString;
 	const FAssetData* AssetData = State.GetAssetByObjectPath(RedirectedPath);
+	if (!AssetData && RedirectedPath.IsSubobject())
+	{
+		// If we found no Asset because it is a subobject, then look for its toplevelobject's Asset
+		SubPathString = RedirectedPath.GetSubPathString();
+		RedirectedPath = FSoftObjectPath(RedirectedPath.GetAssetPath(), FString());
+		AssetData = State.GetAssetByObjectPath(RedirectedPath);
+	}
 
 	// Most of the time this will either not be a redirector or only have one redirect, so optimize for that case
 	TArray<FSoftObjectPath, TInlineAllocator<2>> SeenPaths = { RedirectedPath };
@@ -3031,6 +3040,19 @@ FSoftObjectPath FAssetRegistryImpl::GetRedirectedObjectPath(const FSoftObjectPat
 		AssetData = State.GetAssetByObjectPath(RedirectedPath);
 	}
 
+	if (!SubPathString.IsEmpty())
+	{
+		if (!RedirectedPath.IsSubobject())
+		{
+			RedirectedPath.SetSubPathString(SubPathString);
+		}
+		else
+		{
+			// A complicated case; the redirector pointed to a subobject. Append old subobject path onto the new one
+			// Appending old to new will always use '.' because only the first subobject uses ':'
+			RedirectedPath.SetSubPathString(RedirectedPath.GetSubPathString() + TEXT(".") + SubPathString);
+		}
+	}
 	return RedirectedPath;
 }
 
