AssetRegistry: AssetDataGatherer deadlock fix: don't use AsyncTasks to read the NameBatch section out of the AssetRegistry cache files. It should not be necessary for performance since we are already going wide reading multiple cache files at once, and trying to use AsyncTasks from our ParallelFor can cause a deadlock if the number of worker threads available is too low.
#rnx

diff --git a/Engine/Source/Runtime/AssetRegistry/Private/AssetDataGatherer.cpp b/Engine/Source/Runtime/AssetRegistry/Private/AssetDataGatherer.cpp
--- a/Engine/Source/Runtime/AssetRegistry/Private/AssetDataGatherer.cpp
+++ b/Engine/Source/Runtime/AssetRegistry/Private/AssetDataGatherer.cpp
@@ -4940,15 +4940,18 @@ TArray<FCachePayload> LoadCacheFiles(TConstArrayView<FString> InCacheFilenames,
 		const FString& CacheFilename = InCacheFilenames[Index];
 		auto DoLoad = [&](FArchive& ChecksummingReader)
 		{
-			int32 WorkerReduction = 2; // Current worker + preload task
-			int32 Parallelism = FMath::Max(FTaskGraphInterface::Get().GetNumWorkerThreads() - WorkerReduction, 0);
+			// We are already using AsyncTasks for our own ParallelFor, passing non-zero NumAsyncWorkers into
+			// FAssetRegistryReader will cause our current task to block waiting for another task to run
+			// the name batch loading. That can cause a deadlock due to running out of task threads
+			// if we don't have a large number of threads, so disable the parallelism.
+			constexpr int32 NumAsyncWorkers = 0;
 			
 			// The discovery cache is always serialized with a fixed format.
 			// We discard it before this point if it's not the latest version, and it always includes editor-only data.
 			FAssetRegistryHeader Header;
 			Header.Version = FAssetRegistryVersion::LatestVersion;
 			Header.bFilterEditorOnlyData = false;
-			FAssetRegistryReader RegistryReader(ChecksummingReader, Parallelism, Header);
+			FAssetRegistryReader RegistryReader(ChecksummingReader, NumAsyncWorkers, Header);
 			return RegistryReader.IsError() ? FCachePayload() : SerializeCacheLoad(RegistryReader);
 		};
 
