AssetRegistry: DiscoveryCache: Tweak the wording of the log messages for PlatformFileJournal not available to remove the terms failure and error. This failure is expected in many cases and is completely handled by slower fallbacks, so we do not want it to appear to be an error when diagnosing issues.
#rnx

diff --git a/Engine/Source/Runtime/AssetRegistry/Private/AssetDataGathererDiscoveryCache.cpp b/Engine/Source/Runtime/AssetRegistry/Private/AssetDataGathererDiscoveryCache.cpp
--- a/Engine/Source/Runtime/AssetRegistry/Private/AssetDataGathererDiscoveryCache.cpp
+++ b/Engine/Source/Runtime/AssetRegistry/Private/AssetDataGathererDiscoveryCache.cpp
@@ -77,7 +77,7 @@ void FAssetDataDiscoveryCache::LoadAndUpdateCache()
 		const TCHAR* MissingOperation = (!bReadEnabled && WriteEnabled == EFeatureEnabled::Never) ? TEXT("read or written") :
 			(!bReadEnabled ? TEXT("read") : TEXT("written"));
 		UE_LOG(LogAssetRegistry, Display,
-			TEXT("PlatformFileJournal query failed for volume '%s' of project directory '%s', so AssetDiscovery cache will not be %s. PlatformFileJournal error information:\n\t%s"),
+			TEXT("PlatformFileJournal is not available on volume '%s' of project directory '%s', so AssetDiscovery cache will not be %s. Unavailability reason:\n\t%s"),
 			*TestVolumeName, *ProjectDir, MissingOperation, *TestError);
 	}
 
@@ -121,7 +121,7 @@ void FAssetDataDiscoveryCache::LoadAndUpdateCache()
 		if (!VolumeInfo.bJournalAvailable)
 		{
 			UE_LOG(LogAssetRegistry, Warning,
-				TEXT("PlatformFileJournal query failed for volume '%s'. AssetRegistry discovery of files on this volume will be uncached. PlatformFileJournal error information:\n\t%s"),
+				TEXT("PlatformFileJournal is not available on volume '%s'. AssetRegistry discovery of files on this volume will be uncached. Unavailability reason:\n\t%s"),
 				*VolumeName, *VolumeInfo.LastError);
 		}
 		else
@@ -142,7 +142,7 @@ void FAssetDataDiscoveryCache::LoadAndUpdateCache()
 				break;
 			default:
 				UE_LOG(LogAssetRegistry, Warning,
-					TEXT("PlatformFileJournal query failed for volume '%s'. AssetRegistry discovery of files on this volume will be uncached. PlatformFileJournal error information:")
+					TEXT("PlatformFileJournal is not available for volume '%s'. AssetRegistry discovery of files on this volume will be uncached. Unavailability reason:")
 					TEXT("\n\t%s"),
 					*VolumeName, *VolumeInfo.LastError);
 				break;
