AssetRegistry: Performance Bugfix: Fix spurious rescan of parent directory when a directory is specified with ScanPathsSynchronous(bForceRescan=true).
#rnx

diff --git a/Engine/Source/Runtime/AssetRegistry/Private/AssetDataGatherer.cpp b/Engine/Source/Runtime/AssetRegistry/Private/AssetDataGatherer.cpp
--- a/Engine/Source/Runtime/AssetRegistry/Private/AssetDataGatherer.cpp
+++ b/Engine/Source/Runtime/AssetRegistry/Private/AssetDataGatherer.cpp
@@ -870,6 +870,22 @@ bool FScanDir::TrySetDirectoryProperties(FStringView InRelPath, FInherited& Pare
 					bConfirmedExists = true;
 				}
 				SubDir = &FindOrAddSubDir(FirstComponent);
+				// If the current directory has already been scanned then the SubDir we just created must have been
+				// previously discovered, or it was created on disk after the last time we scanned the current
+				// directory. If it was created on disk, and it is not being force rescanned
+				// (aka Properties->HasScanned=true) then we are allowed to ignore it. If it was previously discovered,
+				// then we completed and deleted it: either it was not IsMonitored, or we scanned it. To avoid an
+				// unnecessary rescan, we should therefore set bHasScanned=true if the current directory has
+				// bHasScanned=true and the SubDir has IsMonitored=true. If force rescan is requested on the SubDir,
+				// then we will set it back to bHasScanned=false in TrySetDirectoryProperties below.
+				if (this->bHasScanned)
+				{
+					FInherited ParentDataForSubDir(ParentData, DirectData);
+					if (SubDir->IsMonitored(ParentDataForSubDir))
+					{
+						SubDir->bHasScanned = true;
+					}
+				}
 				SetComplete(false);
 			}
 		}
