AssetRegistry: Speculative fix for missing AssetRegistry information on Windows11. Turn off use of the FileJournalIterateDirectory API function when the Journal information is not required; use the known-safe IterateDirectoryStat instead.
#rnx

diff --git a/Engine/Source/Runtime/AssetRegistry/Private/AssetDataGatherer.cpp b/Engine/Source/Runtime/AssetRegistry/Private/AssetDataGatherer.cpp
--- a/Engine/Source/Runtime/AssetRegistry/Private/AssetDataGatherer.cpp
+++ b/Engine/Source/Runtime/AssetRegistry/Private/AssetDataGatherer.cpp
@@ -2018,9 +2018,8 @@ void FAssetDataDiscovery::TickInternal(bool bTickAll)
 			}
 
 			FCachedDirScanDir CacheDataToAdd;
-			FPlatformFileManager::Get().GetPlatformFile().FileJournalIterateDirectory(Data.DirLocalAbsPath.ToString(),
-			[this, &Data, &CacheDataToAdd]
-			(const TCHAR* IterFilename, const FFileJournalData& IterData)
+			auto ProcessIterData = [this, &Data, &CacheDataToAdd]
+			(const TCHAR* IterFilename, bool bIsDirectory, const FDateTime& ModificationTime, FFileJournalFileHandle JournalHandle)
 			{
 				FStringView LocalAbsPath(IterFilename);
 				FStringView RelPath;
@@ -2052,12 +2051,12 @@ void FAssetDataDiscovery::TickInternal(bool bTickAll)
 					Data.DirLongPackageName.RemoveSuffix(Data.DirLongPackageName.Len() - DirLongPackageRootNameLen);
 				};
 
-				if (IterData.bIsDirectory)
+				if (bIsDirectory)
 				{
 					if (Cache.IsWriteEnabled() != EFeatureEnabled::Never)
 					{
 						CacheDataToAdd.SubDirRelPaths.Add(FString(RelPath));
-						Cache.QueueAdd(FString(LocalAbsPath), IterData.JournalHandle);
+						Cache.QueueAdd(FString(LocalAbsPath), JournalHandle);
 					}
 
 					FPathViews::AppendPath(Data.DirLongPackageName, RelPath);
@@ -2080,7 +2079,7 @@ void FAssetDataDiscovery::TickInternal(bool bTickAll)
 					{
 						FCachedDirScanFile& FileData = CacheDataToAdd.Files.Emplace_GetRef();
 						FileData.RelPath = FString(RelPath);
-						FileData.ModificationTime = IterData.ModificationTime;
+						FileData.ModificationTime = ModificationTime;
 					}
 					EGatherableFileType FileType = GetFileType(RelPath);
 					// Don't record files that contain invalid packagepath characters (not counting their extension)
@@ -2098,17 +2097,33 @@ void FAssetDataDiscovery::TickInternal(bool bTickAll)
 							FPathViews::AppendPath(Data.DirLongPackageName, BaseName);
 							const bool bBlocked = FileType == EGatherableFileType::PackageFile && IsPackageBlocked(LocalAbsPath);
 							Data.IteratedFiles[Data.NumIteratedFiles++].Assign(LocalAbsPath, Data.DirLongPackageName, RelPath,
-								IterData.ModificationTime, FileType, bBlocked);
+								ModificationTime, FileType, bBlocked);
 						}
 					}
 				}
 				return true;
-			});
+			};
 
 			if (Cache.IsWriteEnabled() != EFeatureEnabled::Never)
 			{
+				FPlatformFileManager::Get().GetPlatformFile().FileJournalIterateDirectory(Data.DirLocalAbsPath.ToString(),
+					[this, &Data, &CacheDataToAdd, &ProcessIterData]
+					(const TCHAR* IterFilename, const FFileJournalData& IterData)
+					{
+						return ProcessIterData(IterFilename, IterData.bIsDirectory, IterData.ModificationTime, IterData.JournalHandle);
+					});
+
 				Cache.QueueAdd(FString(Data.DirLocalAbsPath), MoveTemp(CacheDataToAdd));
 			}
+			else
+			{
+				IFileManager::Get().IterateDirectoryStat(Data.DirLocalAbsPath.ToString(),
+					[this, &Data, &ProcessIterData]
+					(const TCHAR* IterFilename, const FFileStatData& IterData)
+					{
+						return ProcessIterData(IterFilename, IterData.bIsDirectory, IterData.ModificationTime, FFileJournalFileHandle());
+					});
+			}
 
 			Data.bScanned = true;
 		}, EParallelForFlags::BackgroundPriority);
