Improves DoesPackageExistEx by enabling it to use the AssetRegistry when it's available

Before Boot -> PIE
FPackageName::DoesPackageExistEx called 191284 times takes 16.45s in total

After Boot -> PIE
FPackageName::DoesPackageExistEx called 191743 times takes 1.61s in total

The majority of the 1.6s (1s) is from calling DoesPackageExistEx before the AssetRegistry has initialized and can provide a fast path. The 180K calls afterwards consume a total of 600ms compared to before, which took 15s.

The original DoesPackageExistEx code checked for package existence on the filesystem, and at the same time found the case matching package name (the casing the file system uses for the package name). This information is known to the AssetRegistry so we can avoid the duplicate work by deferring to the AssetRegistry.

The AssetRegistry needs to initialize itself, and previously did not expose a means to get the Case Matching FNames for Packages which are being used as the keys to AssetPackageData. We add a new overload to IAssetRegistryInterface::TryGetAssetPackageData to allow CoreUObject code to access the another new overload in the AssetRegistry::GetAssetPackageData which returns the case matching FName for the package as an out parameter. Since IAssetRegistryInterface can be used while the AssetRegistry doesn't exist or is still intializing, IAssetRegistryInterface::TryGetAssetPackageData will return UE::AssetRegistry::EExists::Unknown when the asset registry is not yet available. In the Unknown case, we fallback to the original FPackageName::DoesPackageExistEx codepath which reads from disk.

I've made the (Try)GetAssetPackageData methods that only return FAssetPackageData* to use the new implementation that returns a FAssetPackageData* and populates an FName in order to reduce code duplication, however I'm open to changing this if the extra store is a concern. All paths perform the key search in the AssetRegistry's CachedAssetPackageData map.

diff --git a/Engine/Source/Runtime/AssetRegistry/Private/AssetRegistry.cpp b/Engine/Source/Runtime/AssetRegistry/Private/AssetRegistry.cpp
--- a/Engine/Source/Runtime/AssetRegistry/Private/AssetRegistry.cpp
+++ b/Engine/Source/Runtime/AssetRegistry/Private/AssetRegistry.cpp
@@ -204,13 +204,19 @@ public:
 	}
 
 	virtual UE::AssetRegistry::EExists TryGetAssetPackageData(FName PackageName, class FAssetPackageData& OutPackageData) const override
+	{
+		FName OutCorrectCasePackageName;
+		return TryGetAssetPackageData(PackageName, OutPackageData, OutCorrectCasePackageName);
+	}
+	
+	virtual UE::AssetRegistry::EExists TryGetAssetPackageData(FName PackageName, class FAssetPackageData& OutPackageData, FName& OutCorrectCasePackageName) const override
 	{
 		auto AssetRegistry = IAssetRegistry::Get();
 		if (!AssetRegistry)
 		{
 			return UE::AssetRegistry::EExists::Unknown;
 		}
-		return AssetRegistry->TryGetAssetPackageData(PackageName, OutPackageData);
+		return AssetRegistry->TryGetAssetPackageData(PackageName, OutPackageData, OutCorrectCasePackageName);
 	}
 };
 FAssetRegistryInterface GAssetRegistryInterface;
@@ -2547,11 +2553,17 @@ UE::AssetRegistry::EExists UAssetRegistryImpl::TryGetAssetByObjectPath(const FSo
 }
 
 UE::AssetRegistry::EExists UAssetRegistryImpl::TryGetAssetPackageData(const FName PackageName, FAssetPackageData& OutAssetPackageData) const
+{
+	FName OutCorrectCasePackageName;
+	return TryGetAssetPackageData(PackageName, OutAssetPackageData, OutCorrectCasePackageName);
+}
+
+UE::AssetRegistry::EExists UAssetRegistryImpl::TryGetAssetPackageData(const FName PackageName, FAssetPackageData& OutAssetPackageData, FName& OutCorrectCasePackageName) const
 {
 	FReadScopeLock InterfaceScopeLock(InterfaceLock);
 	bool bAssetRegistryReady = GuardedData.IsInitialSearchStarted() && GuardedData.IsInitialSearchCompleted();
 	const FAssetRegistryState& State = GuardedData.GetState();
-	const FAssetPackageData* FoundData = State.GetAssetPackageData(PackageName);
+	const FAssetPackageData* FoundData = State.GetAssetPackageData(PackageName, OutCorrectCasePackageName);
 	if (!FoundData)
 	{
 		if (!bAssetRegistryReady)
@@ -2882,10 +2894,11 @@ bool UAssetRegistryImpl::DoesPackageExistOnDisk(FName PackageName, FString* OutC
 			return false;
 		}
 
+		FName CorrectCasePackageName;
 		const FAssetPackageData* AssetPackageData;
 		{
 			FReadScopeLock InterfaceScopeLock(InterfaceLock);
-			AssetPackageData = GuardedData.GetState().GetAssetPackageData(PackageName);
+			AssetPackageData = GuardedData.GetState().GetAssetPackageData(PackageName, CorrectCasePackageName);
 		}
 		const static bool bVerifyNegativeResults = FParse::Param(FCommandLine::Get(), TEXT("AssetRegistryValidatePackageExists"));
 		if (bVerifyNegativeResults && !AssetPackageData)
@@ -2913,8 +2926,7 @@ bool UAssetRegistryImpl::DoesPackageExistOnDisk(FName PackageName, FString* OutC
 
 		if (OutCorrectCasePackageName)
 		{
-			// TODO: Implement this correctly by saving the information in AssetPackageData if the capitalization does not match the FName capitalization
-			*OutCorrectCasePackageName = PackageNameStr;
+			*OutCorrectCasePackageName = CorrectCasePackageName.ToString();
 		}
 		if (OutExtension)
 		{
diff --git a/Engine/Source/Runtime/AssetRegistry/Private/AssetRegistry.h b/Engine/Source/Runtime/AssetRegistry/Private/AssetRegistry.h
--- a/Engine/Source/Runtime/AssetRegistry/Private/AssetRegistry.h
+++ b/Engine/Source/Runtime/AssetRegistry/Private/AssetRegistry.h
@@ -47,6 +47,7 @@ public:
 	virtual FAssetData GetAssetByObjectPath(const FSoftObjectPath& ObjectPath, bool bIncludeOnlyOnDiskAssets = false, bool bSkipARFilteredAssets = true) const override;
 	virtual UE::AssetRegistry::EExists TryGetAssetByObjectPath(const FSoftObjectPath& ObjectPath, FAssetData& OutAssetData) const override;
 	virtual UE::AssetRegistry::EExists TryGetAssetPackageData(const FName PackageName, FAssetPackageData& OutAssetPackageData) const override;
+	virtual UE::AssetRegistry::EExists TryGetAssetPackageData(const FName PackageName, FAssetPackageData& OutAssetPackageData, FName& OutCorrectCasePackageName) const override;
 	virtual bool GetAllAssets(TArray<FAssetData>& OutAssetData, bool bIncludeOnlyOnDiskAssets = false) const override;
 	virtual bool EnumerateAllAssets(TFunctionRef<bool(const FAssetData&)> Callback, bool bIncludeOnlyOnDiskAssets = false) const override;
 	virtual void GetPackagesByName(FStringView PackageName, TArray<FName>& OutPackageNames) const override;
diff --git a/Engine/Source/Runtime/AssetRegistry/Private/AssetRegistryState.cpp b/Engine/Source/Runtime/AssetRegistry/Private/AssetRegistryState.cpp
--- a/Engine/Source/Runtime/AssetRegistry/Private/AssetRegistryState.cpp
+++ b/Engine/Source/Runtime/AssetRegistry/Private/AssetRegistryState.cpp
@@ -2467,12 +2467,29 @@ const FAssetPackageData* FAssetRegistryState::GetAssetPackageData(FName PackageN
 	return FoundData ? *FoundData : nullptr;
 }
 
- FAssetPackageData* FAssetRegistryState::GetAssetPackageData(FName PackageName)
+FAssetPackageData* FAssetRegistryState::GetAssetPackageData(FName PackageName)
 {
 	FAssetPackageData** FoundData = CachedPackageData.Find(PackageName);
 	return FoundData ? *FoundData : nullptr;
 }
 
+const FAssetPackageData* FAssetRegistryState::GetAssetPackageData(FName PackageName, FName& OutCorrectCasePackageName) const
+{
+	// CachedPackageData is keyed using the Package Names whose casing matches the filesystem. In order to perform a
+	// single look up for the AssetPackageData while also returning the value of the key used to add to the map originally
+	// we create a KeyIterator which is currently the only means to get a TPair<Key,Value> from a TMap<Key,Value>
+	TMap<FName, FAssetPackageData*>::TConstKeyIterator It = CachedPackageData.CreateConstKeyIterator(PackageName);
+	FSetElementId Id = It.GetId();
+	if (!Id.IsValidId())
+	{
+		return nullptr;
+	}
+
+	const TPair<FName, FAssetPackageData*>& Pair = CachedPackageData.Get(Id);
+	OutCorrectCasePackageName = Pair.Key;
+	return Pair.Value;
+}
+
 FAssetPackageData* FAssetRegistryState::CreateOrGetAssetPackageData(FName PackageName)
 {
 	FAssetPackageData*& Data = CachedPackageData.FindOrAdd(PackageName);
diff --git a/Engine/Source/Runtime/AssetRegistry/Public/AssetRegistry/AssetRegistryModule.h b/Engine/Source/Runtime/AssetRegistry/Public/AssetRegistry/AssetRegistryModule.h
--- a/Engine/Source/Runtime/AssetRegistry/Public/AssetRegistry/AssetRegistryModule.h
+++ b/Engine/Source/Runtime/AssetRegistry/Public/AssetRegistry/AssetRegistryModule.h
@@ -112,12 +112,18 @@ public:
 	}
 
 	virtual UE::AssetRegistry::EExists TryGetAssetPackageData(FName PackageName, FAssetPackageData& OutAssetPackageData) const override
+	{
+		FName OutCorrectCasePackageName;
+		return TryGetAssetPackageData(PackageName, OutAssetPackageData, OutCorrectCasePackageName);
+	}
+	
+	virtual UE::AssetRegistry::EExists TryGetAssetPackageData(FName PackageName, FAssetPackageData& OutAssetPackageData, FName& OutCorrectCasePackageName) const override
 	{
 		IAssetRegistry* AssetRegistry = IAssetRegistry::Get();
 		if (!AssetRegistry)
 		{
 			return UE::AssetRegistry::EExists::Unknown;
 		}
-		return AssetRegistry->TryGetAssetPackageData(PackageName, OutAssetPackageData);
+		return AssetRegistry->TryGetAssetPackageData(PackageName, OutAssetPackageData, OutCorrectCasePackageName);
 	}
 };
diff --git a/Engine/Source/Runtime/AssetRegistry/Public/AssetRegistry/AssetRegistryState.h b/Engine/Source/Runtime/AssetRegistry/Public/AssetRegistry/AssetRegistryState.h
--- a/Engine/Source/Runtime/AssetRegistry/Public/AssetRegistry/AssetRegistryState.h
+++ b/Engine/Source/Runtime/AssetRegistry/Public/AssetRegistry/AssetRegistryState.h
@@ -433,6 +433,7 @@ public:
 
 	/** Returns pointer to the asset package data */
 	ASSETREGISTRY_API const FAssetPackageData* GetAssetPackageData(FName PackageName) const;
+	ASSETREGISTRY_API const FAssetPackageData* GetAssetPackageData(FName PackageName, FName& OutCorrectCasePackageName) const;
 	ASSETREGISTRY_API FAssetPackageData* GetAssetPackageData(FName PackageName);
 
 	/** Returns all package names */
diff --git a/Engine/Source/Runtime/AssetRegistry/Public/AssetRegistry/IAssetRegistry.h b/Engine/Source/Runtime/AssetRegistry/Public/AssetRegistry/IAssetRegistry.h
--- a/Engine/Source/Runtime/AssetRegistry/Public/AssetRegistry/IAssetRegistry.h
+++ b/Engine/Source/Runtime/AssetRegistry/Public/AssetRegistry/IAssetRegistry.h
@@ -347,6 +347,17 @@ public:
 	 */
 	virtual UE::AssetRegistry::EExists TryGetAssetPackageData(FName PackageName, FAssetPackageData& OutAssetPackageData) const = 0;
 
+	/**
+	 * Tries to get the package data for the specified package name. If found, OutCorrectCasePackageName
+	 * will be populated with the PackageName that matches the casing used by the filesystem
+	 * 
+	 * @param PackageName name of the package
+	 * @param OutAssetPackageData out FAssetPackageData
+	 * @param OutCorrectCasePackageName out FName matching filesystem casing
+	 * @return Return code enum
+	 */
+	virtual UE::AssetRegistry::EExists TryGetAssetPackageData(FName PackageName, FAssetPackageData& OutAssetPackageData, FName& OutCorrectCasePackageName) const = 0;
+	
 	/**
 	 * Gets asset data for all assets in the registry.
 	 * This method may be slow, use a filter if possible to avoid iterating over the entire registry.
@@ -457,6 +468,15 @@ public:
 	 */
 	virtual void EnumerateAllPackages(TFunctionRef<void(FName PackageName, const FAssetPackageData& PackageData)> Callback) const = 0;
 
+	/*
+	* Checks if the specified package exists on disk. If the package exists and the optional parameters are provided, 
+	* OutCorrectCasePackageName will be populated with the case matching package name found on disk and 
+	* OutExtension will be populated with the extension for the package (dot included) if the package contains assets
+	* 
+	* @param PackageName				the name of the package to check for existence on disk (eg, /Game/MyFolder/MyAsset)
+	* @param OutCorrectCasePackageName 	out package name matching the case on disk
+	* @param OutExtension				out extension for packages containing assets (dot included)
+	*/
 	virtual bool DoesPackageExistOnDisk(FName PackageName, FString* OutCorrectCasePackageName = nullptr, FString* OutExtension = nullptr) const = 0;
 
 	/** Uses the asset registry to look for ObjectRedirectors. This will follow the chain of redirectors. It will return the original path if no redirectors are found */
diff --git a/Engine/Source/Runtime/CoreUObject/Private/Misc/PackageName.cpp b/Engine/Source/Runtime/CoreUObject/Private/Misc/PackageName.cpp
--- a/Engine/Source/Runtime/CoreUObject/Private/Misc/PackageName.cpp
+++ b/Engine/Source/Runtime/CoreUObject/Private/Misc/PackageName.cpp
@@ -8,6 +8,7 @@
 
 #include "Algo/Find.h"
 #include "Algo/FindLast.h"
+#include "AssetRegistry/AssetData.h"
 #include "Containers/DirectoryTree.h"
 #include "Containers/StringView.h"
 #include "GenericPlatform/GenericPlatformFile.h"
@@ -18,6 +19,7 @@
 #include "Internationalization/PackageLocalizationManager.h"
 #include "IO/IoDispatcher.h"
 #include "Misc/App.h"
+#include "Misc/AssetRegistryInterface.h"
 #include "Misc/AutomationTest.h"
 #include "Misc/CoreDelegates.h"
 #include "Misc/EnumClassFlags.h"
@@ -1853,6 +1855,39 @@ FPackageName::EPackageLocationFilter FPackageName::DoesPackageExistEx(const FPac
 	{
 		return EPackageLocationFilter::None;
 	}
+
+#if WITH_EDITOR
+	IAssetRegistryInterface* AssetRegistry = IAssetRegistryInterface::GetPtr();
+
+	// Todo: The AssetRegistry currently cannot determine if a package comes from the Filesystem 
+	// or cooked content so we avoid registry lookups since we can't provide a reliable Location
+	if (AssetRegistry && ((uint8)Filter & (uint8)EPackageLocationFilter::FileSystem) && !FIoDispatcher::IsInitialized())
+	{
+		FName PackageName = PackagePath.GetPackageFName();
+		FName CorrectCasePackageName;
+		FAssetPackageData AssetPackageData;
+		if (AssetRegistry->TryGetAssetPackageData(PackageName, AssetPackageData, CorrectCasePackageName) == UE::AssetRegistry::EExists::Exists)
+		{
+			if (OutPackagePath)
+			{
+				if (AssetPackageData.Extension != EPackageExtension::Unspecified && AssetPackageData.Extension != EPackageExtension::Custom)
+				{
+					*OutPackagePath = FPackagePath::FromPackageNameUnchecked(MoveTemp(CorrectCasePackageName));
+					OutPackagePath->SetHeaderExtension(AssetPackageData.Extension);
+					return EPackageLocationFilter::FileSystem;
+				}
+				// Intentionally do not return a EPackageLocationFilter if an OutPackagePath was provided but a custom or 
+				// unspecified extension was used. In such cases we must use the disk scanning code futher down to allow the 
+				// PackageResourceManager to handle the unspecified and custom extension cases
+			}
+			else
+			{
+				return EPackageLocationFilter::FileSystem;
+			}
+		}
+	}
+#endif
+
 	TStringBuilder<256> PackageName;
 	PackagePath.AppendPackageName(PackageName);
 
diff --git a/Engine/Source/Runtime/CoreUObject/Public/Misc/AssetRegistryInterface.h b/Engine/Source/Runtime/CoreUObject/Public/Misc/AssetRegistryInterface.h
--- a/Engine/Source/Runtime/CoreUObject/Public/Misc/AssetRegistryInterface.h
+++ b/Engine/Source/Runtime/CoreUObject/Public/Misc/AssetRegistryInterface.h
@@ -283,7 +283,7 @@ class IAssetRegistryInterface
 {
 public:
 	/**
-	 * Tries to gets a pointer to the active AssetRegistryInterface implementation. 
+	 * Tries to get a pointer to the active AssetRegistryInterface implementation. 
 	 */
 	static COREUOBJECT_API IAssetRegistryInterface* GetPtr();
 
@@ -302,7 +302,7 @@ public:
 	virtual UE::AssetRegistry::EExists TryGetAssetByObjectPath(const FSoftObjectPath& ObjectPath, struct FAssetData& OutAssetData) const = 0;
 
 	/**
-	 * Tries to get the pacakge data for a specified path
+	 * Tries to get the package data for the specified package name
 	 *
 	 * @param PackageName name of the package
 	 * @param OutAssetPackageData out FAssetPackageData
@@ -310,6 +310,17 @@ public:
 	 */
 	virtual UE::AssetRegistry::EExists TryGetAssetPackageData(FName PackageName, class FAssetPackageData& OutPackageData) const = 0;
 
+	/**
+	 * Tries to get the package data for the specified package name. If found, OutCorrectCasePackageName
+	 * will be populated with the PackageName that matches the casing used by the filesystem
+	 * 
+	 * @param PackageName name of the package
+	 * @param OutAssetPackageData out FAssetPackageData
+	 * @param OutCorrectCasePackageName out FName matching filesystem casing
+	 * @return Return code enum
+	 */
+	virtual UE::AssetRegistry::EExists TryGetAssetPackageData(FName PackageName, class FAssetPackageData& OutPackageData, FName& OutCorrectCasePackageName) const = 0;
+
 protected:
 
 	static COREUOBJECT_API IAssetRegistryInterface* Default;
diff --git a/Engine/Source/Runtime/CoreUObject/Tests/ObjectPropertyTest.cpp b/Engine/Source/Runtime/CoreUObject/Tests/ObjectPropertyTest.cpp
--- a/Engine/Source/Runtime/CoreUObject/Tests/ObjectPropertyTest.cpp
+++ b/Engine/Source/Runtime/CoreUObject/Tests/ObjectPropertyTest.cpp
@@ -276,6 +276,11 @@ public:
 		return UE::AssetRegistry::EExists::Exists;
 	}
 
+	virtual UE::AssetRegistry::EExists TryGetAssetPackageData(FName PackageName, FAssetPackageData& OutPackageData, FName& OutCorrectCasePackageName) const override
+	{
+		return UE::AssetRegistry::EExists::Exists;
+	}
+
 	IAssetRegistryInterface* Old;
 	TMap<FSoftObjectPath, FAssetData> AssetData;
 };
