Fix up asset registry's generated class value on load before building cached class hierarchy, this ensures that assets that have been copied on the file system don't collide

diff --git a/Engine/Source/Runtime/AssetRegistry/Private/AssetRegistry.cpp b/Engine/Source/Runtime/AssetRegistry/Private/AssetRegistry.cpp
--- a/Engine/Source/Runtime/AssetRegistry/Private/AssetRegistry.cpp
+++ b/Engine/Source/Runtime/AssetRegistry/Private/AssetRegistry.cpp
@@ -5914,7 +5914,8 @@ void FAssetRegistryImpl::AddAssetData(Impl::FEventContext& EventContext, FAssetD
 		const FString ParentClass = AssetData->GetTagValueRef<FString>(FBlueprintTags::ParentClassPath);
 		if (!GeneratedClass.IsEmpty() && !ParentClass.IsEmpty() && GeneratedClass != TEXTVIEW("None") && ParentClass != TEXTVIEW("None"))
 		{
-			const FTopLevelAssetPath GeneratedClassPathName(GeneratedClass);
+			const FTopLevelAssetPath SavedGeneratedClassPathName(GeneratedClass);
+			const FTopLevelAssetPath GeneratedClassPathName(AssetData->PackageName, SavedGeneratedClassPathName.GetAssetName());
 			const FTopLevelAssetPath ParentClassPathName(ParentClass);
 			if (ensureAlwaysMsgf(!GeneratedClassPathName.IsNull() && !ParentClassPathName.IsNull(),
 				TEXT("Short class names used in AddAssetData: GeneratedClass=%s, ParentClass=%s. Short class names in these tags on the Blueprint class should have been converted to path names."),
