AssetRegistry: Fix recording of incorrect capitalization of filepaths in the files on disk held by the AssetRegistry when a ScanPathsSynchronous call is made with incorrect capitalization.
#rnx

diff --git a/Engine/Source/Runtime/AssetRegistry/Private/AssetDataGatherer.cpp b/Engine/Source/Runtime/AssetRegistry/Private/AssetDataGatherer.cpp
--- a/Engine/Source/Runtime/AssetRegistry/Private/AssetDataGatherer.cpp
+++ b/Engine/Source/Runtime/AssetRegistry/Private/AssetDataGatherer.cpp
@@ -2332,15 +2332,47 @@ void FPathExistence::LoadExistenceData()
 	FFileStatData StatData = IFileManager::Get().GetStatData(*LocalAbsPath);
 	if (StatData.bIsValid)
 	{
+		FString CorrectedCapitalization = IFileManager::Get().GetFilenameOnDisk(*LocalAbsPath);
+		if (LocalAbsPath == CorrectedCapitalization)
+		{
+			LocalAbsPath = MoveTemp(CorrectedCapitalization);
+		}
+		else
+		{
+			UE_LOG(LogAssetRegistry, Error,
+				TEXT("FPathExistence failed to gather correct capitalization from disk for %s, because GetFilenameOnDisk returned a non-matching filename"),
+				*LocalAbsPath);
+		}
+
 		ModificationTime = StatData.ModificationTime;
 		PathType = StatData.bIsDirectory ? EType::Directory : EType::File;
 	}
 	else
 	{
-		FString ParentPath = FPaths::GetPath(LocalAbsPath);
+		FString ParentPath, BaseName, Extension;
+		FPaths::Split(LocalAbsPath, ParentPath, BaseName, Extension);
 		StatData = IFileManager::Get().GetStatData(*ParentPath);
-		PathType = (StatData.bIsValid && StatData.bIsDirectory)
-			? EType::MissingButDirExists : EType::MissingParentDir;
+		if (StatData.bIsValid && StatData.bIsDirectory)
+		{
+			FString CorrectedCapitalization = IFileManager::Get().GetFilenameOnDisk(*ParentPath);
+			CorrectedCapitalization = FPaths::Combine(CorrectedCapitalization, BaseName) +
+				(!Extension.IsEmpty() ? TEXT(".") : TEXT("")) + Extension;
+			if (LocalAbsPath == CorrectedCapitalization)
+			{
+				LocalAbsPath = MoveTemp(CorrectedCapitalization);
+			}
+			else
+			{
+				UE_LOG(LogAssetRegistry, Error,
+					TEXT("FPathExistence failed to gather correct capitalization from disk for %s, because GetFilenameOnDisk returned a non-matching filename"),
+					*LocalAbsPath);
+			}
+			PathType = EType::MissingButDirExists;
+		}
+		else
+		{
+			PathType = EType::MissingParentDir;
+		}
 	}
 
 	bHasExistenceData = true;
@@ -5318,6 +5350,10 @@ FString FAssetDataGatherer::NormalizeLocalPath(FStringView LocalPath)
 {
 	FString LocalAbsPath(LocalPath);
 	LocalAbsPath = FPaths::ConvertRelativePathToFull(MoveTemp(LocalAbsPath));
+	while (FPathViews::HasRedundantTerminatingSeparator(LocalAbsPath))
+	{
+		LocalAbsPath.LeftChopInline(1);
+	}
 	return LocalAbsPath;
 }
 
diff --git a/Engine/Source/Runtime/Core/Private/Windows/WindowsPlatformFile.cpp b/Engine/Source/Runtime/Core/Private/Windows/WindowsPlatformFile.cpp
--- a/Engine/Source/Runtime/Core/Private/Windows/WindowsPlatformFile.cpp
+++ b/Engine/Source/Runtime/Core/Private/Windows/WindowsPlatformFile.cpp
@@ -1102,7 +1102,11 @@ FString FWindowsPlatformFile::GetFilenameOnDisk(const TCHAR* Filename)
 {
 	FString NormalizedFileName = *FNormalizedFilename(Filename);
 	TRACE_PLATFORMFILE_BEGIN_OPEN(Filename);
-	HANDLE hFile = CreateFile(*NormalizedFileName, FILE_READ_ATTRIBUTES, FILE_SHARE_READ | FILE_SHARE_WRITE | FILE_SHARE_DELETE, NULL, OPEN_EXISTING, 0, NULL);
+	HANDLE hFile = CreateFile(*NormalizedFileName,
+		FILE_READ_ATTRIBUTES, FILE_SHARE_READ | FILE_SHARE_WRITE | FILE_SHARE_DELETE,
+		NULL, OPEN_EXISTING,
+		FILE_FLAG_BACKUP_SEMANTICS, // Required for opening a directory, which Filename might be
+		NULL);
 	// If the file exists on disk, read the capitalization from the path on disk, otherwise just return the (normalized) input filename
 	if (hFile != INVALID_HANDLE_VALUE)
 	{
