AssetRegistry: Downgrade error "FPathExistence failed to gather correct capitalization from disk" to verbose. It is a spurious error and does not indicate a problem.
#rnx

diff --git a/Engine/Source/Runtime/AssetRegistry/Private/AssetDataGatherer.cpp b/Engine/Source/Runtime/AssetRegistry/Private/AssetDataGatherer.cpp
--- a/Engine/Source/Runtime/AssetRegistry/Private/AssetDataGatherer.cpp
+++ b/Engine/Source/Runtime/AssetRegistry/Private/AssetDataGatherer.cpp
@@ -2339,7 +2339,7 @@ void FPathExistence::LoadExistenceData()
 		}
 		else
 		{
-			UE_LOG(LogAssetRegistry, Error,
+			UE_LOG(LogAssetRegistry, Verbose,
 				TEXT("FPathExistence failed to gather correct capitalization from disk for %s, because GetFilenameOnDisk returned a non-matching filename"),
 				*LocalAbsPath);
 		}
@@ -2363,7 +2363,7 @@ void FPathExistence::LoadExistenceData()
 			}
 			else
 			{
-				UE_LOG(LogAssetRegistry, Error,
+				UE_LOG(LogAssetRegistry, Verbose,
 					TEXT("FPathExistence failed to gather correct capitalization from disk for %s, because GetFilenameOnDisk returned a non-matching filename"),
 					*LocalAbsPath);
 			}
