AssetRegistry: Fix  error "FPathExistence failed to gather correct capitalization from disk". LongPackageNames were not being correctly converted to LocalAbsPath in some cases.
#rnx

diff --git a/Engine/Source/Runtime/AssetRegistry/Private/AssetDataGatherer.cpp b/Engine/Source/Runtime/AssetRegistry/Private/AssetDataGatherer.cpp
--- a/Engine/Source/Runtime/AssetRegistry/Private/AssetDataGatherer.cpp
+++ b/Engine/Source/Runtime/AssetRegistry/Private/AssetDataGatherer.cpp
@@ -2339,7 +2339,7 @@ void FPathExistence::LoadExistenceData()
 		}
 		else
 		{
-			UE_LOG(LogAssetRegistry, Verbose,
+			UE_LOG(LogAssetRegistry, Error,
 				TEXT("FPathExistence failed to gather correct capitalization from disk for %s, because GetFilenameOnDisk returned a non-matching filename"),
 				*LocalAbsPath);
 		}
@@ -2363,7 +2363,7 @@ void FPathExistence::LoadExistenceData()
 			}
 			else
 			{
-				UE_LOG(LogAssetRegistry, Verbose,
+				UE_LOG(LogAssetRegistry, Error,
 					TEXT("FPathExistence failed to gather correct capitalization from disk for %s, because GetFilenameOnDisk returned a non-matching filename"),
 					*LocalAbsPath);
 			}
@@ -5221,7 +5221,7 @@ void FAssetDataGatherer::PrioritizeSearchPath(const FString& PathToPrioritize)
 			FGathererScopeLock ResultsScopeLock(&ResultsLock);
 			SetIsIdle(false);
 			int32 NumPrioritizedPaths;
-			UE::AssetDataGather::Private::FPathExistence QueryPath(PathToPrioritize);
+			UE::AssetDataGather::Private::FPathExistence QueryPath(LocalFilenamePathToPrioritize);
 			SortPathsByPriority(TArrayView<UE::AssetDataGather::Private::FPathExistence>(&QueryPath, 1),
 				Priority, NumPrioritizedPaths);
 		}
