Filter unmounted packages when querying a cooked asset registry in the editor

This provides a consistent experience between the cooked and uncooked asset registry

diff --git a/Engine/Source/Runtime/AssetRegistry/Private/AssetRegistry.cpp b/Engine/Source/Runtime/AssetRegistry/Private/AssetRegistry.cpp
--- a/Engine/Source/Runtime/AssetRegistry/Private/AssetRegistry.cpp
+++ b/Engine/Source/Runtime/AssetRegistry/Private/AssetRegistry.cpp
@@ -802,6 +802,7 @@ void FAssetRegistryImpl::LoadPremadeAssetRegistry(Impl::FEventContext& EventCont
 				AppendState(EventContext, ARState, FAssetRegistryState::EInitializationMode::OnlyUpdateNew);
 			}
 			UpdatePersistentMountPoints();
+			State.bCookedGlobalAssetRegistryState = true;
 		}
 		else
 		{
@@ -2524,7 +2525,8 @@ FAssetData UAssetRegistryImpl::GetAssetByObjectPath(const FSoftObjectPath& Objec
 		FReadScopeLock InterfaceScopeLock(InterfaceLock);
 		const FAssetRegistryState& State = GuardedData.GetState();
 		const FAssetData* FoundData = State.GetAssetByObjectPath(ObjectPath);
-		return (FoundData && (!bSkipARFilteredAssets || !GuardedData.ShouldSkipAsset(FoundData->AssetClassPath, FoundData->PackageFlags))) ? *FoundData : FAssetData();
+		return (FoundData && !State.IsPackageUnmountedAndFiltered(FoundData->PackageName)
+			&& (!bSkipARFilteredAssets || !GuardedData.ShouldSkipAsset(FoundData->AssetClassPath, FoundData->PackageFlags))) ? *FoundData : FAssetData();
 	}
 }
 
diff --git a/Engine/Source/Runtime/AssetRegistry/Private/AssetRegistryState.cpp b/Engine/Source/Runtime/AssetRegistry/Private/AssetRegistryState.cpp
--- a/Engine/Source/Runtime/AssetRegistry/Private/AssetRegistryState.cpp
+++ b/Engine/Source/Runtime/AssetRegistry/Private/AssetRegistryState.cpp
@@ -644,17 +644,11 @@ bool FAssetRegistryState::HasAssets(const FName PackagePath, bool bSkipARFiltere
 	const TArray<FAssetData*>* FoundAssetArray = CachedAssetsByPath.Find(PackagePath);
 	if (FoundAssetArray)
 	{
-		if (bSkipARFilteredAssets)
+		return FoundAssetArray->ContainsByPredicate([this, bSkipARFilteredAssets](FAssetData* AssetData)
 		{
-			return FoundAssetArray->ContainsByPredicate([](FAssetData* AssetData)
-			{
-				return AssetData && !UE::AssetRegistry::FFiltering::ShouldSkipAsset(AssetData->AssetClassPath, AssetData->PackageFlags);
-			});
-		}
-		else
-		{
-			return FoundAssetArray->Num() > 0;
-		}
+			return AssetData && !IsPackageUnmountedAndFiltered(AssetData->PackageName)
+				&& (!bSkipARFilteredAssets || !UE::AssetRegistry::FFiltering::ShouldSkipAsset(AssetData->AssetClassPath, AssetData->PackageFlags));
+		});
 	}
 	return false;
 }
@@ -926,7 +920,7 @@ bool FAssetRegistryState::EnumerateAssets(const FARCompiledFilter& Filter, const
 	const uint32 FilterWithoutPackageFlags = Filter.WithoutPackageFlags;
 	const uint32 FilterWithPackageFlags = Filter.WithPackageFlags;
 	auto ShouldSkipAssetData =
-		[&PackageNamesToSkip, bSkipARFilteredAssets, FilterWithoutPackageFlags, FilterWithPackageFlags]
+		[this, &PackageNamesToSkip, bSkipARFilteredAssets, FilterWithoutPackageFlags, FilterWithPackageFlags]
 		(const FAssetData* AssetData)
 		{
 			if (PackageNamesToSkip.Contains(AssetData->PackageName) |			//-V792
@@ -936,6 +930,11 @@ bool FAssetRegistryState::EnumerateAssets(const FARCompiledFilter& Filter, const
 				return true;
 			}
 
+			if (IsPackageUnmountedAndFiltered(AssetData->PackageName))
+			{
+				return true;
+			}
+
 			return bSkipARFilteredAssets &&
 				UE::AssetRegistry::FFiltering::ShouldSkipAsset(AssetData->AssetClassPath, AssetData->PackageFlags);
 		};
@@ -1062,6 +1061,7 @@ bool FAssetRegistryState::EnumerateAllAssets(const TSet<FName>& PackageNamesToSk
 	{
 		if (AssetData &&
 			!PackageNamesToSkip.Contains(AssetData->PackageName) &&
+			!IsPackageUnmountedAndFiltered(AssetData->PackageName) &&
 			(!bSkipARFilteredAssets || !UE::AssetRegistry::FFiltering::ShouldSkipAsset(AssetData->AssetClassPath, AssetData->PackageFlags)))
 		{
 			if (!Callback(*AssetData))
@@ -2543,6 +2543,18 @@ const TArray<const FAssetData*>& FAssetRegistryState::GetAssetsByClassName(const
 	return InvalidArray;
 }
 
+bool FAssetRegistryState::IsPackageUnmountedAndFiltered(const FName PackageName) const
+{
+	// TODO: This can be removed once UE-178174 is fixed, as there will no longer be unmounted content to enumerate
+#if WITH_EDITOR
+	// Note: We currently only perform this filtering in the editor; runtime use will have to perform its own filtering 
+	//       via FPackageName::IsValidPath so that it can choose to accept the additional cost of running that filter
+	return bCookedGlobalAssetRegistryState && GIsEditor && !FPackageName::IsValidPath(WriteToString<256>(PackageName));
+#else
+	return false;
+#endif
+}
+
 namespace UE::AssetRegistry::Utils
 {
 
diff --git a/Engine/Source/Runtime/AssetRegistry/Public/AssetRegistry/AssetRegistryState.h b/Engine/Source/Runtime/AssetRegistry/Public/AssetRegistry/AssetRegistryState.h
--- a/Engine/Source/Runtime/AssetRegistry/Public/AssetRegistry/AssetRegistryState.h
+++ b/Engine/Source/Runtime/AssetRegistry/Public/AssetRegistry/AssetRegistryState.h
@@ -619,6 +619,12 @@ private:
 	void RemoveAssetData(FAssetData* AssetData, const FCachedAssetKey& Key, bool bRemoveDependencyData,
 		bool& bOutRemovedAssetData, bool& bOutRemovedPackageData);
 
+	/**
+	 * Returns true if the given package should be filtered from the results because the package belongs to an unmounted content path.
+	 * This can only happen when loading a cooked asset registry (@see bCookedGlobalAssetRegistryState), as it may contain state for plugins that are not currently loaded.
+	 */
+	bool IsPackageUnmountedAndFiltered(const FName PackageName) const;
+
 	/** Set of asset data for assets saved to disk. Searched via path name types, implicitly converted to FCachedAssetKey. */
 	FAssetDataMap CachedAssets;
 
@@ -650,6 +656,9 @@ private:
 	int32 NumDependsNodes = 0;
 	int32 NumPackageData = 0;
 
+	/** True if this asset registry state was loaded from a cooked asset registry */
+	bool bCookedGlobalAssetRegistryState = false;
+
 	friend class UAssetRegistryImpl;
 	friend class UE::AssetRegistry::FAssetRegistryImpl;
 };
