AssetRegistry: FPathExistence: Restore severity of log "FPathExistence failed to gather correct capitalization from disk" in 5.5 because we have the fix for the bug that caused it.
#rnx

diff --git a/Engine/Source/Runtime/AssetRegistry/Private/AssetDataGatherer.cpp b/Engine/Source/Runtime/AssetRegistry/Private/AssetDataGatherer.cpp
--- a/Engine/Source/Runtime/AssetRegistry/Private/AssetDataGatherer.cpp
+++ b/Engine/Source/Runtime/AssetRegistry/Private/AssetDataGatherer.cpp
@@ -2339,7 +2339,7 @@ void FPathExistence::LoadExistenceData()
 		}
 		else
 		{
-			UE_LOG(LogAssetRegistry, Verbose,
+			UE_LOG(LogAssetRegistry, Error,
 				TEXT("FPathExistence failed to gather correct capitalization from disk for %s, because GetFilenameOnDisk returned a non-matching filename"),
 				*LocalAbsPath);
 		}
@@ -2363,7 +2363,7 @@ void FPathExistence::LoadExistenceData()
 			}
 			else
 			{
-				UE_LOG(LogAssetRegistry, Verbose,
+				UE_LOG(LogAssetRegistry, Error,
 					TEXT("FPathExistence failed to gather correct capitalization from disk for %s, because GetFilenameOnDisk returned a non-matching filename"),
 					*LocalAbsPath);
 			}
