AssetRegistry: Cap the number of threads used during gather since it does not reduce wall time to do so and reduces thread time. Add log diagnostic for wall time.
#rnx

diff --git a/Engine/Source/Runtime/AssetRegistry/Private/AssetDataGatherer.cpp b/Engine/Source/Runtime/AssetRegistry/Private/AssetDataGatherer.cpp
--- a/Engine/Source/Runtime/AssetRegistry/Private/AssetDataGatherer.cpp
+++ b/Engine/Source/Runtime/AssetRegistry/Private/AssetDataGatherer.cpp
@@ -61,6 +61,12 @@ namespace AssetDataGathererConstants
 	 * In theory higher minimum values could reduce the cost of waking up workers, but in practice it does not have a significant effect.
 	 */
 	static int32 GARDiscoverMinBatchSize = 1;
+	/**
+	 * The maximum number of threads used in the ParallelFor around ReadAssetFile in FAssetDataGatherer::TickInternal.
+	 * Very high numbers of threads does not reduce wall time of the gather because the threads become IO limited.
+	 * Capping the number of threads allows those threads to be used in other async operations in Engine startup.
+	 */
+	static int32 GARGatherThreads = 60;
 }
 
 namespace UE::AssetDataGather::Private
@@ -1554,8 +1560,10 @@ FAssetDataDiscovery::FAssetDataDiscovery(const TArray<FString>& InLongPackageNam
 
 	FParse::Value(FCommandLine::Get(), TEXT("-ARDiscoverThreads="), AssetDataGathererConstants::GARDiscoverThreads);
 	FParse::Value(FCommandLine::Get(), TEXT("-ARDiscoverMinBatchSize="), AssetDataGathererConstants::GARDiscoverMinBatchSize);
+	FParse::Value(FCommandLine::Get(), TEXT("-ARGatherThreads="), AssetDataGathererConstants::GARGatherThreads);
 	AssetDataGathererConstants::GARDiscoverThreads = FMath::Max(0, AssetDataGathererConstants::GARDiscoverThreads);
 	AssetDataGathererConstants::GARDiscoverMinBatchSize = FMath::Max(1, AssetDataGathererConstants::GARDiscoverMinBatchSize);
+	AssetDataGathererConstants::GARGatherThreads = FMath::Max(0, AssetDataGathererConstants::GARGatherThreads);
 }
 
 FAssetDataDiscovery::~FAssetDataDiscovery()
@@ -3583,7 +3591,12 @@ FAssetDataGatherer::FAssetDataGatherer(const TArray<FString>& InLongPackageNames
 {
 	using namespace UE::AssetDataGather::Private;
 
-	TickInternalBatchSize = FMath::Max(1, FTaskGraphInterface::Get().GetNumWorkerThreads()) * AssetDataGathererConstants::SingleThreadFilesPerBatch;
+	int32 NumGatherThreads = FMath::Max(1, FTaskGraphInterface::Get().GetNumWorkerThreads());
+	if (AssetDataGathererConstants::GARGatherThreads > 0)
+	{
+		NumGatherThreads = FMath::Min(NumGatherThreads, AssetDataGathererConstants::GARGatherThreads);
+	}
+	TickInternalBatchSize = NumGatherThreads * AssetDataGathererConstants::SingleThreadFilesPerBatch;
 
 	GPreloadSettings.Initialize();
 	bGatherAssetPackageData = GIsEditor || GPreloadSettings.IsForceDependsGathering();
@@ -4181,7 +4194,13 @@ FAssetDataGatherer::ETickResult FAssetDataGatherer::TickInternal(double& TickSta
 			}
 			return ReturnFlags;
 		}();
-	ParallelFor(ReadContexts.Num(),
+	// We want to restrict the number of threads, but ParallelFor only provides an API for restricting MinBatchSize
+	// NumberOfThreads == ParallelForNum/MinBatchSize == TickInternalBatchSize/SingleThreadFilesPerBatch
+	//                 == Min(WorkerThreads, GARGatherThreads)*SingleThreadFilesPerBatch / SingleThreadFilesPerBatch
+	//                 <= GARGatherThreads
+	const int32 MinBatchSize = AssetDataGathererConstants::SingleThreadFilesPerBatch;
+
+	ParallelFor(TEXT("AssetDataGatherReadAssetFile"), ReadContexts.Num(), MinBatchSize,
 		[this, &ReadContexts](int32 Index)
 		{
 			FReadContext& ReadContext = ReadContexts[Index];
@@ -4463,6 +4482,7 @@ FAssetGatherDiagnostics FAssetDataGatherer::GetDiagnostics()
 	Diag.GatherTimeSeconds = CumulativeGatherTime;
 	Diag.NumCachedAssetFiles = NumCachedAssetFiles;
 	Diag.NumUncachedAssetFiles = NumUncachedAssetFiles;
+	Diag.WallTimeSeconds = static_cast<float>((FDateTime::Now() - GatherStartTime).GetTotalSeconds());
 	return Diag;
 }
 
diff --git a/Engine/Source/Runtime/AssetRegistry/Private/AssetDataGatherer.h b/Engine/Source/Runtime/AssetRegistry/Private/AssetDataGatherer.h
--- a/Engine/Source/Runtime/AssetRegistry/Private/AssetDataGatherer.h
+++ b/Engine/Source/Runtime/AssetRegistry/Private/AssetDataGatherer.h
@@ -58,6 +58,8 @@ struct FAssetGatherDiagnostics
 	float DiscoveryTimeSeconds;
 	/** Time spent reading asset files on disk / from cache */
 	float GatherTimeSeconds;
+	/** Time in between gatherer start and the call to GetDiagnostics. */
+	float WallTimeSeconds;
 	/** How many directories in the search results were read from the cache. */
 	int32 NumCachedDirectories;
 	/** How many directories in the search results were not in the cache and were read by scanning the disk. */
diff --git a/Engine/Source/Runtime/AssetRegistry/Private/AssetRegistry.cpp b/Engine/Source/Runtime/AssetRegistry/Private/AssetRegistry.cpp
--- a/Engine/Source/Runtime/AssetRegistry/Private/AssetRegistry.cpp
+++ b/Engine/Source/Runtime/AssetRegistry/Private/AssetRegistry.cpp
@@ -4564,11 +4564,11 @@ void FAssetRegistryImpl::LogSearchDiagnostics(double StartTime)
 	Telemetry.NumCachedAssetFiles = Diagnostics.NumCachedAssetFiles;
 	Telemetry.NumUncachedAssetFiles = Diagnostics.NumUncachedAssetFiles;
 	FTelemetryRouter::Get().ProvideTelemetry(Telemetry);
-	UE_LOG(LogAssetRegistry, Log, TEXT("AssetRegistryGather time %.4fs: AssetDataDiscovery %0.4fs, AssetDataGather %0.4fs, StoreResults %0.4fs.")
+	UE_LOG(LogAssetRegistry, Log, TEXT("AssetRegistryGather time %.4fs: AssetDataDiscovery %0.4fs, AssetDataGather %0.4fs, StoreResults %0.4fs. Wall time %0.4fs.")
 		TEXT("\n\tNumCachedDirectories %d. NumUncachedDirectories %d. NumCachedFiles %d. NumUncachedFiles %d."),
 		Total, Diagnostics.DiscoveryTimeSeconds, Diagnostics.GatherTimeSeconds, StoreGatherResultsTimeSeconds,
-		Diagnostics.NumCachedDirectories, Diagnostics.NumUncachedDirectories, Diagnostics.NumCachedAssetFiles,
-		Diagnostics.NumUncachedAssetFiles);
+		Diagnostics.WallTimeSeconds, Diagnostics.NumCachedDirectories, Diagnostics.NumUncachedDirectories,
+		Diagnostics.NumCachedAssetFiles, Diagnostics.NumUncachedAssetFiles);
 
 #if !NO_LOGGING
 	if (bVerboseLogging)
