Asset Registry: fixed UAssetRegistryImpl::GetRedirectedObjectPath not scanning assets during redirector chain parsing, resulting in redirectors not properly resolved in some cases.
Also removed all 5.1 deprecated asset registry functions and fixed deprecated use cases in our code.

diff --git a/Engine/Plugins/FX/Niagara/Source/NiagaraEditor/Private/NiagaraNodeFunctionCall.cpp b/Engine/Plugins/FX/Niagara/Source/NiagaraEditor/Private/NiagaraNodeFunctionCall.cpp
--- a/Engine/Plugins/FX/Niagara/Source/NiagaraEditor/Private/NiagaraNodeFunctionCall.cpp
+++ b/Engine/Plugins/FX/Niagara/Source/NiagaraEditor/Private/NiagaraNodeFunctionCall.cpp
@@ -338,9 +338,7 @@ void UNiagaraNodeFunctionCall::AllocateDefaultPins()
 	if (FunctionScriptAssetObjectPath != NAME_None && FunctionScript == nullptr)
 	{
 		FAssetRegistryModule& AssetRegistryModule = FModuleManager::LoadModuleChecked<FAssetRegistryModule>(TEXT("AssetRegistry"));
-PRAGMA_DISABLE_DEPRECATION_WARNINGS;
-		FAssetData ScriptAssetData = AssetRegistryModule.Get().GetAssetByObjectPath(FunctionScriptAssetObjectPath);
-PRAGMA_ENABLE_DEPRECATION_WARNINGS;
+		FAssetData ScriptAssetData = AssetRegistryModule.Get().GetAssetByObjectPath(FSoftObjectPath(FunctionScriptAssetObjectPath.ToString()));
 		if (ScriptAssetData.IsValid())
 		{
 			FunctionScript = Cast<UNiagaraScript>(ScriptAssetData.GetAsset());
@@ -557,9 +555,7 @@ bool UNiagaraNodeFunctionCall::CanAddToGraph(UNiagaraGraph* TargetGraph, FString
 	if (FunctionScriptAssetObjectPath != NAME_None && FunctionScript == nullptr)
 	{
 		FAssetRegistryModule& AssetRegistryModule = FModuleManager::LoadModuleChecked<FAssetRegistryModule>(TEXT("AssetRegistry"));
-PRAGMA_DISABLE_DEPRECATION_WARNINGS;
-		FAssetData ScriptAssetData = AssetRegistryModule.Get().GetAssetByObjectPath(FunctionScriptAssetObjectPath);
-PRAGMA_ENABLE_DEPRECATION_WARNINGS;
+		FAssetData ScriptAssetData = AssetRegistryModule.Get().GetAssetByObjectPath(FSoftObjectPath(FunctionScriptAssetObjectPath.ToString()));
 		if (ScriptAssetData.IsValid())
 		{
 			SpawningFunctionScript = Cast<UNiagaraScript>(ScriptAssetData.GetAsset());
diff --git a/Engine/Source/Editor/UnrealEd/Private/Commandlets/WorldPartitionDataLayerToAssetCommandlet.cpp b/Engine/Source/Editor/UnrealEd/Private/Commandlets/WorldPartitionDataLayerToAssetCommandlet.cpp
--- a/Engine/Source/Editor/UnrealEd/Private/Commandlets/WorldPartitionDataLayerToAssetCommandlet.cpp
+++ b/Engine/Source/Editor/UnrealEd/Private/Commandlets/WorldPartitionDataLayerToAssetCommandlet.cpp
@@ -22,9 +22,6 @@
 
 DEFINE_LOG_CATEGORY(LogDataLayerToAssetCommandlet);
 
-// The CommandLet uses deprecated function in AACtor & AWorldDataLayers to perform the conversion of UDEPRECATED_DataLayer -> UDataInstance/UDataLayerAsset
-PRAGMA_DISABLE_DEPRECATION_WARNINGS
-
 void UDataLayerConversionInfo::SetDataLayerToConvert(const UDeprecatedDataLayerInstance* InDataLayerToConvert)
 {
 	DataLayerToConvert = InDataLayerToConvert;
@@ -345,7 +342,7 @@ bool UDataLayerToAssetCommandlet::BuildConversionInfos(TStrongObjectPtr<UDataLay
 
 	TArray<FAssetData> ExistingDataLayerAssets;
 	IAssetRegistry& AssetRegistry = FAssetRegistryModule::GetRegistry();
-	AssetRegistry.GetAssetsByClass(UDataLayerAsset::StaticClass()->GetFName(), ExistingDataLayerAssets);
+	AssetRegistry.GetAssetsByClass(FTopLevelAssetPath(UDataLayerAsset::StaticClass()), ExistingDataLayerAssets);
 	for (FAssetData& AssetData : ExistingDataLayerAssets)
 	{
 		if (IsAssetInConversionFolder(AssetData.GetSoftObjectPath()))
@@ -532,6 +529,7 @@ bool UDataLayerToAssetCommandlet::RemapActorDataLayersToAssets(TStrongObjectPtr<
 
 uint32 UDataLayerToAssetCommandlet::RemapActorDataLayers(TStrongObjectPtr<UDataLayerToAssetCommandletContext>& CommandletContext, AActor* Actor)
 {
+PRAGMA_DISABLE_DEPRECATION_WARNINGS
 	uint32 ErrorCount = 0;
 
 	const TArray<FActorDataLayer>& ActorDataLayers = Actor->GetActorDataLayers();
@@ -569,6 +567,7 @@ uint32 UDataLayerToAssetCommandlet::RemapActorDataLayers(TStrongObjectPtr<UDataL
 	}
 
 	return ErrorCount;
+PRAGMA_ENABLE_DEPRECATION_WARNINGS
 }
 
 // Multiple run of the commandlet can lead to Actors referencing Data Layer Assets in different folders. Always remap to the newly created asset.
@@ -755,6 +754,4 @@ bool UDataLayerToAssetCommandlet::CommitConversion(TStrongObjectPtr<UDataLayerTo
 bool UDataLayerToAssetCommandlet::IsAssetInConversionFolder(const FSoftObjectPath& DataLayerAsset)
 {
 	return DataLayerAsset.GetAssetPathString().StartsWith(ConversionFolder);
-}
-
-PRAGMA_ENABLE_DEPRECATION_WARNINGS
+}
\ No newline at end of file
diff --git a/Engine/Source/Runtime/AssetRegistry/Private/AssetRegistry.cpp b/Engine/Source/Runtime/AssetRegistry/Private/AssetRegistry.cpp
--- a/Engine/Source/Runtime/AssetRegistry/Private/AssetRegistry.cpp
+++ b/Engine/Source/Runtime/AssetRegistry/Private/AssetRegistry.cpp
@@ -2073,12 +2073,6 @@ static FTopLevelAssetPath TryConvertShortTypeNameToPathName(FName ClassName)
 	return ClassPathName;
 }
 
-bool UAssetRegistryImpl::GetAssetsByClass(FName ClassName, TArray<FAssetData>& OutAssetData, bool bSearchSubClasses) const
-{
-	FTopLevelAssetPath ClassPathName = TryConvertShortTypeNameToPathName(ClassName);
-	return GetAssetsByClass(ClassPathName, OutAssetData, bSearchSubClasses);
-}
-
 bool UAssetRegistryImpl::GetAssetsByClass(FTopLevelAssetPath ClassPathName, TArray<FAssetData>& OutAssetData, bool bSearchSubClasses) const
 {
 	FARFilter Filter;
@@ -2530,13 +2524,6 @@ FAssetData UAssetRegistryImpl::GetAssetByObjectPath(const FSoftObjectPath& Objec
 	}
 }
 
-FAssetData UAssetRegistryImpl::GetAssetByObjectPath(const FName ObjectPath, bool bIncludeOnlyOnDiskAssets) const
-{
-PRAGMA_DISABLE_DEPRECATION_WARNINGS;
-	return GetAssetByObjectPath(FSoftObjectPath(ObjectPath.ToString()), bIncludeOnlyOnDiskAssets);
-PRAGMA_ENABLE_DEPRECATION_WARNINGS;
-}
-
 UE::AssetRegistry::EExists UAssetRegistryImpl::TryGetAssetByObjectPath(const FSoftObjectPath& ObjectPath, FAssetData& OutAssetData) const
 {
 	FReadScopeLock InterfaceScopeLock(InterfaceLock);
@@ -2988,25 +2975,36 @@ bool UAssetRegistryImpl::DoesPackageExistOnDisk(FName PackageName, FString* OutC
 	}
 }
 
-FSoftObjectPath UAssetRegistryImpl::GetRedirectedObjectPath(const FSoftObjectPath& ObjectPath) const
+FSoftObjectPath UAssetRegistryImpl::GetRedirectedObjectPath(const FSoftObjectPath& ObjectPath)
 {
-	FReadScopeLock InterfaceScopeLock(InterfaceLock);
-	return GuardedData.GetRedirectedObjectPath(ObjectPath);
-}
+	if (IsSearchAllAssets())
+	{
+		FReadScopeLock InterfaceScopeLock(InterfaceLock);
+		return GuardedData.GetRedirectedObjectPath(ObjectPath, nullptr, nullptr, /*bNeedsScanning*/ false);
+	}
 
-FName UAssetRegistryImpl::GetRedirectedObjectPath(const FName ObjectPath) const
-{
-	FReadScopeLock InterfaceScopeLock(InterfaceLock);
-PRAGMA_DISABLE_DEPRECATION_WARNINGS
-	return GuardedData.GetRedirectedObjectPath(ObjectPath);
-PRAGMA_ENABLE_DEPRECATION_WARNINGS
+	FSoftObjectPath RedirectedObjectPath;
+	UE::AssetRegistry::Impl::FEventContext EventContext;
+	UE::AssetRegistry::Impl::FClassInheritanceContext InheritanceContext;
+	UE::AssetRegistry::Impl::FClassInheritanceBuffer InheritanceBuffer;
+	{
+		LLM_SCOPE(ELLMTag::AssetRegistry);
+		FWriteScopeLock WriteScopeLock(InterfaceLock);
+		GetInheritanceContextWithRequiredLock(WriteScopeLock, InheritanceContext, InheritanceBuffer);
+		RedirectedObjectPath = GuardedData.GetRedirectedObjectPath(ObjectPath, &EventContext, &InheritanceContext, /*bNeedsScanning*/ true);
+	}	
+	Broadcast(EventContext);
+
+	return RedirectedObjectPath;
 }
 
 namespace UE::AssetRegistry
 {
 
-FSoftObjectPath FAssetRegistryImpl::GetRedirectedObjectPath(const FSoftObjectPath& ObjectPath) const
+FSoftObjectPath FAssetRegistryImpl::GetRedirectedObjectPath(const FSoftObjectPath& ObjectPath, UE::AssetRegistry::Impl::FEventContext* EventContext, UE::AssetRegistry::Impl::FClassInheritanceContext* InheritanceContext, bool bNeedsScanning)
 {
+	check(!bNeedsScanning || (EventContext && InheritanceContext));
+
 	FSoftObjectPath RedirectedPath = ObjectPath;
 
 	// For legacy behavior, for the first object pointed to, we look up the object in memory
@@ -3024,6 +3022,12 @@ FSoftObjectPath FAssetRegistryImpl::GetRedirectedObjectPath(const FSoftObjectPat
 		RedirectedPath = FSoftObjectPath(Redirector->DestinationObject);
 	}
 
+	if (bNeedsScanning)
+	{
+		UE::AssetRegistry::Impl::FScanPathContext Context(*EventContext, *InheritanceContext, {}, { RedirectedPath.ToString() });
+		ScanPathsSynchronous(Context);
+	}
+
 	FString SubPathString;
 	const FAssetData* AssetData = State.GetAssetByObjectPath(RedirectedPath);
 	if (!AssetData && RedirectedPath.IsSubobject())
@@ -3056,6 +3060,12 @@ FSoftObjectPath FAssetRegistryImpl::GetRedirectedObjectPath(const FSoftObjectPat
 			break;
 		}
 
+		if (bNeedsScanning)
+		{
+			UE::AssetRegistry::Impl::FScanPathContext Context(*EventContext, *InheritanceContext, {}, { RedirectedPath.ToString() });
+			ScanPathsSynchronous(Context);
+		}
+
 		SeenPaths.Add(RedirectedPath);
 		AssetData = State.GetAssetByObjectPath(RedirectedPath);
 	}
@@ -3075,25 +3085,6 @@ FSoftObjectPath FAssetRegistryImpl::GetRedirectedObjectPath(const FSoftObjectPat
 	}
 	return RedirectedPath;
 }
-
-UE_DEPRECATED(5.1, "Asset path FNames have been deprecated, use FSoftObjectPath instead.")
-FName FAssetRegistryImpl::GetRedirectedObjectPath(const FName ObjectPath) const
-{
-	return FName(*GetRedirectedObjectPath(FSoftObjectPath(ObjectPath.ToString())).ToString());
-}
-
-}
-
-bool UAssetRegistryImpl::GetAncestorClassNames(FName ClassName, TArray<FName>& OutAncestorClassNames) const
-{
-	FTopLevelAssetPath ClassPathName = TryConvertShortTypeNameToPathName(ClassName);
-	TArray<FTopLevelAssetPath> OutAncestorClassPathNames;
-	bool bResult = GetAncestorClassNames(ClassPathName, OutAncestorClassPathNames);
-	for (FTopLevelAssetPath AncestorPathName : OutAncestorClassPathNames)
-	{
-		OutAncestorClassNames.Add(AncestorPathName.GetAssetName());
-	}
-	return bResult;
 }
 
 bool UAssetRegistryImpl::GetAncestorClassNames(FTopLevelAssetPath ClassName, TArray<FTopLevelAssetPath>& OutAncestorClassNames) const
@@ -3157,28 +3148,6 @@ bool FAssetRegistryImpl::GetAncestorClassNames(Impl::FClassInheritanceContext& I
 
 	return bFoundClass;
 }
-
-}
-
-void UAssetRegistryImpl::GetDerivedClassNames(const TArray<FName>& ClassNames, const TSet<FName>& ExcludedClassNames,
-	TSet<FName>& OutDerivedClassNames) const
-{
-	TArray<FTopLevelAssetPath> ClassPaths;
-	for (FName ClassName : ClassNames)
-	{
-		ClassPaths.Add(TryConvertShortTypeNameToPathName(ClassName));
-	}
-	TSet<FTopLevelAssetPath> ExcludedClassPathNames;
-	for (FName ExcludedClassName : ExcludedClassNames)
-	{
-		ExcludedClassPathNames.Add(TryConvertShortTypeNameToPathName(ExcludedClassName));
-	}
-	TSet<FTopLevelAssetPath> OutDerivedClassPathNames;
-	GetDerivedClassNames(ClassPaths, ExcludedClassPathNames, OutDerivedClassPathNames);
-	for (FTopLevelAssetPath DerivedClassPathName : OutDerivedClassPathNames)
-	{
-		OutDerivedClassNames.Add(DerivedClassPathName.GetAssetName());
-	}
 }
 
 void UAssetRegistryImpl::GetDerivedClassNames(const TArray<FTopLevelAssetPath>& ClassNames, const TSet<FTopLevelAssetPath>& ExcludedClassNames,
@@ -6703,7 +6672,7 @@ void FAssetRegistryImpl::UpdateRedirectCollector()
 	for (const FAssetData* AssetData : RedirectorAssets)
 	{
 		FSoftObjectPath Source = AssetData->GetSoftObjectPath();
-		FSoftObjectPath Destination = GetRedirectedObjectPath(Source);
+		FSoftObjectPath Destination = GetRedirectedObjectPath(Source, nullptr, nullptr, /*bNeedsScanning*/ false);
 
 		if (Destination != Source)
 		{
diff --git a/Engine/Source/Runtime/AssetRegistry/Private/AssetRegistry.h b/Engine/Source/Runtime/AssetRegistry/Private/AssetRegistry.h
--- a/Engine/Source/Runtime/AssetRegistry/Private/AssetRegistry.h
+++ b/Engine/Source/Runtime/AssetRegistry/Private/AssetRegistry.h
@@ -34,16 +34,12 @@ public:
 	virtual bool GetAssetsByPackageName(FName PackageName, TArray<FAssetData>& OutAssetData, bool bIncludeOnlyOnDiskAssets = false, bool bSkipARFilteredAssets=true) const override;
 	virtual bool GetAssetsByPath(FName PackagePath, TArray<FAssetData>& OutAssetData, bool bRecursive = false, bool bIncludeOnlyOnDiskAssets = false) const override;
 	virtual bool GetAssetsByPaths(TArray<FName> PackagePath, TArray<FAssetData>& OutAssetData, bool bRecursive = false, bool bIncludeOnlyOnDiskAssets = false) const override;
-	UE_DEPRECATED(5.1, "Class names are now represented by path names. Please use a version of this function that uses FTopLevelAssetPath.")
-	virtual bool GetAssetsByClass(FName ClassName, TArray<FAssetData>& OutAssetData, bool bSearchSubClasses = false) const override;
 	virtual bool GetAssetsByClass(FTopLevelAssetPath ClassPathName, TArray<FAssetData>& OutAssetData, bool bSearchSubClasses = false) const override;
 	virtual bool GetAssetsByTags(const TArray<FName>& AssetTags, TArray<FAssetData>& OutAssetData) const override;
 	virtual bool GetAssetsByTagValues(const TMultiMap<FName, FString>& AssetTagsAndValues, TArray<FAssetData>& OutAssetData) const override;
 	virtual bool GetAssets(const FARFilter& Filter, TArray<FAssetData>& OutAssetData, bool bSkipARFilteredAssets = true) const override;
 	virtual bool EnumerateAssets(const FARFilter& Filter, TFunctionRef<bool(const FAssetData&)> Callback, bool bSkipARFilteredAssets=true) const override;
 	virtual bool EnumerateAssets(const FARCompiledFilter& Filter, TFunctionRef<bool(const FAssetData&)> Callback, bool bSkipARFilteredAssets = true) const override;
-	UE_DEPRECATED(5.1, "Asset path FNames have been deprecated, use FSoftObjectPath instead.")
-	virtual FAssetData GetAssetByObjectPath( const FName ObjectPath, bool bIncludeOnlyOnDiskAssets = false ) const override;
 	virtual FAssetData GetAssetByObjectPath(const FSoftObjectPath& ObjectPath, bool bIncludeOnlyOnDiskAssets = false, bool bSkipARFilteredAssets = true) const override;
 	virtual UE::AssetRegistry::EExists TryGetAssetByObjectPath(const FSoftObjectPath& ObjectPath, FAssetData& OutAssetData) const override;
 	virtual UE::AssetRegistry::EExists TryGetAssetPackageData(const FName PackageName, FAssetPackageData& OutAssetPackageData) const override;
@@ -61,14 +57,8 @@ public:
 	virtual TOptional<FAssetPackageData> GetAssetPackageDataCopy(FName PackageName) const override;
 	virtual void EnumerateAllPackages(TFunctionRef<void(FName PackageName, const FAssetPackageData& PackageData)> Callback) const override;
 	virtual bool DoesPackageExistOnDisk(FName PackageName, FString* OutCorrectCasePackageName = nullptr, FString* OutExtension = nullptr) const override;
-	virtual FSoftObjectPath GetRedirectedObjectPath(const FSoftObjectPath& ObjectPath) const override;
-	UE_DEPRECATED(5.1, "Asset path FNames have been deprecated, use FSoftObjectPath instead.")
-	virtual FName GetRedirectedObjectPath(const FName ObjectPath) const override;
-	UE_DEPRECATED(5.1, "Class names are now represented by path names. Please use a version of this function that uses FTopLevelAssetPath.")
-	virtual bool GetAncestorClassNames(FName ClassName, TArray<FName>& OutAncestorClassNames) const override;
+	virtual FSoftObjectPath GetRedirectedObjectPath(const FSoftObjectPath& ObjectPath) override;
 	virtual bool GetAncestorClassNames(FTopLevelAssetPath ClassName, TArray<FTopLevelAssetPath>& OutAncestorClassNames) const override;
-	UE_DEPRECATED(5.1, "Class names are now represented by path names. Please use a version of this function that uses FTopLevelAssetPath.")
-	virtual void GetDerivedClassNames(const TArray<FName>& ClassNames, const TSet<FName>& ExcludedClassNames, TSet<FName>& OutDerivedClassNames) const override;
 	virtual void GetDerivedClassNames(const TArray<FTopLevelAssetPath>& ClassNames, const TSet<FTopLevelAssetPath>& ExcludedClassNames, TSet<FTopLevelAssetPath>& OutDerivedClassNames) const override;	
 	virtual void GetAllCachedPaths(TArray<FString>& OutPathList) const override;
 	virtual void EnumerateAllCachedPaths(TFunctionRef<bool(FString)> Callback) const override;
@@ -211,10 +201,6 @@ public:
 
 protected:
 	virtual void SetManageReferences(const TMultiMap<FAssetIdentifier, FAssetIdentifier>& ManagerMap, bool bClearExisting, UE::AssetRegistry::EDependencyCategory RecurseType, TSet<FDependsNode*>& ExistingManagedNodes, ShouldSetManagerPredicate ShouldSetManager = nullptr) override;
-	virtual bool SetPrimaryAssetIdForObjectPath(const FName ObjectPath, FPrimaryAssetId PrimaryAssetId) override
-	{
-		return SetPrimaryAssetIdForObjectPath(FSoftObjectPath(ObjectPath), PrimaryAssetId);
-	}
 	virtual bool SetPrimaryAssetIdForObjectPath(const FSoftObjectPath& ObjectPath, FPrimaryAssetId PrimaryAssetId) override;
 
 private:
diff --git a/Engine/Source/Runtime/AssetRegistry/Private/AssetRegistryImpl.h b/Engine/Source/Runtime/AssetRegistry/Private/AssetRegistryImpl.h
--- a/Engine/Source/Runtime/AssetRegistry/Private/AssetRegistryImpl.h
+++ b/Engine/Source/Runtime/AssetRegistry/Private/AssetRegistryImpl.h
@@ -134,9 +134,8 @@ public:
 
 
 	bool HasAssets(const FName PackagePath, const bool bRecursive) const;
-	FSoftObjectPath GetRedirectedObjectPath(const FSoftObjectPath& ObjectPath) const;
-	UE_DEPRECATED(5.1, "Asset path FNames have been deprecated, use FSoftObjectPath instead.")
-	FName GetRedirectedObjectPath(const FName ObjectPath) const;
+	FSoftObjectPath GetRedirectedObjectPath(const FSoftObjectPath& ObjectPath, UE::AssetRegistry::Impl::FEventContext* EventContext, 
+		UE::AssetRegistry::Impl::FClassInheritanceContext* InheritanceContext, bool bNeedsScanning);
 	bool GetAncestorClassNames(Impl::FClassInheritanceContext& InheritanceContext, FTopLevelAssetPath ClassName,
 		TArray<FTopLevelAssetPath>& OutAncestorClassNames) const;
 	void CompileFilter(Impl::FClassInheritanceContext& InheritanceContext, const FARFilter& InFilter,
diff --git a/Engine/Source/Runtime/AssetRegistry/Public/AssetRegistry/IAssetRegistry.h b/Engine/Source/Runtime/AssetRegistry/Public/AssetRegistry/IAssetRegistry.h
--- a/Engine/Source/Runtime/AssetRegistry/Public/AssetRegistry/IAssetRegistry.h
+++ b/Engine/Source/Runtime/AssetRegistry/Public/AssetRegistry/IAssetRegistry.h
@@ -247,9 +247,6 @@ public:
 	UFUNCTION(BlueprintCallable, BlueprintPure=false, Category = "AssetRegistry")
 	virtual bool GetAssetsByClass(FTopLevelAssetPath ClassPathName, TArray<FAssetData>& OutAssetData, bool bSearchSubClasses = false) const = 0;
 
-	UE_DEPRECATED(5.1, "Class names are now represented by path names. Please use a version of this function that uses FTopLevelAssetPath.")
-	virtual bool GetAssetsByClass(FName ClassPathName, TArray<FAssetData>& OutAssetData, bool bSearchSubClasses = false) const = 0;
-
 	/**
 	 * Gets asset data for all assets with the supplied tags, regardless of their value
 	 *
@@ -292,18 +289,6 @@ public:
 	virtual bool EnumerateAssets(const FARCompiledFilter& Filter, TFunctionRef<bool(const FAssetData&)> Callback,
 		bool bSkipARFilteredAssets = true) const = 0;
 
-	/**
-	 * Gets the asset data for the specified object path
-	 *
-	 * @param ObjectPath the path of the object to be looked up
-	 * @param bIncludeOnlyOnDiskAssets If true, use only DiskGatheredData, do not calculate from UObjects.
-	 *        @see IAssetRegistry class header for bIncludeOnlyOnDiskAssets.
-	 * @return the assets data;Will be invalid if object could not be found
-	 */
-	UE_DEPRECATED(5.1, "Asset path FNames have been deprecated, use Soft Object Path instead.")
-	UFUNCTION(BlueprintCallable, BlueprintPure=false, Category="AssetRegistry")
-	virtual FAssetData GetAssetByObjectPath( const FName ObjectPath, bool bIncludeOnlyOnDiskAssets = false ) const = 0;
-
 	/**
 	 * Gets the asset data for the specified object path
 	 *
@@ -480,25 +465,16 @@ public:
 	virtual bool DoesPackageExistOnDisk(FName PackageName, FString* OutCorrectCasePackageName = nullptr, FString* OutExtension = nullptr) const = 0;
 
 	/** Uses the asset registry to look for ObjectRedirectors. This will follow the chain of redirectors. It will return the original path if no redirectors are found */
-	virtual FSoftObjectPath GetRedirectedObjectPath(const FSoftObjectPath& ObjectPath) const = 0;
-
-	UE_DEPRECATED(5.1, "Asset path FNames have been deprecated, use FSoftObjectPath instead.")
-	virtual FName GetRedirectedObjectPath(const FName ObjectPath) const = 0;
+	virtual FSoftObjectPath GetRedirectedObjectPath(const FSoftObjectPath& ObjectPath) = 0;
 
 	/** Returns true if the specified ClassName's ancestors could be found. If so, OutAncestorClassNames is a list of all its ancestors. This can be slow if temporary caching mode is not on */
 	UFUNCTION(BlueprintCallable, BlueprintPure = false, Category = "AssetRegistry", meta = (DisplayName = "GetAncestorClassNames", ScriptName = "GetAncestorClassNames"))
 	virtual bool GetAncestorClassNames(FTopLevelAssetPath ClassPathName, TArray<FTopLevelAssetPath>& OutAncestorClassNames) const = 0;
 
-	UE_DEPRECATED(5.1, "Class names are now represented by path names. Please use a version of this function that uses FTopLevelAssetPath.")
-	virtual bool GetAncestorClassNames(FName ClassName, TArray<FName>& OutAncestorClassNames) const = 0;
-
 	/** Returns the names of all classes derived by the supplied class names, excluding any classes matching the excluded class names. This can be slow if temporary caching mode is not on */
 	UFUNCTION(BlueprintCallable, BlueprintPure = false, Category = "AssetRegistry", meta = (DisplayName = "GetDerivedClassNames", ScriptName = "GetDerivedClassNames"))
 	virtual void GetDerivedClassNames(const TArray<FTopLevelAssetPath>& ClassNames, const TSet<FTopLevelAssetPath>& ExcludedClassNames, TSet<FTopLevelAssetPath>& OutDerivedClassNames) const = 0;
 	
-	UE_DEPRECATED(5.1, "Class names are now represented by path names. Please use a version of this function that uses FTopLevelAssetPath.")
-	virtual void GetDerivedClassNames(const TArray<FName>& ClassNames, const TSet<FName>& ExcludedClassNames, TSet<FName>& OutDerivedClassNames) const = 0;
-
 	/** Gets a list of all paths that are currently cached */
 	UFUNCTION(BlueprintCallable, BlueprintPure=false, Category = "AssetRegistry")
 	virtual void GetAllCachedPaths(TArray<FString>& OutPathList) const = 0;
@@ -916,8 +892,6 @@ protected:
 	virtual void SetManageReferences(const TMultiMap<FAssetIdentifier, FAssetIdentifier>& ManagerMap, bool bClearExisting, UE::AssetRegistry::EDependencyCategory RecurseType, TSet<FDependsNode*>& ExistingManagedNodes, ShouldSetManagerPredicate ShouldSetManager = nullptr) = 0;
 
 	/** Sets the PrimaryAssetId for a specific asset. This should only be called by the AssetManager, and is needed when the AssetManager is more up to date than the on disk Registry */
-	UE_DEPRECATED(5.1, "Asset path FNames have been deprecated, use FSoftObjectPath instead.")
-	virtual bool SetPrimaryAssetIdForObjectPath(const FName ObjectPath, FPrimaryAssetId PrimaryAssetId) = 0;
 	virtual bool SetPrimaryAssetIdForObjectPath(const FSoftObjectPath& ObjectPath, FPrimaryAssetId PrimaryAssetId) = 0;
 };
 
