AssetDataGatherer: Turn off writing in -game because it is expensive and the writes are not used.
#rnx

diff --git a/Engine/Source/Runtime/AssetRegistry/Private/AssetDataGatherer.cpp b/Engine/Source/Runtime/AssetRegistry/Private/AssetDataGatherer.cpp
--- a/Engine/Source/Runtime/AssetRegistry/Private/AssetDataGatherer.cpp
+++ b/Engine/Source/Runtime/AssetRegistry/Private/AssetDataGatherer.cpp
@@ -182,7 +182,11 @@ void FPreloadSettings::Initialize()
 	bool bNoAssetRegistryCacheRead = FParse::Param(FCommandLine::Get(), TEXT("NoAssetRegistryCacheRead"));
 	uint32 MultiprocessId = UE::GetMultiprocessId();
 	bool bMultiprocess = MultiprocessId > 0 || FParse::Param(FCommandLine::Get(), TEXT("multiprocess"));
-	bool bNoAssetRegistryCacheWrite = FParse::Param(FCommandLine::Get(), TEXT("NoAssetRegistryCacheWrite")) || bMultiprocess;
+	bool bNoAssetRegistryCacheWrite = FParse::Param(FCommandLine::Get(), TEXT("NoAssetRegistryCacheWrite"))
+		// Don't write in multiprocess because we will collide writing the cache files
+		|| bMultiprocess
+		// Cooked game/server and editor -game do not need to write the cache; they get it from editor or cooking
+		|| !GIsEditor;
 	bGatherCacheReadEnabled = !bNoAssetRegistryCache && !bNoAssetRegistryCacheRead;
 	bGatherCacheWriteEnabled = !bNoAssetRegistryCache && !bNoAssetRegistryCacheWrite;
 	bool bPlatformSupportsDiscoveryCacheInvalidation = FPlatformFileManager::Get().GetPlatformFile().FileJournalIsAvailable();
