AssetRegistry PackageReader: Improve unloadable package warning messages to include description of the consequences. Also remove calls to UE_ASSET_LOG since PackageReader is called too early before some of the systems used by UE_ASSET_LOG have been initialized such as PackageName->IPackageResourceManager.
#rnx

diff --git a/Engine/Source/Runtime/AssetRegistry/Private/PackageReader.cpp b/Engine/Source/Runtime/AssetRegistry/Private/PackageReader.cpp
--- a/Engine/Source/Runtime/AssetRegistry/Private/PackageReader.cpp
+++ b/Engine/Source/Runtime/AssetRegistry/Private/PackageReader.cpp
@@ -108,10 +108,12 @@ bool FPackageReader::OpenPackageFile(EOpenPackageResult& OutErrorCode)
 	// Validate the summary.
 
 	// Make sure this is indeed a package
-	if( PackageFileSummary.Tag != PACKAGE_FILE_TAG || IsError())
+	if (PackageFileSummary.Tag != PACKAGE_FILE_TAG || IsError())
 	{
 		// Unrecognized or malformed package file
-		UE_LOG(LogAssetRegistry, Error, TEXT("Package %s has malformed tag"), *PackageFilename);
+		UE_LOG(LogAssetRegistry, Error,
+			TEXT("Package is unloadable: %s. Reason: Invalid value for PACKAGE_FILE_TAG at start of file."),
+			*PackageFilename);
 		OutErrorCode = EOpenPackageResult::MalformedTag;
 		return false;
 	}
@@ -124,7 +126,8 @@ bool FPackageReader::OpenPackageFile(EOpenPackageResult& OutErrorCode)
 	{
 		// Log a warning rather than an error. Linkerload gracefully handles this case.
 		UE_LOG(LogAssetRegistry, Warning,
-			TEXT("Package %s is unversioned which cannot be opened by the current process"), *PackageFilename);
+			TEXT("Package is unloadable: %s. Reason: Package was saved unversioned and the current process does not support loading unversioned packages."),
+			*PackageFilename);
 		OutErrorCode = EOpenPackageResult::Unversioned;
 		return false;
 	}
@@ -134,7 +137,8 @@ bool FPackageReader::OpenPackageFile(EOpenPackageResult& OutErrorCode)
 		IsEnforcePackageCompatibleVersionCheck())
 	{
 		// Log a warning rather than an error. Linkerload gracefully handles this case.
-		UE_LOG(LogAssetRegistry, Warning, TEXT("Package %s is too old. Min Version: %i  Package Version: %i"),
+		UE_LOG(LogAssetRegistry, Warning,
+			TEXT("Package is unloadable: %s. Reason: Version is too old. Min Version: %i, Package Version: %i."),
 			*PackageFilename, (int32)VER_UE4_OLDEST_LOADABLE_PACKAGE,
 			PackageFileSummary.GetFileVersionUE().FileVersionUE4);
 
@@ -147,7 +151,8 @@ bool FPackageReader::OpenPackageFile(EOpenPackageResult& OutErrorCode)
 		IsEnforcePackageCompatibleVersionCheck())
 	{
 		// Log a warning rather than an error. Linkerload gracefully handles this case.
-		UE_LOG(LogAssetRegistry, Warning, TEXT("Package %s is too new. Engine Version: %i  Package Version: %i"),
+		UE_LOG(LogAssetRegistry, Warning,
+			TEXT("Package is unloadable: %s. Reason: Version is too new. Engine Version: %i, Package Version: %i."),
 			*PackageFilename, GPackageFileUEVersion.ToValue(), PackageFileSummary.GetFileVersionUE().ToValue());
 
 		OutErrorCode = EOpenPackageResult::VersionTooNew;
@@ -158,7 +163,8 @@ bool FPackageReader::OpenPackageFile(EOpenPackageResult& OutErrorCode)
 		IsEnforcePackageCompatibleVersionCheck())
 	{
 		// Log a warning rather than an error. Linkerload gracefully handles this case.
-		UE_LOG(LogAssetRegistry, Warning, TEXT("Package %s is too new. Licensee Version: %i Package Licensee Version: %i"),
+		UE_LOG(LogAssetRegistry, Warning,
+			TEXT("Package is unloadable: %s. Reason: LicenseeVersion is too new. Licensee Version: %i, Package Licensee Version: %i."),
 			*PackageFilename, GPackageFileLicenseeUEVersion, PackageFileSummary.GetFileVersionLicenseeUE());
 
 		OutErrorCode = EOpenPackageResult::VersionTooNew;
@@ -188,7 +194,8 @@ bool FPackageReader::OpenPackageFile(EOpenPackageResult& OutErrorCode)
 		{
 			if (IsEnforcePackageCompatibleVersionCheck())
 			{
-				UE_LOG(LogAssetRegistry, Error, TEXT("Package %s has newer custom version of %s"),
+				UE_LOG(LogAssetRegistry, Error,
+					TEXT("Package is unloadable: %s. Reason: Custom version is too new; the package has newer custom version of %s."),
 					*PackageFilename, *Diff.Version->GetFriendlyName().ToString());
 				OutErrorCode = EOpenPackageResult::VersionTooNew;
 			}
@@ -256,7 +263,7 @@ bool FPackageReader::StartSerializeSection(int64 Offset)
 		CorruptPackageWarningArguments.Add(TEXT("FileName"), FText::FromString(PackageFileName)); \
 		UE_LOG(LogAssetRegistry, Warning, TEXT("%s"), \
 			*FText::Format(NSLOCTEXT("AssetRegistry", MessageKey, \
-			"Cannot read AssetRegistry Data in {FileName}, skipping it. Error: " MessageKey "."), \
+			"Package is unloadable: {FileName}. Reason: " MessageKey "."), \
 			CorruptPackageWarningArguments).ToString()); \
 	} while (false)
 
@@ -407,8 +414,8 @@ bool FPackageReader::ReadAssetRegistryData(TArray<FAssetData*>& AssetDataList, b
 	using namespace UE::AssetRegistry;
 
 	EReadPackageDataMainErrorCode ErrorCode;
-	if (!ReadPackageDataMain(*this, PackageName, PackageFileSummary, AssetRegistryDependencyDataOffset, AssetDataList, ErrorCode,
-		&ImportMap, &ExportMap))
+	if (!ReadPackageDataMain(*this, PackageName, PackageFileSummary, AssetRegistryDependencyDataOffset, AssetDataList,
+		ErrorCode, &ImportMap, &ExportMap))
 	{
 		switch (ErrorCode)
 		{
@@ -1328,7 +1335,9 @@ FArchive& FPackageReader::operator<<( FName& Name )
 
 	if( !NameMap.IsValidIndex(NameIndex) )
 	{
-		UE_LOG(LogAssetRegistry, Warning, TEXT("Bad name index %i/%i when reading package %s"), NameIndex, NameMap.Num(), *PackageFilename );
+		UE_LOG(LogAssetRegistry, Warning,
+			TEXT("Package is unloadable: %s. Reason: Bad name index %i/%i when reading package."),
+			*PackageFilename, NameIndex, NameMap.Num() );
 		SetError();
 		return *this;
 	}
@@ -1709,9 +1718,13 @@ namespace UE::AssetRegistry
 			if (!bFullObjectPath)
 			{
 				// if we do not have a full object path, ensure that we have a top level object for the package and not a sub object
-				if (!ensureMsgf(!ObjectPackageData.ObjectPath.Contains(TEXT("."), ESearchCase::CaseSensitive), TEXT("Cannot make FAssetData for sub object %s in package %s!"), *ObjectPackageData.ObjectPath, *PackageName))
+				if (!ensureMsgf(!ObjectPackageData.ObjectPath.Contains(TEXT("."), ESearchCase::CaseSensitive),
+					TEXT("Package is loadable but its AssetRegistry data is corrupt: %s. Reason: Cannot make FAssetData for sub object %s."),
+					*PackageName , *ObjectPackageData.ObjectPath))
 				{
-					UE_ASSET_LOG(LogAssetRegistry, Warning, *PackageName, TEXT("Cannot make FAssetData for sub object %s!"), *ObjectPackageData.ObjectPath);
+					UE_LOG(LogAssetRegistry, Warning,
+						TEXT("Package is loadable but its AssetRegistry data is corrupt: %s. Reason: Cannot make FAssetData for sub object %s."),
+						*PackageName, *ObjectPackageData.ObjectPath);
 					continue;
 				}
 				ObjectPackageData.ObjectPath = PackageName + TEXT(".") + ObjectPackageData.ObjectPath;
@@ -1719,7 +1732,9 @@ namespace UE::AssetRegistry
 			// Previously export couldn't have its outer as an import
 			else if (PackageFileSummary.GetFileVersionUE() < VER_UE4_NON_OUTER_PACKAGE_IMPORT)
 			{
-				UE_ASSET_LOG(LogAssetRegistry, Warning, *PackageName, TEXT("Package has invalid export %s, resave source package!"), *ObjectPackageData.ObjectPath);
+				UE_LOG(LogAssetRegistry, Warning,
+					TEXT("Package is loadable but has invalid data; resave the package! Package: %s. Reason: Export %s is invalid."),
+					*PackageName, *ObjectPackageData.ObjectPath);
 				continue;
 			}
 
