Wait for completion of an async asset registry scan in UAssetRegistryImpl::SearchAllAssets to fix occasional problem where SearchAllAssets returns before IsLoadingAssets will return false.

diff --git a/Engine/Source/Runtime/AssetRegistry/Private/AssetRegistry.cpp b/Engine/Source/Runtime/AssetRegistry/Private/AssetRegistry.cpp
--- a/Engine/Source/Runtime/AssetRegistry/Private/AssetRegistry.cpp
+++ b/Engine/Source/Runtime/AssetRegistry/Private/AssetRegistry.cpp
@@ -1749,7 +1749,12 @@ void UAssetRegistryImpl::SearchAllAssets(bool bSynchronousSearch)
 	TRACE_CPUPROFILER_EVENT_SCOPE_STR("UAssetRegistryImpl::SearchAllAssets");
 	using namespace UE::AssetRegistry::Impl;
 
-	double StartTime = FPlatformTime::Seconds();
+	if (bSynchronousSearch)
+	{
+		// Ensure any ongoing async scan finishes fully first
+		WaitForCompletion();
+	}
+
 	FEventContext EventContext;
 	{
 		LLM_SCOPE(ELLMTag::AssetRegistry);
