AssetRegistry: Gatherer: Fix bug in the register of parent/child relationship of overlapping mount points in the AssetDataGatherer.
#rnx

diff --git a/Engine/Source/Runtime/AssetRegistry/Private/AssetDataGatherer.cpp b/Engine/Source/Runtime/AssetRegistry/Private/AssetDataGatherer.cpp
--- a/Engine/Source/Runtime/AssetRegistry/Private/AssetDataGatherer.cpp
+++ b/Engine/Source/Runtime/AssetRegistry/Private/AssetDataGatherer.cpp
@@ -2841,6 +2841,7 @@ void FAssetDataDiscovery::AddMountPointInternal(const FString& LocalAbsPath, FSt
 		FStringView RelPath;
 		verify(FPathViews::TryMakeChildPathRelativeTo(LocalAbsPath, ParentMount->GetLocalAbsPath(), RelPath));
 		ParentMount->AddChildMount(&Mount);
+		Mount.SetParentMount(ParentMount);
 		for (FMountDir* ChildMount : ChildMounts)
 		{
 			ParentMount->RemoveChildMount(ChildMount);
@@ -2849,7 +2850,7 @@ void FAssetDataDiscovery::AddMountPointInternal(const FString& LocalAbsPath, FSt
 	for (FMountDir* ChildMount : ChildMounts)
 	{
 		Mount.AddChildMount(ChildMount);
-		ChildMount->SetParentMount(ParentMount);
+		ChildMount->SetParentMount(&Mount);
 	}
 }
 
