Making the asset gatherer tick run on a background thread. In order to do this, it was also necessary to refactor PostLoadAssetRegistryTags, deprecating the existing function and implementing a new ThreadedPostLoadAssetRegistryTagsOverride. The multithreading functionality is disabled in this checkin. -dpcvars=AssetRegistry.TickGatherOnGTOnly=0  will enable it.

diff --git a/Engine/Plugins/Editor/ContentBrowser/ContentBrowserAssetDataSource/Source/ContentBrowserAssetDataSource/Private/ContentBrowserAssetDataSource.cpp b/Engine/Plugins/Editor/ContentBrowser/ContentBrowserAssetDataSource/Source/ContentBrowserAssetDataSource/Private/ContentBrowserAssetDataSource.cpp
--- a/Engine/Plugins/Editor/ContentBrowser/ContentBrowserAssetDataSource/Source/ContentBrowserAssetDataSource/Private/ContentBrowserAssetDataSource.cpp
+++ b/Engine/Plugins/Editor/ContentBrowser/ContentBrowserAssetDataSource/Source/ContentBrowserAssetDataSource/Private/ContentBrowserAssetDataSource.cpp
@@ -1764,6 +1764,8 @@ void UContentBrowserAssetDataSource::EnumerateFoldersMatchingFilter(UContentBrow
 
 		if (bStartingPathIsFullyVirtual)
 		{
+			IAssetRegistry::FPauseBackgroundProcessingScope PauseBackgroundProcessingScope;
+
 			// Virtual paths not supported by PathPassesCompiledDataFilter, enumerate internal paths in hierarchy and propagate results to virtual parents
 			TSet<FName> VirtualPathsPassedFilter;
 			VirtualPathsPassedFilter.Reserve(DataSource->GetRootPathVirtualTree().NumPaths());
diff --git a/Engine/Plugins/MovieScene/TemplateSequence/Source/TemplateSequence/Private/TemplateSequence.cpp b/Engine/Plugins/MovieScene/TemplateSequence/Source/TemplateSequence/Private/TemplateSequence.cpp
--- a/Engine/Plugins/MovieScene/TemplateSequence/Source/TemplateSequence/Private/TemplateSequence.cpp
+++ b/Engine/Plugins/MovieScene/TemplateSequence/Source/TemplateSequence/Private/TemplateSequence.cpp
@@ -268,18 +268,18 @@ void UTemplateSequence::GetAssetRegistryTags(FAssetRegistryTagsContext Context)
 
 #if WITH_EDITOR
 
-void UTemplateSequence::PostLoadAssetRegistryTags(const FAssetData& InAssetData, TArray<FAssetRegistryTag>& OutTagsAndValuesToUpdate) const
+void UTemplateSequence::ThreadedPostLoadAssetRegistryTagsOverride(FPostLoadAssetRegistryTagsContext& Context) const
 {
-	Super::PostLoadAssetRegistryTags(InAssetData, OutTagsAndValuesToUpdate);
+	Super::ThreadedPostLoadAssetRegistryTagsOverride(Context);
 
 	static const FName BoundActorClassTagName(TEXT("BoundActorClass"));
-	FString BoundActorClassTagValue = InAssetData.GetTagValueRef<FString>(BoundActorClassTagName);
+	FString BoundActorClassTagValue = Context.GetAssetData().GetTagValueRef<FString>(BoundActorClassTagName);
 	if (!BoundActorClassTagValue.IsEmpty() && FPackageName::IsShortPackageName(BoundActorClassTagValue))
 	{
-		FTopLevelAssetPath BoundActorClassPathName = UClass::TryConvertShortTypeNameToPathName<UStruct>(BoundActorClassTagValue, ELogVerbosity::Warning, TEXT("UTemplateSequence::PostLoadAssetRegistryTags"));
+		FTopLevelAssetPath BoundActorClassPathName = UClass::TryConvertShortTypeNameToPathName<UStruct>(BoundActorClassTagValue, ELogVerbosity::Warning, TEXT("UTemplateSequence::ThreadedPostLoadAssetRegistryTagsOverride"));
 		if (!BoundActorClassPathName.IsNull())
 		{
-			OutTagsAndValuesToUpdate.Add(FAssetRegistryTag(BoundActorClassTagName, BoundActorClassPathName.ToString(), FAssetRegistryTag::TT_Alphabetical));
+			Context.AddTagToUpdate(FAssetRegistryTag(BoundActorClassTagName, BoundActorClassPathName.ToString(), FAssetRegistryTag::TT_Alphabetical));
 		}
 	}
 }
diff --git a/Engine/Plugins/MovieScene/TemplateSequence/Source/TemplateSequence/Public/TemplateSequence.h b/Engine/Plugins/MovieScene/TemplateSequence/Source/TemplateSequence/Public/TemplateSequence.h
--- a/Engine/Plugins/MovieScene/TemplateSequence/Source/TemplateSequence/Public/TemplateSequence.h
+++ b/Engine/Plugins/MovieScene/TemplateSequence/Source/TemplateSequence/Public/TemplateSequence.h
@@ -47,11 +47,12 @@ public:
 	virtual void GetAssetRegistryTags(TArray<FAssetRegistryTag>& OutTags) const override;
 
 #if WITH_EDITOR
-	virtual void PostLoadAssetRegistryTags(const FAssetData& InAssetData, TArray<FAssetRegistryTag>& OutTagsAndValuesToUpdate) const;
 	virtual FText GetDisplayName() const override;
 
 	virtual ETrackSupport IsTrackSupported(TSubclassOf<class UMovieSceneTrack> InTrackClass) const override;
 	virtual void GetAssetRegistryTagMetadata(TMap<FName, FAssetRegistryTagMetadata>& OutMetadata) const override;
+protected:
+	virtual void ThreadedPostLoadAssetRegistryTagsOverride(FPostLoadAssetRegistryTagsContext& Context) const;
 #endif
 
 private:
diff --git a/Engine/Plugins/Runtime/StateTree/Source/StateTreeModule/Private/StateTree.cpp b/Engine/Plugins/Runtime/StateTree/Source/StateTreeModule/Private/StateTree.cpp
--- a/Engine/Plugins/Runtime/StateTree/Source/StateTreeModule/Private/StateTree.cpp
+++ b/Engine/Plugins/Runtime/StateTree/Source/StateTreeModule/Private/StateTree.cpp
@@ -280,18 +280,18 @@ void UStateTree::GetAssetRegistryTags(FAssetRegistryTagsContext Context) const
 	Super::GetAssetRegistryTags(Context);
 }
 
-void UStateTree::PostLoadAssetRegistryTags(const FAssetData& InAssetData, TArray<FAssetRegistryTag>& OutTagsAndValuesToUpdate) const
+void UStateTree::ThreadedPostLoadAssetRegistryTagsOverride(FPostLoadAssetRegistryTagsContext& Context) const
 {
-	Super::PostLoadAssetRegistryTags(InAssetData, OutTagsAndValuesToUpdate);
+	Super::ThreadedPostLoadAssetRegistryTagsOverride(Context);
 
 	static const FName SchemaTag(TEXT("Schema"));
-	const FString SchemaTagValue = InAssetData.GetTagValueRef<FString>(SchemaTag);
+	const FString SchemaTagValue = Context.GetAssetData().GetTagValueRef<FString>(SchemaTag);
 	if (!SchemaTagValue.IsEmpty() && FPackageName::IsShortPackageName(SchemaTagValue))
 	{
-		const FTopLevelAssetPath SchemaTagClassPathName = UClass::TryConvertShortTypeNameToPathName<UStruct>(SchemaTagValue, ELogVerbosity::Warning, TEXT("UStateTree::PostLoadAssetRegistryTags"));
+		const FTopLevelAssetPath SchemaTagClassPathName = UClass::TryConvertShortTypeNameToPathName<UStruct>(SchemaTagValue, ELogVerbosity::Warning, TEXT("UStateTree::ThreadedPostLoadAssetRegistryTagsOverride"));
 		if (!SchemaTagClassPathName.IsNull())
 		{
-			OutTagsAndValuesToUpdate.Add(FAssetRegistryTag(SchemaTag, SchemaTagClassPathName.ToString(), FAssetRegistryTag::TT_Alphabetical));
+			Context.AddTagToUpdate(FAssetRegistryTag(SchemaTag, SchemaTagClassPathName.ToString(), FAssetRegistryTag::TT_Alphabetical));
 		}
 	}
 }
diff --git a/Engine/Plugins/Runtime/StateTree/Source/StateTreeModule/Public/StateTree.h b/Engine/Plugins/Runtime/StateTree/Source/StateTreeModule/Public/StateTree.h
--- a/Engine/Plugins/Runtime/StateTree/Source/StateTreeModule/Public/StateTree.h
+++ b/Engine/Plugins/Runtime/StateTree/Source/StateTreeModule/Public/StateTree.h
@@ -208,7 +208,9 @@ protected:
 	virtual void GetAssetRegistryTags(FAssetRegistryTagsContext Context) const override;
 	UE_DEPRECATED(5.4, "Implement the version that takes FAssetRegistryTagsContext instead.")
 	virtual void GetAssetRegistryTags(TArray<FAssetRegistryTag>& OutTags) const override;
-	virtual void PostLoadAssetRegistryTags(const FAssetData& InAssetData, TArray<FAssetRegistryTag>& OutTagsAndValuesToUpdate) const override;
+protected:
+	virtual void ThreadedPostLoadAssetRegistryTagsOverride(FPostLoadAssetRegistryTagsContext& Context) const override;
+public:
 	virtual EDataValidationResult IsDataValid(class FDataValidationContext& Context) const override;
 #endif
 
diff --git a/Engine/Source/Runtime/AssetRegistry/Private/AssetDataGatherer.cpp b/Engine/Source/Runtime/AssetRegistry/Private/AssetDataGatherer.cpp
--- a/Engine/Source/Runtime/AssetRegistry/Private/AssetDataGatherer.cpp
+++ b/Engine/Source/Runtime/AssetRegistry/Private/AssetDataGatherer.cpp
@@ -2,6 +2,7 @@
 
 #include "AssetDataGatherer.h"
 #include "AssetDataGathererPrivate.h"
+#include "AssetRegistry.h"
 
 #include "Algo/AnyOf.h"
 #include "Algo/Find.h"
@@ -84,6 +85,14 @@ static FAutoConsoleVariableRef CVarIgnoreEmptyDirectories(
 	bIgnoreEmptyDirectories,
 	TEXT("If true, completely empty leaf directories are ignored by the asset registry while scanning"));
 
+//@TODO This should be set to false (enabling multithreaded processing) only after resolving UE-209921 
+// (relating to the GC lock).
+bool bTickGatherOnGTOnly = true;
+static FAutoConsoleVariableRef CVarTickGatherOnGTOnly(
+	TEXT("AssetRegistry.TickGatherOnGTOnly"),
+	bTickGatherOnGTOnly,
+	TEXT("If true, TickGatherer will only be called from the game thread."));
+
 void LexFromString(EFeatureEnabledReadWrite& OutValue, FStringView Text)
 {
 	Text.TrimStartAndEndInline();
@@ -3572,12 +3581,14 @@ FPreloader GPreloader;
 } // namespace UE::AssetDataGather::Private
 
 FAssetDataGatherer::FAssetDataGatherer(const TArray<FString>& InLongPackageNamesDenyList,
-	const TArray<FString>& InMountRelativePathsDenyList, bool bInAsyncEnabled)
-	: Thread(nullptr)
+	const TArray<FString>& InMountRelativePathsDenyList, bool bInAsyncEnabled, UE::AssetRegistry::FAssetRegistryImpl& InRegistryImpl)
+	: AssetRegistry(InRegistryImpl)
+	, Thread(nullptr)
 	, bAsyncEnabled(bInAsyncEnabled)
 	, GatherStartTime(FDateTime::Now())
 	, IsStopped(0)
-	, IsPaused(0)
+	, IsGatheringPaused(0)
+	, IsProcessingPaused(0)
 	, bInitialPluginsLoaded(false)
 	, bSaveAsyncCacheTriggered(false)
 	, CurrentSearchTime(0.)
@@ -3709,15 +3720,26 @@ uint32 FAssetDataGatherer::Run()
 			bool bLocalIdle = false;
 			{
 				FGathererScopeLock ResultsScopeLock(&ResultsLock); // bIsIdle requires the lock
-				if (IsStopped || bSaveAsyncCacheTriggered || (!IsPaused && !bIsIdle))
+				if (IsStopped || bSaveAsyncCacheTriggered || (!IsGatheringPaused && !bIsIdle))
 				{
 					break;
 				}
 				bLocalIdle = bIsIdle;
 			}
-			// No work to do. Sleep for a little and try again later.
-			// TODO: Need IsPaused to be a condition variable so we avoid sleeping while waiting for it and then taking a long time to wake after it is unset.
-			FPlatformProcess::Sleep(bLocalIdle ? IdleSleepTime : PausedSleepTime);
+
+			UE::AssetRegistry::Impl::EGatherStatus Status = UE::AssetRegistry::Impl::EGatherStatus::Complete;
+			if (bLocalIdle && !UE::AssetDataGather::Private::bTickGatherOnGTOnly 
+				&& (IsProcessingPaused.load(std::memory_order_relaxed) == 0))
+			{
+				Status = Cast<UAssetRegistryImpl>(IAssetRegistry::Get())->TickOnBackgroundThread();
+			}
+
+			// TODO: Need IsGatheringPaused to be a condition variable so we avoid sleeping while waiting for it and then taking a long time to wake after it is unset.
+			if ((Status != UE::AssetRegistry::Impl::EGatherStatus::TickActiveGatherActive)
+				&& (Status != UE::AssetRegistry::Impl::EGatherStatus::TickActiveGatherIdle))
+			{
+				FPlatformProcess::Sleep(bLocalIdle ? IdleSleepTime : PausedSleepTime);
+			}
 		}
 	}
 	return 0;
@@ -3750,7 +3772,7 @@ void FAssetDataGatherer::InnerTickLoop(bool bInSynchronousTick, bool bContribute
 			{
 				break;
 			}
-			if (IsStopped || (!bInSynchronousTick && IsPaused))
+			if (IsStopped || (!bInSynchronousTick && IsGatheringPaused))
 			{
 				break;
 			}
@@ -3897,7 +3919,7 @@ void FAssetDataGatherer::SaveMonolithicCacheFile(const TArray<TPair<FName,FDiskC
 		}
 	}
 
-	FScopedPause ScopedPause(*this);
+	FScopedGatheringPause ScopedPause(*this);
 	FGathererScopeLock TickScopeLock(&TickLock);
 	bIsSavingAsyncCache = false;
 	check(CacheInUseCount > 0);
@@ -3906,21 +3928,21 @@ void FAssetDataGatherer::SaveMonolithicCacheFile(const TArray<TPair<FName,FDiskC
 	LastCacheWriteTime = FPlatformTime::Seconds();
 }
 
-FAssetDataGatherer::FScopedPause::FScopedPause(const FAssetDataGatherer& InOwner)
+FAssetDataGatherer::FScopedGatheringPause::FScopedGatheringPause(const FAssetDataGatherer& InOwner)
 	:Owner(InOwner)
 {
 	if (!Owner.IsSynchronous())
 	{
-		Owner.IsPaused++;
+		Owner.IsGatheringPaused++;
 	}
 }
 
-FAssetDataGatherer::FScopedPause::~FScopedPause()
+FAssetDataGatherer::FScopedGatheringPause::~FScopedGatheringPause()
 {
 	if (!Owner.IsSynchronous())
 	{
-		check(Owner.IsPaused > 0)
-		Owner.IsPaused--;
+		check(Owner.IsGatheringPaused > 0)
+		Owner.IsGatheringPaused--;
 	}
 }
 
@@ -3965,8 +3987,10 @@ FAssetDataGatherer::ETickResult FAssetDataGatherer::TickInternal(double& TickSta
 	typedef TInlineAllocator<AssetDataGathererConstants::ExpectedMaxBatchSize> FBatchInlineAllocator;
 
 	TArray<FGatheredPathData, FBatchInlineAllocator> LocalFilesToSearch;
-	TArray<FAssetData*, FBatchInlineAllocator> LocalAssetResults;
+	TArray<TUniquePtr<FAssetData>, FBatchInlineAllocator> LocalAssetResults;
+	TArray<TUniquePtr<FAssetData>, FBatchInlineAllocator> LocalAssetResultsForGameThread;
 	TArray<FPackageDependencyData, FBatchInlineAllocator> LocalDependencyResults;
+	TArray<FPackageDependencyData, FBatchInlineAllocator> LocalDependencyResultsForGameThread;
 	TArray<FString, FBatchInlineAllocator> LocalCookedPackageNamesWithoutAssetDataResults;
 	TArray<FName, FBatchInlineAllocator> LocalVerseResults;
 	TArray<FString, FBatchInlineAllocator> LocalBlockedResults;
@@ -4172,13 +4196,34 @@ FAssetDataGatherer::ETickResult FAssetDataGatherer::TickInternal(double& TickSta
 			DiskCachedAssetData->DependencyData.bHasPackageData = bGatherAssetPackageData;
 			DiskCachedAssetData->DependencyData.bHasDependencyData = bGatherDependsData;
 
-			LocalAssetResults.Reserve(LocalAssetResults.Num() + DiskCachedAssetData->AssetDataList.Num());
+			bool MustBeHandledByGameThread = false;
+			// 
+			// In the future, we may need to process certain assets on the game thread because, e.g.,
+			// we may need to support PostLoadAssetRegistryTags running on the game thread. 
+			// The infrastructure is provided here to handle that case. In order to do so,
+			// implement ClassRequiresGameThreadProcessing in UAssetRegistryImpl and call it as shown below
+			// The rest of the functions in the asset registry respect the separation of data into general and 
+			// ForGameThread containers. In particular, these are consumed in FAssetRegistryImpl::TickGatherer.
+			// 
+			//for (const FAssetData& AssetData : DiskCachedAssetData->AssetDataList)
+			//{
+			//	if (AssetRegistry.ClassRequiresGameThreadProcessing(AssetData.GetClass()))
+			//	{
+			//		MustBeHandledByGameThread = true;
+			//		break;
+			//	}
+			//}
+
+			TArray<TUniquePtr<FAssetData>, FBatchInlineAllocator>& TargetAssetResults = MustBeHandledByGameThread ? LocalAssetResultsForGameThread : LocalAssetResults;
+			TArray<FPackageDependencyData, FBatchInlineAllocator>& TargetDependencyResults = MustBeHandledByGameThread ? LocalDependencyResultsForGameThread : LocalDependencyResults;
+
+			TargetAssetResults.Reserve(TargetAssetResults.Num() + DiskCachedAssetData->AssetDataList.Num());
 			for (const FAssetData& AssetData : DiskCachedAssetData->AssetDataList)
 			{
-				LocalAssetResults.Add(new FAssetData(AssetData));
+				TargetAssetResults.Add(MakeUnique<FAssetData>(AssetData));
 			}
+			TargetDependencyResults.Add(DiskCachedAssetData->DependencyData);
 
-			LocalDependencyResults.Add(DiskCachedAssetData->DependencyData);
 
 			AddToCache(PackageName, DiskCachedAssetData);
 		}
@@ -4209,7 +4254,7 @@ FAssetDataGatherer::ETickResult FAssetDataGatherer::TickInternal(double& TickSta
 		[this, &ReadContexts](int32 Index)
 		{
 			FReadContext& ReadContext = ReadContexts[Index];
-			if (!bSynchronousTick && IsPaused)
+			if (!bSynchronousTick && IsGatheringPaused)
 			{
 				ReadContext.bCanceled = true;
 				return;
@@ -4268,9 +4313,30 @@ FAssetDataGatherer::ETickResult FAssetDataGatherer::TickInternal(double& TickSta
 				AddToCache(ReadContext.PackageName, NewData);
 			}
 
+			bool MustBeHandledByGameThread = false;
+			// 
+			// In the future, we may need to process certain assets on the game thread because, e.g.,
+			// we may need to support PostLoadAssetRegistryTags running on the game thread. 
+			// The infrastructure is provided here to handle that case. In order to do so,
+			// implement ClassRequiresGameThreadProcessing in UAssetRegistryImpl and call it as shown below
+			// The rest of the functions in the asset registry respect the separation of data into general and 
+			// ForGameThread containers. In particular, these are consumed in FAssetRegistryImpl::TickGatherer.
+			// 
+			//for (const FAssetData* BackgroundAssetData : ReadContext.AssetDataFromFile)
+			//{
+			//	if (AssetRegistry.ClassRequiresGameThreadProcessing(BackgroundAssetData->GetClass()))
+			//	{
+			//		MustBeHandledByGameThread = true;
+			//		break;
+			//	}
+			//}
+
+			TArray<TUniquePtr<FAssetData>, FBatchInlineAllocator>& TargetAssetResults = MustBeHandledByGameThread ? LocalAssetResultsForGameThread : LocalAssetResults;
+			TArray<FPackageDependencyData, FBatchInlineAllocator>& TargetDependencyResults = MustBeHandledByGameThread ? LocalDependencyResultsForGameThread : LocalDependencyResults;
+
 			// Add the results from the package into our output results
-			LocalAssetResults.Append(MoveTemp(ReadContext.AssetDataFromFile));
-			LocalDependencyResults.Add(MoveTemp(ReadContext.DependencyData));
+			TargetAssetResults.Append(MoveTemp(ReadContext.AssetDataFromFile));
+			TargetDependencyResults.Add(MoveTemp(ReadContext.DependencyData));
 		}
 		else if (ReadContext.bCanAttemptAssetRetry)
 		{
@@ -4285,7 +4351,9 @@ FAssetDataGatherer::ETickResult FAssetDataGatherer::TickInternal(double& TickSta
 
 		// Submit the results into the thread-shared lists
 		AssetResults.Append(MoveTemp(LocalAssetResults));
+		AssetResultsForGameThread.Append(MoveTemp(LocalAssetResultsForGameThread));
 		DependencyResults.Append(MoveTemp(LocalDependencyResults));
+		DependencyResultsForGameThread.Append(MoveTemp(LocalDependencyResultsForGameThread));
 		CookedPackageNamesWithoutAssetDataResults.Append(MoveTemp(LocalCookedPackageNamesWithoutAssetDataResults));
 		VerseResults.Append(MoveTemp(LocalVerseResults));
 		BlockedResults.Append(MoveTemp(LocalBlockedResults));
@@ -4435,26 +4503,20 @@ void FAssetDataGatherer::AddToCache(FName PackageName, FDiskCachedAssetData* Dis
 
 void FAssetDataGatherer::GetAndTrimSearchResults(FResults& InOutResults, FResultContext& OutContext)
 {
-	FGathererScopeLock ResultsScopeLock(&ResultsLock);
-
 	auto MoveAppendRangeToRingBuffer = [](auto& InOutRingBuffer, auto& InArray)
 	{
 		InOutRingBuffer.MoveAppendRange(InArray.GetData(), InArray.Num());
 		InArray.Reset();
 	};
 
-	for (FAssetData* AssetData : AssetResults)
-	{
-		InOutResults.Assets.Add(AssetData->PackageName, AssetData);
-	}
-	AssetResults.Reset();
+	// GetPackageResults takes its own lock.
+	GetPackageResults(InOutResults);
+
+	FGathererScopeLock ResultsScopeLock(&ResultsLock);
+
 	MoveAppendRangeToRingBuffer(InOutResults.Paths, DiscoveredPaths);
-	for (FPackageDependencyData& DependencyData : DependencyResults)
-	{
-		FName PackageName = DependencyData.PackageName;
-		InOutResults.Dependencies.Add(PackageName, MoveTemp(DependencyData));
-	}
-	DependencyResults.Reset();
+
+
 	MoveAppendRangeToRingBuffer(InOutResults.CookedPackageNamesWithoutAssetData, CookedPackageNamesWithoutAssetDataResults);
 	MoveAppendRangeToRingBuffer(InOutResults.VerseFiles, VerseResults);
 
@@ -4491,21 +4553,32 @@ FAssetGatherDiagnostics FAssetDataGatherer::GetDiagnostics()
 	return Diag;
 }
 
-void FAssetDataGatherer::GetPackageResults(TMultiMap<FName, FAssetData*>& OutAssetResults, TMultiMap<FName, FPackageDependencyData>& OutDependencyResults)
+void FAssetDataGatherer::GetPackageResults(FResults& InOutResults)
 {
 	FGathererScopeLock ResultsScopeLock(&ResultsLock);
 
-	for (FAssetData* AssetData : AssetResults)
+	for (TUniquePtr<FAssetData>& AssetData : AssetResults)
 	{
-		OutAssetResults.Add(AssetData->PackageName, AssetData);
+		InOutResults.Assets.Add(AssetData->PackageName, MoveTemp(AssetData));
 	}
 	AssetResults.Reset();
+	for (TUniquePtr<FAssetData>& AssetData : AssetResultsForGameThread)
+	{
+		InOutResults.AssetsForGameThread.Add(AssetData->PackageName, MoveTemp(AssetData));
+	}
+	AssetResultsForGameThread.Reset();
 	for (FPackageDependencyData& DependencyData : DependencyResults)
 	{
 		FName PackageName = DependencyData.PackageName;
-		OutDependencyResults.Add(PackageName, MoveTemp(DependencyData));
+		InOutResults.Dependencies.Add(PackageName, MoveTemp(DependencyData));
 	}
 	DependencyResults.Reset();
+	for (FPackageDependencyData& DependencyData : DependencyResultsForGameThread)
+	{
+		FName PackageName = DependencyData.PackageName;
+		InOutResults.DependenciesForGameThread.Add(PackageName, MoveTemp(DependencyData));
+	}
+	DependencyResultsForGameThread.Reset();
 }
 
 void FAssetDataGatherer::WaitOnPath(FStringView InPath)
@@ -4602,7 +4675,7 @@ void FAssetDataGatherer::WaitOnPathsInternal(TArrayView<UE::AssetDataGather::Pri
 	LLM_SCOPE(ELLMTag::AssetRegistry);
 
 	// Request a halt to the async tick
-	FScopedPause ScopedPause(*this);
+	FScopedGatheringPause ScopedPause(*this);
 	CHECK_IS_NOT_LOCKED_CURRENT_THREAD(ResultsLock);
 	{
 		FGathererScopeLock TickScopeLock(&TickLock);
@@ -4681,7 +4754,7 @@ void FAssetDataGatherer::WaitForIdle(float TimeoutSeconds)
 	CHECK_IS_NOT_LOCKED_CURRENT_THREAD(ResultsLock);
 
 	// Request a halt to the async tick
-	FScopedPause ScopedPause(*this);
+	FScopedGatheringPause ScopedPause(*this);
 	// Tick until idle
 	for (;;)
 	{
@@ -4763,7 +4836,7 @@ void FAssetDataGatherer::LoadCacheFiles(TConstArrayView<FString> CacheFilenames)
 	}
 
 	TArray<FCachePayload> Payloads = UE::AssetDataGather::Private::LoadCacheFiles(CacheFilenames, false /* bIsMonolithicCache */);
-	FScopedPause ScopedPause(*this);
+	FScopedGatheringPause ScopedPause(*this);
 	CHECK_IS_NOT_LOCKED_CURRENT_THREAD(ResultsLock);
 	FGathererScopeLock TickScopeLock(&TickLock);
 	ConsumeCacheFiles(MoveTemp(Payloads));
@@ -5108,7 +5181,7 @@ SIZE_T FAssetDataGatherer::GetAllocatedSize() const
 
 	Result += sizeof(*Discovery) + Discovery->GetAllocatedSize();
 
-	FScopedPause ScopedPause(*this);
+	FScopedGatheringPause ScopedPause(*this);
 	CHECK_IS_NOT_LOCKED_CURRENT_THREAD(ResultsLock);
 	FGathererScopeLock TickScopeLock(&TickLock);
 	FGathererScopeLock ResultsScopeLock(&ResultsLock);
@@ -5117,7 +5190,7 @@ SIZE_T FAssetDataGatherer::GetAllocatedSize() const
 
 	Result += AssetResults.GetAllocatedSize();
 	FAssetDataTagMapSharedView::FMemoryCounter TagMemoryUsage;
-	for (FAssetData* Value : AssetResults)
+	for (const TUniquePtr<FAssetData>& Value : AssetResults)
 	{
 		Result += sizeof(*Value);
 		TagMemoryUsage.Include(Value->TagsAndValues);
diff --git a/Engine/Source/Runtime/AssetRegistry/Private/AssetDataGatherer.h b/Engine/Source/Runtime/AssetRegistry/Private/AssetDataGatherer.h
--- a/Engine/Source/Runtime/AssetRegistry/Private/AssetDataGatherer.h
+++ b/Engine/Source/Runtime/AssetRegistry/Private/AssetDataGatherer.h
@@ -29,6 +29,8 @@ struct FAssetData;
 class FDiskCachedAssetData;
 class FAssetRegistryReader;
 class FAssetRegistryWriter; // Not defined if !ALLOW_NAME_BATCH_SAVING
+namespace UE::AssetRegistry { class FAssetRegistryImpl; }
+
 namespace UE::AssetDataGather::Private
 {
 class FAssetDataDiscovery;
@@ -77,7 +79,7 @@ class FAssetDataGatherer : public FRunnable
 {
 public:
 	FAssetDataGatherer(const TArray<FString>& InLongPackageNamesDenyList,
-		const TArray<FString>& InMountRelativePathsDenyList, bool bInAsyncEnabled);
+		const TArray<FString>& InMountRelativePathsDenyList, bool bInAsyncEnabled, UE::AssetRegistry::FAssetRegistryImpl& InRegistryImpl);
 	virtual ~FAssetDataGatherer();
 
 	void OnInitialSearchCompleted();
@@ -108,24 +110,28 @@ public:
 	// Receiving Results (possibly while tick is running)
 	struct FResults
 	{
-		TMultiMap<FName, FAssetData*> Assets;
+		TMultiMap<FName, TUniquePtr<FAssetData>> Assets;
+		TMultiMap<FName, TUniquePtr<FAssetData>> AssetsForGameThread;
 		TRingBuffer<FString> Paths;
 		TMultiMap<FName, FPackageDependencyData> Dependencies;
+		TMultiMap<FName, FPackageDependencyData> DependenciesForGameThread;
 		TRingBuffer<FString> CookedPackageNamesWithoutAssetData;
 		TRingBuffer<FName> VerseFiles;
 		TArray<FString> BlockedFiles;
 
 		SIZE_T GetAllocatedSize() const
 		{
-			return Assets.GetAllocatedSize() + Paths.GetAllocatedSize() + Dependencies.GetAllocatedSize() +
-				CookedPackageNamesWithoutAssetData.GetAllocatedSize() + VerseFiles.GetAllocatedSize() +
+			return Assets.GetAllocatedSize() + AssetsForGameThread.GetAllocatedSize() + Paths.GetAllocatedSize() + Dependencies.GetAllocatedSize() +
+				DependenciesForGameThread.GetAllocatedSize() + 	CookedPackageNamesWithoutAssetData.GetAllocatedSize() + VerseFiles.GetAllocatedSize() +
 				BlockedFiles.GetAllocatedSize();
 		}
 		void Shrink()
 		{
 			Assets.Shrink();
+			AssetsForGameThread.Shrink();
 			Paths.Trim();
 			Dependencies.Shrink();
+			DependenciesForGameThread.Shrink();
 			CookedPackageNamesWithoutAssetData.Trim();
 			VerseFiles.Trim();
 			BlockedFiles.Shrink();
@@ -144,9 +150,8 @@ public:
 	void GetAndTrimSearchResults(FResults& InOutResults, FResultContext& OutContext);
 	/** Get diagnostics for telemetry or logging. */
 	FAssetGatherDiagnostics GetDiagnostics();
-	/** Gets just the AssetResults and DependencyResults from the data gatherer. */
-	void GetPackageResults(TMultiMap<FName, FAssetData*>& OutAssetResults,
-		TMultiMap<FName, FPackageDependencyData>& OutDependencyResults);
+	/** Gets just the Assets, AssetsForGameThread, Dependencies, and DependenciesForGameThread from the data gatherer. */
+	void GetPackageResults(FResults& InOutResults);
 	/**
 	 * Wait for all monitored assets under the given path to be added to search results.
 	 * Returns immediately if the given path is not monitored.
@@ -246,6 +251,13 @@ public:
 		FPackageDependencyData& DependencyData, TArray<FString>& CookedPackagesToLoadUponDiscovery,
 		FPackageReader::EReadOptions Options);
 
+	/** Callable by the main thread to request that this thread pause/resume processing data. Gathering can 
+	 *  still proceed during this time.
+	 */
+	void PauseProcessing() { IsProcessingPaused.fetch_add(1, std::memory_order_relaxed); }
+	void ResumeProcessing() { IsProcessingPaused.fetch_sub(1, std::memory_order_relaxed); }
+	bool IsProcessingPauseRequested() const { return IsProcessingPaused.load(std::memory_order_relaxed) != 0; }
+
 private:
 	enum class ETickResult
 	{
@@ -348,10 +360,10 @@ private:
 	void Shrink();
 
 	/** Scoped guard for pausing the asynchronous tick. */
-	struct FScopedPause
+	struct FScopedGatheringPause
 	{
-		FScopedPause(const FAssetDataGatherer& InOwner);
-		~FScopedPause();
+		FScopedGatheringPause(const FAssetDataGatherer& InOwner);
+		~FScopedGatheringPause();
 		const FAssetDataGatherer& Owner;
 	};
 
@@ -374,8 +386,8 @@ private:
 	 */
 	mutable FGathererCriticalSection ResultsLock;
 
-
 	// Variable section for variables that are constant during threading.
+	UE::AssetRegistry::FAssetRegistryImpl& AssetRegistry;
 
 	/**
 	 * Thread to run async Ticks on. Constant during threading.
@@ -400,10 +412,14 @@ private:
 
 	// Variable section for variables that are atomics read/writable from outside critical sections.
 
-	/** > 0 if we've been asked to abort work in progress at the next opportunity. */
+	/** > 0 if we've been asked to abort gathering work in progress at the next opportunity. */
 	std::atomic<uint32> IsStopped;
-	/** > 0 if we've been asked to pause the worker thread so a synchronous function can take over the tick. */
-	mutable std::atomic<uint32> IsPaused;
+	/** > 0 if we've been asked to pause the worker thread gathering work so a synchronous function can take over the tick. */
+	mutable std::atomic<uint32> IsGatheringPaused;
+	
+	/** > 0 if we've been asked to pause processing work (but not gathering work) at the next opportunity */
+	mutable std::atomic<uint32> IsProcessingPaused;
+
 	/**
 	 * Discovery subsystem; decides which paths to search and queries the FileManager to search directories.
 	 * Pointer is constant during threading. Object pointed to internally provides threadsafety.
@@ -424,9 +440,13 @@ private:
 	TUniquePtr<UE::AssetDataGather::Private::FFilesToSearch> FilesToSearch;
 
 	/** The asset data gathered from the searched files. */
-	TArray<FAssetData*> AssetResults;
+	TArray<TUniquePtr<FAssetData>> AssetResults;
+	/** Like AssetResults but for assets that must be processed on the game thread */
+	TArray<TUniquePtr<FAssetData>> AssetResultsForGameThread;
 	/** Dependency data gathered from the searched files packages. */
 	TArray<FPackageDependencyData> DependencyResults;
+	/** Like DependencyResults but for assets that must be processed on the game thread */
+	TArray<FPackageDependencyData> DependencyResultsForGameThread;
 	/**
 	 * A list of cooked packages that did not have asset data in them.
 	 * These assets may still contain assets (if they were older for example). 
diff --git a/Engine/Source/Runtime/AssetRegistry/Private/AssetRegistry.cpp b/Engine/Source/Runtime/AssetRegistry/Private/AssetRegistry.cpp
--- a/Engine/Source/Runtime/AssetRegistry/Private/AssetRegistry.cpp
+++ b/Engine/Source/Runtime/AssetRegistry/Private/AssetRegistry.cpp
@@ -162,6 +162,11 @@ namespace UE::AssetRegistry::Impl
 		TEXT("AssetRegistry.MaxSecondsPerFrame"),
 		UE::AssetRegistry::Impl::MaxSecondsPerFrame,
 		TEXT("Maximum amount of time allowed for Asset Registry processing, in seconds"));
+	float MaxSecondsPerTickBackgroundThread = 0.1f;
+	static FAutoConsoleVariableRef CVarAssetRegistryMaxSecondsPerTickBackgroundThread(
+		TEXT("AssetRegistry.MaxSecondsPerTickBackgroundThread"),
+		UE::AssetRegistry::Impl::MaxSecondsPerTickBackgroundThread,
+		TEXT("Maximum amount of time allowed for Asset Registry processing, in seconds, per iteration on the background thread. Very large values could result in main thread delays due to the background thread holding locks."));
 
 	/** If true, defer sorting of dependencies until loading is complete */
 	bool bDeferDependencySort = false;
@@ -181,6 +186,87 @@ namespace UE::AssetRegistry::Impl
 	const FName DestinationObjectFName(TEXT("DestinationObject"));
 }
 
+namespace UE::AssetRegistry
+{
+/** Keeps a FRWLock read-locked while this scope lives */
+/** This is almost a clone of the existing FReadScopeLock and similar types
+	however this adds an extra flag to help the background processing thread 
+	know when a higher priority thread would like to gain access to the protected
+	data */
+template<typename TScopeLockType> class TRWScopeLockWithPriority
+{
+public:
+	UE_NODISCARD_CTOR explicit TRWScopeLockWithPriority(Private::FRWLockWithPriority& InLock, 
+		Private::ELockPriority InPriority = Private::PriorityHigh)
+	: Lock(InLock)
+	, Priority(InPriority)
+	{
+		if (Priority == Private::PriorityHigh)
+		{
+			Lock.HighPriorityWaitersCount.fetch_add(1, std::memory_order_relaxed);
+		}
+		GuardWrapper.Emplace(Lock);
+		if (Priority == Private::PriorityHigh)
+		{
+			Lock.HighPriorityWaitersCount.fetch_sub(1, std::memory_order_relaxed);
+		}
+	}
+
+	TOptional<TScopeLockType> GuardWrapper;
+	UE::AssetRegistry::Private::FRWLockWithPriority& Lock;
+	UE::AssetRegistry::Private::ELockPriority Priority;
+};
+
+class FRWScopeLockWithPriority
+{
+public:
+	UE_NODISCARD_CTOR explicit FRWScopeLockWithPriority(Private::FRWLockWithPriority& InLockObject, 
+			FRWScopeLockType InLockType, Private::ELockPriority InPriority = Private::PriorityHigh)
+		: Lock(InLockObject)
+		, Priority(InPriority)
+		, LockType(InLockType)
+	{
+		if (Priority == Private::PriorityHigh)
+		{
+			Lock.HighPriorityWaitersCount.fetch_add(1, std::memory_order_relaxed);
+		}
+		GuardWrapper.Emplace(Lock, LockType);
+		if (Priority == Private::PriorityHigh)
+		{
+			Lock.HighPriorityWaitersCount.fetch_sub(1, std::memory_order_relaxed);
+		}
+	}
+
+	// NOTE: As the name suggests, this function should be used with caution. 
+	// It releases the read lock _before_ acquiring a new write lock. This is not an atomic operation and the caller should 
+	// not treat it as such. 
+	// E.g. Pointers read from protected data structures prior to this call may be invalid after the function is called. 
+	void ReleaseReadOnlyLockAndAcquireWriteLock_USE_WITH_CAUTION()
+	{
+		if (LockType == SLT_ReadOnly)
+		{
+			GuardWrapper.Reset();
+			if (Priority == Private::PriorityHigh)
+			{
+				Lock.HighPriorityWaitersCount.fetch_add(1, std::memory_order_relaxed);
+			}
+			GuardWrapper.Emplace(Lock, SLT_Write);
+			if (Priority == Private::PriorityHigh)
+			{
+				Lock.HighPriorityWaitersCount.fetch_sub(1, std::memory_order_relaxed);
+			}
+			LockType = SLT_Write;
+		}
+	}
+
+	UE::AssetRegistry::Private::FRWLockWithPriority& Lock;
+	TOptional<FRWScopeLock> GuardWrapper;
+	UE::AssetRegistry::Private::ELockPriority Priority;
+	FRWScopeLockType LockType;
+};
+
+}
+
 /**
  * Implementation of IAssetRegistryInterface; forwards calls from the CoreUObject-accessible IAssetRegistryInterface into the AssetRegistry-accessible IAssetRegistry
  */
@@ -587,7 +673,7 @@ void FAsyncConsumer::PrepareForConsume()
 	++ReferenceCount;
 };
 
-void FAsyncConsumer::Wait(UAssetRegistryImpl& UARI, FWriteScopeLock& ScopeLock)
+void FAsyncConsumer::Wait(UAssetRegistryImpl& UARI, UE::AssetRegistry::FInterfaceWriteScopeLock& ScopeLock)
 {
 	// Called within the lock
 	if (ReferenceCount == 0)
@@ -634,7 +720,7 @@ void FAsyncConsumer::Consume(UAssetRegistryImpl& UARI, UE::AssetRegistry::Impl::
 
 namespace UE::AssetRegistry
 {
-void FAssetRegistryImpl::ConditionalLoadPremadeAssetRegistry(UAssetRegistryImpl& UARI, Impl::FEventContext& EventContext, FWriteScopeLock& ScopeLock)
+void FAssetRegistryImpl::ConditionalLoadPremadeAssetRegistry(UAssetRegistryImpl& UARI, Impl::FEventContext& EventContext, UE::AssetRegistry::FInterfaceWriteScopeLock& ScopeLock)
 {
 	AsyncConsumer.Wait(UARI, ScopeLock);
 }
@@ -657,7 +743,7 @@ void FAssetRegistryImpl::ConsumeOrDeferPreloadedPremade(UAssetRegistryImpl& UARI
 		{
 			Impl::FEventContext EventContext;
 			{
-				FWriteScopeLock InterfaceScopeLock(UARI.InterfaceLock);
+				UE::AssetRegistry::FInterfaceWriteScopeLock InterfaceScopeLock(UARI.InterfaceLock);
 				AsyncConsumer.Consume(UARI, EventContext, LoadResult, MoveTemp(ARState));
 			}
 			UARI.Broadcast(EventContext);
@@ -752,7 +838,7 @@ UAssetRegistryImpl::UAssetRegistryImpl(const FObjectInitializer& ObjectInitializ
 
 	{
 		LLM_SCOPE(ELLMTag::AssetRegistry);
-		FWriteScopeLock InterfaceScopeLock(InterfaceLock);
+		UE::AssetRegistry::FInterfaceWriteScopeLock InterfaceScopeLock(InterfaceLock);
 		GetInheritanceContextWithRequiredLock(InterfaceScopeLock, Context.InheritanceContext,
 			Context.InheritanceBuffer);
 
@@ -851,7 +937,7 @@ void FAssetRegistryImpl::Initialize(Impl::FInitializeContext& Context)
 
 	bInitialSearchStarted = false;
 	bInitialSearchCompleted = true;
-	GatherStatus = Impl::EGatherStatus::Active;
+	GatherStatus = Impl::EGatherStatus::TickActiveGatherActive;
 	PerformanceMode = Impl::EPerformanceMode::MostlyStatic;
 
 	bSearchAllAssets = false;
@@ -911,7 +997,8 @@ void FAssetRegistryImpl::Initialize(Impl::FInitializeContext& Context)
 	bool bStartedAsyncGather = false;
 	if (ShouldSearchAllAssetsAtStart())
 	{
-		ConstructGatherer();
+		verify(TryConstructGathererIfNeeded());
+
 		if (GlobalGatherer->IsAsyncEnabled())
 		{
 			SearchAllAssetsInitialAsync(Context.Events, Context.InheritanceContext);
@@ -959,6 +1046,8 @@ void FAssetRegistryImpl::RebuildAssetDependencyGathererMapIfNeeded()
 		return;
 	}
 
+	FWriteScopeLock ScopeLock(RegisteredDependencyGathererClassesLock);
+
 	RegisteredDependencyGathererClasses.Reset();
 
 	TArray<UObject*> Classes;
@@ -1104,7 +1193,7 @@ UAssetRegistryImpl::UAssetRegistryImpl(FVTableHelper& Helper)
 
 bool UAssetRegistryImpl::OnResolveRedirect(const FString& InPackageName, FString& OutPackageName)
 {
-	FReadScopeLock InterfaceScopeLock(InterfaceLock);
+	UE::AssetRegistry::FInterfaceReadScopeLock InterfaceScopeLock(InterfaceLock);
 	return GuardedData.ResolveRedirect(InPackageName, OutPackageName);
 }
 
@@ -1201,7 +1290,7 @@ void UAssetRegistryImpl::OnPluginLoadingPhaseComplete(ELoadingPhase::Type Loadin
 	}
 	{
 		LLM_SCOPE(ELLMTag::AssetRegistry);
-		FWriteScopeLock InterfaceScopeLock(InterfaceLock);
+		UE::AssetRegistry::FInterfaceWriteScopeLock InterfaceScopeLock(InterfaceLock);
 		GuardedData.OnPostEngineInit(bPhaseSuccessful);
 	}
 
@@ -1226,6 +1315,7 @@ void FAssetRegistryImpl::OnPostEngineInit(bool bPhaseSuccessful)
 	// If we ever need to update the Filtering list outside of the game thread, we will need to defer the update 
 	// of the Filtering namespace to the tick function; UE::AssetRegistry::Filtering can only be used in game thread
 	check(IsInGameThread());
+
 	Utils::PopulateSkipClasses(SkipUncookedClasses, SkipCookedClasses);
 	UE::AssetRegistry::FFiltering::SetSkipClasses(SkipUncookedClasses, SkipCookedClasses);
 #endif
@@ -1264,7 +1354,7 @@ void UAssetRegistryImpl::InitializeSerializationOptions(FAssetRegistrySerializat
 {
 	if (PlatformIniName.IsEmpty())
 	{
-		FReadScopeLock InterfaceScopeLock(InterfaceLock);
+		UE::AssetRegistry::FInterfaceReadScopeLock InterfaceScopeLock(InterfaceLock);
 		// Use options we already loaded, the first pass for this happens at object creation time so this is always valid when queried externally
 		GuardedData.CopySerializationOptions(Options, Target);
 	}
@@ -1526,7 +1616,7 @@ void FAssetRegistryImpl::CollectCodeGeneratorClasses()
 void UAssetRegistryImpl::OnRefreshNativeClasses()
 {
 	LLM_SCOPE(ELLMTag::AssetRegistry);
-	FWriteScopeLock InterfaceScopeLock(InterfaceLock);
+	UE::AssetRegistry::FInterfaceWriteScopeLock InterfaceScopeLock(InterfaceLock);
 	GuardedData.RefreshNativeClasses();
 }
 
@@ -1556,7 +1646,7 @@ void UAssetRegistryImpl::OnFEngineLoopInitCompleteSearchAllAssets()
 void UAssetRegistryImpl::OnAssetDependencyGathererRegistered()
 {
 	LLM_SCOPE(ELLMTag::AssetRegistry);
-	FWriteScopeLock InterfaceScopeLock(InterfaceLock);
+	UE::AssetRegistry::FInterfaceWriteScopeLock InterfaceScopeLock(InterfaceLock);
 	GuardedData.OnAssetDependencyGathererRegistered();
 }
 #endif
@@ -1564,15 +1654,27 @@ void UAssetRegistryImpl::OnAssetDependencyGathererRegistered()
 void UAssetRegistryImpl::OnEnginePreExit()
 {
 	LLM_SCOPE(ELLMTag::AssetRegistry);
-	FWriteScopeLock InterfaceScopeLock(InterfaceLock);
-	GuardedData.OnEnginePreExit();
+
+	TUniquePtr<FAssetDataGatherer> GlobalGatherer;
+	{
+		FScopeLock GatheredDataGuard(&GatheredDataProcessingLock);
+		UE::AssetRegistry::FInterfaceWriteScopeLock InterfaceScopeLock(InterfaceLock);
+		GlobalGatherer = MoveTemp(GuardedData.AccessGlobalGatherer());
+		if (GlobalGatherer.IsValid())
+		{
+			GlobalGatherer->Stop();
+		}
+	}
+	// Now that we are no longer holding the lock, we can destroy the gatherer
+	GlobalGatherer.Reset();
+
 }
 
 void UAssetRegistryImpl::FinishDestroy()
 {
 	LLM_SCOPE(ELLMTag::AssetRegistry);
 	{
-		FWriteScopeLock InterfaceScopeLock(InterfaceLock);
+		UE::AssetRegistry::FInterfaceWriteScopeLock InterfaceScopeLock(InterfaceLock);
 
 		// Stop listening for content mount point events
 		FPackageName::OnContentPathMounted().RemoveAll(this);
@@ -1669,17 +1771,15 @@ UAssetRegistryImpl& UAssetRegistryImpl::Get()
 namespace UE::AssetRegistry
 {
 
-void FAssetRegistryImpl::OnEnginePreExit()
-{
-	// Shut down the GlobalGatherer's gather threads, before we start tearing down the engine
-	GlobalGatherer.Reset();
-}
-
-void FAssetRegistryImpl::ConstructGatherer()
+bool FAssetRegistryImpl::TryConstructGathererIfNeeded()
 {
 	if (GlobalGatherer.IsValid())
 	{
-		return;
+		return true;
+	}
+	else if (IsEngineExitRequested())
+	{
+		return false;
 	}
 
 	TArray<FString> PathsDenyList;
@@ -1691,7 +1791,7 @@ void FAssetRegistryImpl::ConstructGatherer()
 	}
 
 	bool bAsyncGatherEnabled = !IsRunningGame();
-	GlobalGatherer = MakeUnique<FAssetDataGatherer>(PathsDenyList, ContentSubPathsDenyList, bAsyncGatherEnabled);
+	GlobalGatherer = MakeUnique<FAssetDataGatherer>(PathsDenyList, ContentSubPathsDenyList, bAsyncGatherEnabled, *this);
 
 	// Read script packages if all initial plugins have been loaded, otherwise do nothing; we wait for the callback.
 	ELoadingPhase::Type LoadingPhase = IPluginManager::Get().GetLastCompletedLoadingPhase();
@@ -1699,6 +1799,7 @@ void FAssetRegistryImpl::ConstructGatherer()
 	{
 		ReadScriptPackages();
 	}
+	return true;
 }
 
 void FAssetRegistryImpl::SearchAllAssetsInitialAsync(Impl::FEventContext& EventContext,
@@ -1758,7 +1859,7 @@ void UAssetRegistryImpl::SearchAllAssets(bool bSynchronousSearch)
 	FEventContext EventContext;
 	{
 		LLM_SCOPE(ELLMTag::AssetRegistry);
-		FWriteScopeLock InterfaceScopeLock(InterfaceLock);
+		UE::AssetRegistry::FInterfaceWriteScopeLock InterfaceScopeLock(InterfaceLock);
 		FClassInheritanceContext InheritanceContext;
 		FClassInheritanceBuffer InheritanceBuffer;
 		GetInheritanceContextWithRequiredLock(InterfaceScopeLock, InheritanceContext, InheritanceBuffer);
@@ -1772,7 +1873,8 @@ void UAssetRegistryImpl::SearchAllAssets(bool bSynchronousSearch)
 #if WITH_EDITOR
 	if (bSynchronousSearch)
 	{
-		ProcessLoadedAssetsToUpdateCache(EventContext, -1., EGatherStatus::Complete);
+		UE::AssetRegistry::Impl::FInterruptionContext InterruptionContext(-1., -1.);
+		ProcessLoadedAssetsToUpdateCache(EventContext, EGatherStatus::Complete, InterruptionContext);
 	}
 #endif
 	Broadcast(EventContext);
@@ -1780,13 +1882,13 @@ void UAssetRegistryImpl::SearchAllAssets(bool bSynchronousSearch)
 
 bool UAssetRegistryImpl::IsSearchAllAssets() const
 {
-	FReadScopeLock InterfaceScopeLock(InterfaceLock);
+	UE::AssetRegistry::FInterfaceReadScopeLock InterfaceScopeLock(InterfaceLock);
 	return GuardedData.IsSearchAllAssets();
 }
 
 bool UAssetRegistryImpl::IsSearchAsync() const
 {
-	FReadScopeLock InterfaceScopeLock(InterfaceLock);
+	UE::AssetRegistry::FInterfaceReadScopeLock InterfaceScopeLock(InterfaceLock);
 	return GuardedData.IsInitialSearchStarted();
 }
 
@@ -1797,7 +1899,11 @@ void FAssetRegistryImpl::SearchAllAssets(Impl::FEventContext& EventContext,
 	Impl::FClassInheritanceContext& InheritanceContext, bool bSynchronousSearch)
 {
 	TRACE_BEGIN_REGION(TEXT("Asset Registry Scan"));
-	ConstructGatherer();
+	if (!TryConstructGathererIfNeeded())
+	{
+		return;
+	}
+
 	FAssetDataGatherer& Gatherer = *GlobalGatherer;
 	if (!Gatherer.IsAsyncEnabled())
 	{
@@ -1822,8 +1928,8 @@ void FAssetRegistryImpl::SearchAllAssets(Impl::FEventContext& EventContext,
 	if (bSynchronousSearch)
 	{
 		Gatherer.WaitForIdle();
-		bool bUnusedInterrupted;
-		Impl::EGatherStatus UnusedStatus = TickGatherer(EventContext, InheritanceContext, -1., bUnusedInterrupted);
+		Impl::FInterruptionContext InterruptionContext(-1., -1.);
+		Impl::EGatherStatus UnusedStatus = TickGatherer(EventContext, InheritanceContext, InterruptionContext);
 #if WITH_EDITOR
 		if (!bInitialSearchStarted)
 		{
@@ -1854,7 +1960,9 @@ void UAssetRegistryImpl::WaitForCompletion()
 	// while we consume the results in parallel.
 	{
 		LLM_SCOPE(ELLMTag::AssetRegistry);
-		FWriteScopeLock InterfaceScopeLock(InterfaceLock);
+		// We don't need to take the GatheredDataProcessingLock here because we actually *do*
+		// want to block until we can proceed
+		UE::AssetRegistry::FInterfaceWriteScopeLock InterfaceScopeLock(InterfaceLock);
 		FClassInheritanceContext InheritanceContext;
 		FClassInheritanceBuffer InheritanceBuffer;
 		GetInheritanceContextWithRequiredLock(InterfaceScopeLock, InheritanceContext, InheritanceBuffer);
@@ -1885,7 +1993,7 @@ void UAssetRegistryImpl::WaitForCompletion()
 		{
 			// Keep the LLM scope limited so it does not surround the broadcast which calls external code
 			LLM_SCOPE(ELLMTag::AssetRegistry);
-			FWriteScopeLock InterfaceScopeLock(InterfaceLock);
+			UE::AssetRegistry::FInterfaceWriteScopeLock InterfaceScopeLock(InterfaceLock);
 			FClassInheritanceContext InheritanceContext;
 			FClassInheritanceBuffer InheritanceBuffer;
 			GetInheritanceContextWithRequiredLock(InterfaceScopeLock, InheritanceContext, InheritanceBuffer);
@@ -1898,14 +2006,15 @@ void UAssetRegistryImpl::WaitForCompletion()
 
 			GuardedData.WaitForGathererIdleIfSynchronous();
 
-			bool bUnusedInterrupted;
-			Status = GuardedData.TickGatherer(EventContext, InheritanceContext, -1., bUnusedInterrupted);
+			UE::AssetRegistry::Impl::FInterruptionContext InterruptionContext(-1., -1.);
+			Status = GuardedData.TickGatherer(EventContext, InheritanceContext, InterruptionContext);
 		}
 #if WITH_EDITOR
-		ProcessLoadedAssetsToUpdateCache(EventContext, -1., Status);
+		UE::AssetRegistry::Impl::FInterruptionContext InterruptionContext(-1., -1.);
+		ProcessLoadedAssetsToUpdateCache(EventContext, Status, InterruptionContext);
 #endif
 		Broadcast(EventContext);
-		if (Status != EGatherStatus::Active && Status != EGatherStatus::WaitingForEvents)
+		if (!IsTickActive(Status) && Status != EGatherStatus::WaitingForEvents)
 		{
 			break;
 		}
@@ -1925,7 +2034,7 @@ void UAssetRegistryImpl::WaitForPremadeAssetRegistry()
 	FEventContext EventContext;
 	{
 		LLM_SCOPE(ELLMTag::AssetRegistry);
-		FWriteScopeLock InterfaceScopeLock(InterfaceLock);
+		UE::AssetRegistry::FInterfaceWriteScopeLock InterfaceScopeLock(InterfaceLock);
 		FClassInheritanceContext InheritanceContext;
 		FClassInheritanceBuffer InheritanceBuffer;
 		GetInheritanceContextWithRequiredLock(InterfaceScopeLock, InheritanceContext, InheritanceBuffer);
@@ -1937,7 +2046,7 @@ void UAssetRegistryImpl::WaitForPremadeAssetRegistry()
 void UAssetRegistryImpl::ClearGathererCache()
 {
 	LLM_SCOPE(ELLMTag::AssetRegistry);
-	FWriteScopeLock InterfaceScopeLock(InterfaceLock);
+	UE::AssetRegistry::FInterfaceWriteScopeLock InterfaceScopeLock(InterfaceLock);
 	GuardedData.ClearGathererCache();
 }
 
@@ -1961,7 +2070,7 @@ void UAssetRegistryImpl::WaitForPackage(const FString& PackageName)
 	UE::AssetRegistry::Impl::FEventContext EventContext;
 	{
 		LLM_SCOPE(ELLMTag::AssetRegistry);
-		FWriteScopeLock InterfaceScopeLock(InterfaceLock);
+		UE::AssetRegistry::FInterfaceWriteScopeLock InterfaceScopeLock(InterfaceLock);
 		if (GuardedData.IsLoadingAssets())
 		{
 			FString LocalPath;
@@ -1976,7 +2085,7 @@ void UAssetRegistryImpl::WaitForPackage(const FString& PackageName)
 
 bool UAssetRegistryImpl::HasAssets(const FName PackagePath, const bool bRecursive) const
 {
-	FReadScopeLock InterfaceScopeLock(InterfaceLock);
+	UE::AssetRegistry::FInterfaceReadScopeLock InterfaceScopeLock(InterfaceLock);
 	return GuardedData.HasAssets(PackagePath, bRecursive);
 }
 
@@ -2132,7 +2241,7 @@ bool UAssetRegistryImpl::GetAssets(const FARFilter& InFilter, TArray<FAssetData>
 	}
 
 	{
-		FReadScopeLock InterfaceScopeLock(InterfaceLock);
+		UE::AssetRegistry::FInterfaceReadScopeLock InterfaceScopeLock(InterfaceLock);
 		GuardedData.EnumerateDiskAssets(CompiledFilter, PackagesToSkip, [&OutAssetData](const FAssetData& AssetData)
 			{
 				OutAssetData.Emplace(AssetData);
@@ -2179,7 +2288,7 @@ bool UAssetRegistryImpl::EnumerateAssets(const FARCompiledFilter& InFilter, TFun
 
 	TArray<FAssetData, TInlineAllocator<128>> FoundAssets;
 	{
-		FReadScopeLock InterfaceScopeLock(InterfaceLock);
+		UE::AssetRegistry::FInterfaceReadScopeLock InterfaceScopeLock(InterfaceLock);
 		GuardedData.EnumerateDiskAssets(InFilter, PackagesToSkip, [&FoundAssets](const FAssetData& AssetData)
 			{
 				FoundAssets.Emplace(AssetData);
@@ -2425,7 +2534,7 @@ void EnumerateMemoryAssetsHelper(const FARCompiledFilter& InFilter, TSet<FName>&
 }
 
 void EnumerateMemoryAssets(const FARCompiledFilter& InFilter, TSet<FName>& OutPackageNamesWithAssets,
-	bool& bOutStopIteration, FRWLock& InterfaceLock, const FAssetRegistryState& GuardedDataState,
+	bool& bOutStopIteration, UE::AssetRegistry::Private::FInterfaceRWLock& InterfaceLock, const FAssetRegistryState& GuardedDataState,
 	TFunctionRef<bool(FAssetData&&)> Callback, bool bSkipARFilteredAssets)
 {
 	check(!InFilter.IsEmpty() && Utils::IsFilterValid(InFilter));
@@ -2439,7 +2548,7 @@ void EnumerateMemoryAssets(const FARCompiledFilter& InFilter, TSet<FName>& OutPa
 				// are too expensive to regularly compute but that exist in the on-disk Asset from SavePackage.
 				// Our contract for on-disk versus in-memory tags is that in-memory tags override on-disk tags, but we
 				// keep any on-disk tags that do not exist in the in-memory tags because they may be extended tags.
-				FReadScopeLock InterfaceScopeLock(InterfaceLock);
+				UE::AssetRegistry::FInterfaceReadScopeLock InterfaceScopeLock(InterfaceLock);
 				const FAssetData* OnDiskAssetData = GuardedDataState.GetAssetByObjectPath(FSoftObjectPath(Object));
 				if (OnDiskAssetData)
 				{
@@ -2521,7 +2630,7 @@ FAssetData UAssetRegistryImpl::GetAssetByObjectPath(const FSoftObjectPath& Objec
 	}
 
 	{
-		FReadScopeLock InterfaceScopeLock(InterfaceLock);
+		UE::AssetRegistry::FInterfaceReadScopeLock InterfaceScopeLock(InterfaceLock);
 		const FAssetRegistryState& State = GuardedData.GetState();
 		const FAssetData* FoundData = State.GetAssetByObjectPath(ObjectPath);
 		return (FoundData && !State.IsPackageUnmountedAndFiltered(FoundData->PackageName)
@@ -2538,7 +2647,7 @@ PRAGMA_ENABLE_DEPRECATION_WARNINGS;
 
 UE::AssetRegistry::EExists UAssetRegistryImpl::TryGetAssetByObjectPath(const FSoftObjectPath& ObjectPath, FAssetData& OutAssetData) const
 {
-	FReadScopeLock InterfaceScopeLock(InterfaceLock);
+	UE::AssetRegistry::FInterfaceReadScopeLock InterfaceScopeLock(InterfaceLock);
 	bool bAssetRegistryReady = GuardedData.IsInitialSearchStarted() && GuardedData.IsInitialSearchCompleted();
 	const FAssetRegistryState& State = GuardedData.GetState();
 	const FAssetData* FoundData = State.GetAssetByObjectPath(ObjectPath);
@@ -2562,7 +2671,7 @@ UE::AssetRegistry::EExists UAssetRegistryImpl::TryGetAssetPackageData(const FNam
 
 UE::AssetRegistry::EExists UAssetRegistryImpl::TryGetAssetPackageData(const FName PackageName, FAssetPackageData& OutAssetPackageData, FName& OutCorrectCasePackageName) const
 {
-	FReadScopeLock InterfaceScopeLock(InterfaceLock);
+	UE::AssetRegistry::FInterfaceReadScopeLock InterfaceScopeLock(InterfaceLock);
 	bool bAssetRegistryReady = GuardedData.IsInitialSearchStarted() && GuardedData.IsInitialSearchCompleted();
 	const FAssetRegistryState& State = GuardedData.GetState();
 	const FAssetPackageData* FoundData = State.GetAssetPackageData(PackageName, OutCorrectCasePackageName);
@@ -2596,7 +2705,7 @@ bool UAssetRegistryImpl::GetAllAssets(TArray<FAssetData>& OutAssetData, bool bIn
 	}
 
 	{
-		FReadScopeLock InterfaceScopeLock(InterfaceLock);
+		UE::AssetRegistry::FInterfaceReadScopeLock InterfaceScopeLock(InterfaceLock);
 		GuardedData.EnumerateAllDiskAssets(PackageNamesToSkip,
 			[&OutAssetData](const FAssetData& AssetData)
 			{
@@ -2632,7 +2741,7 @@ bool UAssetRegistryImpl::EnumerateAllAssets(TFunctionRef<bool(const FAssetData&)
 	// We have to call the callback on a copy rather than a reference since the callback may reenter the lock
 	TArray<FAssetData, TInlineAllocator<128>> OnDiskAssetDatas;
 	{
-		FReadScopeLock InterfaceScopeLock(InterfaceLock);
+		UE::AssetRegistry::FInterfaceReadScopeLock InterfaceScopeLock(InterfaceLock);
 		GuardedData.EnumerateAllDiskAssets(PackageNamesToSkip,
 			[&OnDiskAssetDatas](const FAssetData& AssetData)
 			{
@@ -2689,7 +2798,7 @@ void FAssetRegistryImpl::EnumerateAllDiskAssets(TSet<FName>& PackageNamesToSkip,
 
 void UAssetRegistryImpl::GetPackagesByName(FStringView PackageName, TArray<FName>& OutPackageNames) const
 {
-	FReadScopeLock InterfaceScopeLock(InterfaceLock);
+	UE::AssetRegistry::FInterfaceReadScopeLock InterfaceScopeLock(InterfaceLock);
 	const FAssetRegistryState& State = GuardedData.GetState();
 	UE_CLOG(GuardedData.IsInitialSearchStarted() && !GuardedData.IsInitialSearchCompleted(), LogAssetRegistry, Warning,
 		TEXT("GetPackagesByName has been called before AssetRegistry gather is complete and it does not wait. ")
@@ -2703,7 +2812,7 @@ FName UAssetRegistryImpl::GetFirstPackageByName(FStringView PackageName) const
 	FName LongPackageName;
 	bool bSearchAllAssets;
 	{
-		FReadScopeLock InterfaceScopeLock(InterfaceLock);
+		UE::AssetRegistry::FInterfaceReadScopeLock InterfaceScopeLock(InterfaceLock);
 		const FAssetRegistryState& State = GuardedData.GetState();
 		UE_CLOG(GuardedData.IsInitialSearchStarted() && !GuardedData.IsInitialSearchCompleted(), LogAssetRegistry, Warning,
 			TEXT("GetFirstPackageByName has been called before AssetRegistry gather is complete and it does not wait. ")
@@ -2737,13 +2846,13 @@ FName UAssetRegistryImpl::GetFirstPackageByName(FStringView PackageName) const
 
 bool UAssetRegistryImpl::GetDependencies(const FAssetIdentifier& AssetIdentifier, TArray<FAssetIdentifier>& OutDependencies, UE::AssetRegistry::EDependencyCategory Category, const UE::AssetRegistry::FDependencyQuery& Flags) const
 {
-	FReadScopeLock InterfaceScopeLock(InterfaceLock);
+	UE::AssetRegistry::FInterfaceReadScopeLock InterfaceScopeLock(InterfaceLock);
 	return GuardedData.GetState().GetDependencies(AssetIdentifier, OutDependencies, Category, Flags);
 }
 
 bool UAssetRegistryImpl::GetDependencies(const FAssetIdentifier& AssetIdentifier, TArray<FAssetDependency>& OutDependencies, UE::AssetRegistry::EDependencyCategory Category, const UE::AssetRegistry::FDependencyQuery& Flags) const
 {
-	FReadScopeLock InterfaceScopeLock(InterfaceLock);
+	UE::AssetRegistry::FInterfaceReadScopeLock InterfaceScopeLock(InterfaceLock);
 	return GuardedData.GetState().GetDependencies(AssetIdentifier, OutDependencies, Category, Flags);
 }
 
@@ -2798,13 +2907,13 @@ bool IAssetRegistry::K2_GetDependencies(FName PackageName, const FAssetRegistryD
 
 bool UAssetRegistryImpl::GetReferencers(const FAssetIdentifier& AssetIdentifier, TArray<FAssetIdentifier>& OutReferencers, UE::AssetRegistry::EDependencyCategory Category, const UE::AssetRegistry::FDependencyQuery& Flags) const
 {
-	FReadScopeLock InterfaceScopeLock(InterfaceLock);
+	UE::AssetRegistry::FInterfaceReadScopeLock InterfaceScopeLock(InterfaceLock);
 	return GuardedData.GetState().GetReferencers(AssetIdentifier, OutReferencers, Category, Flags);
 }
 
 bool UAssetRegistryImpl::GetReferencers(const FAssetIdentifier& AssetIdentifier, TArray<FAssetDependency>& OutReferencers, UE::AssetRegistry::EDependencyCategory Category, const UE::AssetRegistry::FDependencyQuery& Flags) const
 {
-	FReadScopeLock InterfaceScopeLock(InterfaceLock);
+	UE::AssetRegistry::FInterfaceReadScopeLock InterfaceScopeLock(InterfaceLock);
 	return GuardedData.GetState().GetReferencers(AssetIdentifier, OutReferencers, Category, Flags);
 }
 
@@ -2842,14 +2951,14 @@ bool IAssetRegistry::K2_GetReferencers(FName PackageName, const FAssetRegistryDe
 
 TOptional<FAssetPackageData> UAssetRegistryImpl::GetAssetPackageDataCopy(FName PackageName) const
 {
-	FReadScopeLock InterfaceScopeLock(InterfaceLock);
+	UE::AssetRegistry::FInterfaceReadScopeLock InterfaceScopeLock(InterfaceLock);
 	const FAssetPackageData* AssetPackageData = GuardedData.GetState().GetAssetPackageData(PackageName);
 	return AssetPackageData ? *AssetPackageData : TOptional<FAssetPackageData>();
 }
 
 void UAssetRegistryImpl::EnumerateAllPackages(TFunctionRef<void(FName PackageName, const FAssetPackageData& PackageData)> Callback) const
 {
-	FReadScopeLock InterfaceScopeLock(InterfaceLock);
+	UE::AssetRegistry::FInterfaceReadScopeLock InterfaceScopeLock(InterfaceLock);
 	for (const TPair<FName, const FAssetPackageData*>& Pair : GuardedData.GetState().GetAssetPackageDataMap())
 	{
 		Callback(Pair.Key, *Pair.Value);
@@ -2901,7 +3010,7 @@ bool UAssetRegistryImpl::DoesPackageExistOnDisk(FName PackageName, FString* OutC
 		FName CorrectCasePackageName;
 		const FAssetPackageData* AssetPackageData;
 		{
-			FReadScopeLock InterfaceScopeLock(InterfaceLock);
+			UE::AssetRegistry::FInterfaceReadScopeLock InterfaceScopeLock(InterfaceLock);
 			AssetPackageData = GuardedData.GetState().GetAssetPackageData(PackageName, CorrectCasePackageName);
 		}
 		const static bool bVerifyNegativeResults = FParse::Param(FCommandLine::Get(), TEXT("AssetRegistryValidatePackageExists"));
@@ -2991,7 +3100,7 @@ FSoftObjectPath UAssetRegistryImpl::GetRedirectedObjectPath(const FSoftObjectPat
 {
 	if (IsSearchAllAssets())
 	{
-		FReadScopeLock InterfaceScopeLock(InterfaceLock);
+		UE::AssetRegistry::FInterfaceReadScopeLock InterfaceScopeLock(InterfaceLock);
 		return GuardedData.GetRedirectedObjectPath(ObjectPath, nullptr, nullptr, /*bNeedsScanning*/ false);
 	}
 
@@ -3001,7 +3110,7 @@ FSoftObjectPath UAssetRegistryImpl::GetRedirectedObjectPath(const FSoftObjectPat
 	UE::AssetRegistry::Impl::FClassInheritanceBuffer InheritanceBuffer;
 	{
 		LLM_SCOPE(ELLMTag::AssetRegistry);
-		FWriteScopeLock WriteScopeLock(InterfaceLock);
+		UE::AssetRegistry::FInterfaceWriteScopeLock WriteScopeLock(InterfaceLock);
 		GetInheritanceContextWithRequiredLock(WriteScopeLock, InheritanceContext, InheritanceBuffer);
 		RedirectedObjectPath = GuardedData.GetRedirectedObjectPath(ObjectPath, &EventContext, &InheritanceContext, /*bNeedsScanning*/ true);
 	}	
@@ -3103,7 +3212,7 @@ bool UAssetRegistryImpl::GetAncestorClassNames(FTopLevelAssetPath ClassName, TAr
 {
 	UE::AssetRegistry::Impl::FClassInheritanceContext InheritanceContext;
 	UE::AssetRegistry::Impl::FClassInheritanceBuffer InheritanceBuffer;
-	FRWScopeLock InterfaceScopeLock(InterfaceLock, SLT_ReadOnly);
+	UE::AssetRegistry::FInterfaceRWScopeLock InterfaceScopeLock(InterfaceLock, SLT_ReadOnly);
 	const_cast<UAssetRegistryImpl*>(this)->GetInheritanceContextWithRequiredLock(
 		InterfaceScopeLock, InheritanceContext, InheritanceBuffer);
 	return GuardedData.GetAncestorClassNames(InheritanceContext, ClassName, OutAncestorClassNames);
@@ -3167,7 +3276,7 @@ void UAssetRegistryImpl::GetDerivedClassNames(const TArray<FTopLevelAssetPath>&
 {
 	UE::AssetRegistry::Impl::FClassInheritanceContext InheritanceContext;
 	UE::AssetRegistry::Impl::FClassInheritanceBuffer InheritanceBuffer;
-	FRWScopeLock InterfaceScopeLock(InterfaceLock, SLT_ReadOnly);
+	UE::AssetRegistry::FInterfaceRWScopeLock InterfaceScopeLock(InterfaceLock, SLT_ReadOnly);
 	const_cast<UAssetRegistryImpl*>(this)->GetInheritanceContextWithRequiredLock(
 		InterfaceScopeLock, InheritanceContext, InheritanceBuffer);
 	GuardedData.GetSubClasses(InheritanceContext, ClassNames, ExcludedClassNames, OutDerivedClassNames);
@@ -3175,7 +3284,7 @@ void UAssetRegistryImpl::GetDerivedClassNames(const TArray<FTopLevelAssetPath>&
 
 void UAssetRegistryImpl::GetAllCachedPaths(TArray<FString>& OutPathList) const
 {
-	FReadScopeLock InterfaceScopeLock(InterfaceLock);
+	UE::AssetRegistry::FInterfaceReadScopeLock InterfaceScopeLock(InterfaceLock);
 	const FPathTree& CachedPathTree = GuardedData.GetCachedPathTree();
 	OutPathList.Reserve(OutPathList.Num() + CachedPathTree.NumPaths());
 	CachedPathTree.EnumerateAllPaths([&OutPathList](FName Path)
@@ -3197,7 +3306,7 @@ void UAssetRegistryImpl::EnumerateAllCachedPaths(TFunctionRef<bool(FName)> Callb
 {
 	TArray<FName> FoundPaths;
 	{
-		FReadScopeLock InterfaceScopeLock(InterfaceLock);
+		UE::AssetRegistry::FInterfaceReadScopeLock InterfaceScopeLock(InterfaceLock);
 		const FPathTree& CachedPathTree = GuardedData.GetCachedPathTree();
 		FoundPaths.Reserve(CachedPathTree.NumPaths());
 		CachedPathTree.EnumerateAllPaths([&FoundPaths](FName Path)
@@ -3217,7 +3326,7 @@ void UAssetRegistryImpl::EnumerateAllCachedPaths(TFunctionRef<bool(FName)> Callb
 
 void UAssetRegistryImpl::GetSubPaths(const FString& InBasePath, TArray<FString>& OutPathList, bool bInRecurse) const
 {
-	FReadScopeLock InterfaceScopeLock(InterfaceLock);
+	UE::AssetRegistry::FInterfaceReadScopeLock InterfaceScopeLock(InterfaceLock);
 	const FPathTree& CachedPathTree = GuardedData.GetCachedPathTree();
 	CachedPathTree.EnumerateSubPaths(*InBasePath, [&OutPathList](FName Path)
 	{
@@ -3228,7 +3337,7 @@ void UAssetRegistryImpl::GetSubPaths(const FString& InBasePath, TArray<FString>&
 
 void UAssetRegistryImpl::GetSubPaths(const FName& InBasePath, TArray<FName>& OutPathList, bool bInRecurse) const
 {
-	FReadScopeLock InterfaceScopeLock(InterfaceLock);
+	UE::AssetRegistry::FInterfaceReadScopeLock InterfaceScopeLock(InterfaceLock);
 	const FPathTree& CachedPathTree = GuardedData.GetCachedPathTree();
 	CachedPathTree.EnumerateSubPaths(InBasePath, [&OutPathList](FName Path)
 	{
@@ -3241,7 +3350,7 @@ void UAssetRegistryImpl::EnumerateSubPaths(const FString& InBasePath, TFunctionR
 {
 	TArray<FName, TInlineAllocator<64>> SubPaths;
 	{
-		FReadScopeLock InterfaceScopeLock(InterfaceLock);
+		UE::AssetRegistry::FInterfaceReadScopeLock InterfaceScopeLock(InterfaceLock);
 		const FPathTree& CachedPathTree = GuardedData.GetCachedPathTree();
 		CachedPathTree.EnumerateSubPaths(FName(*InBasePath), [&SubPaths](FName PathName)
 		{
@@ -3262,7 +3371,7 @@ void UAssetRegistryImpl::EnumerateSubPaths(const FName InBasePath, TFunctionRef<
 {
 	TArray<FName, TInlineAllocator<64>> SubPaths;
 	{
-		FReadScopeLock InterfaceScopeLock(InterfaceLock);
+		UE::AssetRegistry::FInterfaceReadScopeLock InterfaceScopeLock(InterfaceLock);
 		const FPathTree& CachedPathTree = GuardedData.GetCachedPathTree();
 		CachedPathTree.EnumerateSubPaths(InBasePath, [&SubPaths](FName PathName)
 		{
@@ -3423,7 +3532,7 @@ void UAssetRegistryImpl::CompileFilter(const FARFilter& InFilter, FARCompiledFil
 {
 	UE::AssetRegistry::Impl::FClassInheritanceContext InheritanceContext;
 	UE::AssetRegistry::Impl::FClassInheritanceBuffer InheritanceBuffer;
-	FRWScopeLock InterfaceScopeLock(InterfaceLock, SLT_ReadOnly);
+	UE::AssetRegistry::FInterfaceRWScopeLock InterfaceScopeLock(InterfaceLock, SLT_ReadOnly);
 	if (InFilter.bRecursiveClasses)
 	{
 		const_cast<UAssetRegistryImpl*>(this)->GetInheritanceContextWithRequiredLock(
@@ -3683,13 +3792,13 @@ void PrioritizeAssetInstall(const FAssetData& AssetData)
 
 bool UAssetRegistryImpl::HasVerseFiles(FName PackagePath, bool bRecursive /*= false*/) const
 {
-	FReadScopeLock InterfaceScopeLock(InterfaceLock);
+	UE::AssetRegistry::FInterfaceReadScopeLock InterfaceScopeLock(InterfaceLock);
 	return GuardedData.GetVerseFilesByPath(PackagePath, /*OutFilePaths=*/nullptr, bRecursive);
 }
 
 bool UAssetRegistryImpl::GetVerseFilesByPath(FName PackagePath, TArray<FName>& OutFilePaths, bool bRecursive /*= false*/) const
 {
-	FReadScopeLock InterfaceScopeLock(InterfaceLock);
+	UE::AssetRegistry::FInterfaceReadScopeLock InterfaceScopeLock(InterfaceLock);
 	return GuardedData.GetVerseFilesByPath(PackagePath, &OutFilePaths, bRecursive);
 }
 
@@ -3734,7 +3843,7 @@ bool UAssetRegistryImpl::AddPath(const FString& PathToAdd)
 	bool bResult;
 	{
 		LLM_SCOPE(ELLMTag::AssetRegistry);
-		FWriteScopeLock InterfaceScopeLock(InterfaceLock);
+		UE::AssetRegistry::FInterfaceWriteScopeLock InterfaceScopeLock(InterfaceLock);
 		bResult = GuardedData.AddPath(EventContext, PathToAdd);
 	}
 	Broadcast(EventContext);
@@ -3771,7 +3880,7 @@ bool UAssetRegistryImpl::RemovePath(const FString& PathToRemove)
 	bool bResult;
 	{
 		LLM_SCOPE(ELLMTag::AssetRegistry);
-		FWriteScopeLock InterfaceScopeLock(InterfaceLock);
+		UE::AssetRegistry::FInterfaceWriteScopeLock InterfaceScopeLock(InterfaceLock);
 		bResult = GuardedData.RemoveAssetPath(EventContext, FName(*PathToRemove));
 	}
 	Broadcast(EventContext);
@@ -3785,7 +3894,7 @@ bool UAssetRegistryImpl::PathExists(const FString& PathToTest) const
 
 bool UAssetRegistryImpl::PathExists(const FName PathToTest) const
 {
-	FReadScopeLock InterfaceScopeLock(InterfaceLock);
+	UE::AssetRegistry::FInterfaceReadScopeLock InterfaceScopeLock(InterfaceLock);
 	return GuardedData.GetCachedPathTree().PathExists(PathToTest);
 }
 
@@ -3849,7 +3958,7 @@ void UAssetRegistryImpl::ScanPathsSynchronousInternal(const TArray<FString>& InD
 	bool bInitialSearchCompleted;
 	{
 		LLM_SCOPE(ELLMTag::AssetRegistry);
-		FWriteScopeLock InterfaceScopeLock(InterfaceLock);
+		UE::AssetRegistry::FInterfaceWriteScopeLock InterfaceScopeLock(InterfaceLock);
 		GetInheritanceContextWithRequiredLock(InterfaceScopeLock, InheritanceContext, InheritanceBuffer);
 
 		bInitialSearchStarted = GuardedData.IsInitialSearchStarted();
@@ -3866,7 +3975,8 @@ void UAssetRegistryImpl::ScanPathsSynchronousInternal(const TArray<FString>& InD
 #if WITH_EDITOR
 	if (bWaitForInMemoryObjects)
 	{
-		ProcessLoadedAssetsToUpdateCache(EventContext, -1., Context.Status);
+		UE::AssetRegistry::Impl::FInterruptionContext InterruptionContext(-1., -1.);
+		ProcessLoadedAssetsToUpdateCache(EventContext, Context.Status, InterruptionContext);
 	}
 #endif
 	Broadcast(EventContext);
@@ -3900,7 +4010,7 @@ void UAssetRegistryImpl::ScanPathsSynchronousInternal(const TArray<FString>& InD
 void UAssetRegistryImpl::PrioritizeSearchPath(const FString& PathToPrioritize)
 {
 	LLM_SCOPE(ELLMTag::AssetRegistry);
-	FWriteScopeLock InterfaceScopeLock(InterfaceLock);
+	UE::AssetRegistry::FInterfaceWriteScopeLock InterfaceScopeLock(InterfaceLock);
 	GuardedData.PrioritizeSearchPath(PathToPrioritize);
 }
 
@@ -3936,7 +4046,7 @@ void UAssetRegistryImpl::AssetCreated(UObject* NewAsset)
 		UE::AssetRegistry::Impl::FEventContext EventContext;
 		{
 			LLM_SCOPE(ELLMTag::AssetRegistry);
-			FWriteScopeLock InterfaceScopeLock(InterfaceLock);
+			UE::AssetRegistry::FInterfaceWriteScopeLock InterfaceScopeLock(InterfaceLock);
 			// If this package was marked as an empty package before, it is no longer empty, so remove it from the list
 			GuardedData.RemoveEmptyPackage(NewPackage->GetFName());
 
@@ -3973,7 +4083,7 @@ void UAssetRegistryImpl::AssetDeleted(UObject* DeletedAsset)
 		bool bShouldSkipAsset;
 		{
 			LLM_SCOPE(ELLMTag::AssetRegistry);
-			FWriteScopeLock InterfaceScopeLock(InterfaceLock);
+			UE::AssetRegistry::FInterfaceWriteScopeLock InterfaceScopeLock(InterfaceLock);
 
 			// Deleting the last asset in a package causes the package to be garbage collected.
 			// If the UPackage object is GCed, it will be considered 'Unloaded' which will cause it to
@@ -4046,7 +4156,7 @@ void UAssetRegistryImpl::AssetRenamed(const UObject* RenamedAsset, const FString
 		UE::AssetRegistry::Impl::FEventContext EventContext;
 		{
 			LLM_SCOPE(ELLMTag::AssetRegistry);
-			FWriteScopeLock InterfaceScopeLock(InterfaceLock);
+			UE::AssetRegistry::FInterfaceWriteScopeLock InterfaceScopeLock(InterfaceLock);
 			GuardedData.RemoveEmptyPackage(NewPackage->GetFName());
 
 			if (bOldPackageIsEmpty)
@@ -4082,7 +4192,7 @@ void UAssetRegistryImpl::AssetsSaved(TArray<FAssetData>&& Assets)
 #if WITH_EDITOR
 	UE::AssetRegistry::Impl::FEventContext EventContext;
 	{
-		FWriteScopeLock InterfaceScopeLock(InterfaceLock);
+		UE::AssetRegistry::FInterfaceWriteScopeLock InterfaceScopeLock(InterfaceLock);
 		GuardedData.AssetsSaved(EventContext, MoveTemp(Assets));
 	}
 	Broadcast(EventContext);
@@ -4103,7 +4213,7 @@ void UAssetRegistryImpl::AssetUpdateTags(UObject* Object, EAssetRegistryTagsCall
 
 	UE::AssetRegistry::Impl::FEventContext EventContext;
 	{
-		FWriteScopeLock InterfaceScopeLock(InterfaceLock);
+		UE::AssetRegistry::FInterfaceWriteScopeLock InterfaceScopeLock(InterfaceLock);
 		GuardedData.AssetsSaved(EventContext, MoveTemp(Assets));
 	}
 	Broadcast(EventContext);
@@ -4148,7 +4258,7 @@ void UAssetRegistryImpl::AssetTagsFinalized(const UObject& FinalizedAsset)
 	}
 	LLM_SCOPE(ELLMTag::AssetRegistry);
 
-	FWriteScopeLock InterfaceScopeLock(InterfaceLock);
+	UE::AssetRegistry::FInterfaceWriteScopeLock InterfaceScopeLock(InterfaceLock);
 	GuardedData.AddLoadedAssetToProcess(FinalizedAsset);
 #endif
 }
@@ -4174,7 +4284,7 @@ bool UAssetRegistryImpl::VerseCreated(const FString& FilePathOnDisk)
 	UE::AssetRegistry::Impl::FEventContext EventContext;
 	{
 		LLM_SCOPE(ELLMTag::AssetRegistry);
-		FWriteScopeLock InterfaceScopeLock(InterfaceLock);
+		UE::AssetRegistry::FInterfaceWriteScopeLock InterfaceScopeLock(InterfaceLock);
 		GuardedData.AddVerseFile(EventContext, *VersePackagePathName);
 	}
 	Broadcast(EventContext);
@@ -4203,7 +4313,7 @@ bool UAssetRegistryImpl::VerseDeleted(const FString& FilePathOnDisk)
 	UE::AssetRegistry::Impl::FEventContext EventContext;
 	{
 		LLM_SCOPE(ELLMTag::AssetRegistry);
-		FWriteScopeLock InterfaceScopeLock(InterfaceLock);
+		UE::AssetRegistry::FInterfaceWriteScopeLock InterfaceScopeLock(InterfaceLock);
 		GuardedData.RemoveVerseFile(EventContext, *VersePackagePathName);
 	}
 	Broadcast(EventContext);
@@ -4218,7 +4328,7 @@ void UAssetRegistryImpl::PackageDeleted(UPackage* DeletedPackage)
 	if (ensure(DeletedPackage))
 	{
 		LLM_SCOPE(ELLMTag::AssetRegistry);
-		FWriteScopeLock InterfaceScopeLock(InterfaceLock);
+		UE::AssetRegistry::FInterfaceWriteScopeLock InterfaceScopeLock(InterfaceLock);
 		GuardedData.RemovePackageData(EventContext, DeletedPackage->GetFName());
 	}
 	Broadcast(EventContext);
@@ -4226,7 +4336,7 @@ void UAssetRegistryImpl::PackageDeleted(UPackage* DeletedPackage)
 
 bool UAssetRegistryImpl::IsLoadingAssets() const
 {
-	FReadScopeLock InterfaceScopeLock(InterfaceLock);
+	UE::AssetRegistry::FInterfaceReadScopeLock InterfaceScopeLock(InterfaceLock);
 	return GuardedData.IsLoadingAssets();
 }
 
@@ -4240,27 +4350,84 @@ bool FAssetRegistryImpl::IsLoadingAssets() const
 
 }
 
+UE::AssetRegistry::Impl::EGatherStatus UAssetRegistryImpl::TickOnBackgroundThread()
+{
+	UE::AssetRegistry::Impl::EGatherStatus Status = UE::AssetRegistry::Impl::EGatherStatus::TickActiveGatherActive;
+
+	do
+	{
+		LLM_SCOPE(ELLMTag::AssetRegistry);
+		if (GatheredDataProcessingLock.TryLock())
+		{
+			ON_SCOPE_EXIT { GatheredDataProcessingLock.Unlock(); };
+			UE::AssetRegistry::Impl::FEventContext EventContext;
+			UE::AssetRegistry::Impl::FClassInheritanceContext InheritanceContext;
+			UE::AssetRegistry::Impl::FInitializeContext InitializeContext{ *this };
+			UE::AssetRegistry::Impl::FClassInheritanceBuffer InheritanceBuffer;
+			UE::AssetRegistry::FInterfaceWriteScopeLock InterfaceScopeLock(InterfaceLock, UE::AssetRegistry::Private::PriorityLow);
+			GetInheritanceContextWithRequiredLock(InterfaceScopeLock, InitializeContext.InheritanceContext, InitializeContext.InheritanceBuffer);
+
+			UE::AssetRegistry::Impl::FInterruptionContext::ShouldExitEarlyCallbackType EarlyExitHelper = 
+				[this]()->bool { return InterfaceLock.HasWaiters() || IsBackgroundProcessingPaused(); };
+
+			UE::AssetRegistry::Impl::FInterruptionContext InterruptionContext(FPlatformTime::Seconds(), 
+				UE::AssetRegistry::Impl::MaxSecondsPerTickBackgroundThread, EarlyExitHelper);
+			Status = GuardedData.TickGatherer(EventContext, InitializeContext.InheritanceContext, InterruptionContext);
+
+			DeferredEvents.Append(MoveTemp(EventContext));
+			EventContext.Clear();
+		}
+		else 
+		{
+			// If the game thread is holding the processing lock,
+			// let's just exit and let the thread run function decide what to do 
+			return UE::AssetRegistry::Impl::EGatherStatus::UnableToProgress;
+		}
+
+		if (IsBackgroundProcessingPaused())
+		{
+			return UE::AssetRegistry::Impl::EGatherStatus::UnableToProgress;
+		}
+
+		// This ensures that if there are multiple waiters we don't get in ahead of them
+		while(InterfaceLock.HasWaiters())
+		{
+			if (IsBackgroundProcessingPaused())
+			{
+				return UE::AssetRegistry::Impl::EGatherStatus::UnableToProgress;
+			}
+			FPlatformProcess::Yield();
+		}
+	} while (Status == UE::AssetRegistry::Impl::EGatherStatus::TickActiveGatherIdle);
+
+	return Status;
+}
+
+
 void UAssetRegistryImpl::Tick(float DeltaTime)
 {
 	checkf(IsInGameThread(), TEXT("The tick function executes deferred loads and events and must be on the game thread to do so."));
 	TRACE_CPUPROFILER_EVENT_SCOPE_STR("Asset Registry Tick");
 
-	UE::AssetRegistry::Impl::EGatherStatus Status = UE::AssetRegistry::Impl::EGatherStatus::Active;
+	UE::AssetRegistry::Impl::EGatherStatus Status = UE::AssetRegistry::Impl::EGatherStatus::TickActiveGatherActive;
 	double TickStartTime = -1; // Force a full flush if DeltaTime < 0
 	if (DeltaTime >= 0)
 	{
 		TickStartTime = FPlatformTime::Seconds();
 	}
 
-
-	bool bInterrupted;
+	bool bInterruptedOrShouldProcessDeferredEvents = false;
 	do
 	{
-		bInterrupted = false;
+		bInterruptedOrShouldProcessDeferredEvents = false;
+		UE::AssetRegistry::Impl::FInterruptionContext InterruptionContext(TickStartTime, UE::AssetRegistry::Impl::MaxSecondsPerFrame);
+
 		UE::AssetRegistry::Impl::FEventContext EventContext;
+
+		if (GatheredDataProcessingLock.TryLock())
 		{
 			LLM_SCOPE(ELLMTag::AssetRegistry);
-			FWriteScopeLock InterfaceScopeLock(InterfaceLock);
+			UE::AssetRegistry::FInterfaceWriteScopeLock InterfaceScopeLock(InterfaceLock);
 			UE::AssetRegistry::Impl::FClassInheritanceContext InheritanceContext;
 			UE::AssetRegistry::Impl::FClassInheritanceBuffer InheritanceBuffer;
 			GetInheritanceContextWithRequiredLock(InterfaceScopeLock, InheritanceContext, InheritanceBuffer);
@@ -4271,23 +4438,39 @@ void UAssetRegistryImpl::Tick(float DeltaTime)
 			if (EventContext.IsEmpty())
 			{
 				// Tick the Gatherer
-				Status = GuardedData.TickGatherer(EventContext, InheritanceContext, TickStartTime, bInterrupted);
+				Status = GuardedData.TickGatherer(EventContext, InheritanceContext, InterruptionContext);
 			}
 			else
 			{
 				// Skip the TickGather to deal with the DeferredEvents first
-				bInterrupted = true;
+				bInterruptedOrShouldProcessDeferredEvents = true;
 			}
+			GatheredDataProcessingLock.Unlock();
+		}
+		else
+		{
+			// We still take the interface lock but only briefly in this case. We don't try to actually process the data 
+			// because the background thread is working on it. We just pump the events.
+			UE::AssetRegistry::FInterfaceWriteScopeLock InterfaceScopeLock(InterfaceLock);
+			EventContext.Append(MoveTemp(DeferredEvents));
+			DeferredEvents.Clear();
 		}
 
+		bInterruptedOrShouldProcessDeferredEvents = bInterruptedOrShouldProcessDeferredEvents || InterruptionContext.WasInterrupted();
+
 #if WITH_EDITOR
-		if (!bInterrupted)
+		if (!bInterruptedOrShouldProcessDeferredEvents)
 		{
-			ProcessLoadedAssetsToUpdateCache(EventContext, TickStartTime, Status);
+			ProcessLoadedAssetsToUpdateCache(EventContext, Status, InterruptionContext);
+			bInterruptedOrShouldProcessDeferredEvents = bInterruptedOrShouldProcessDeferredEvents || InterruptionContext.WasInterrupted();
 		}
 #endif
-		Broadcast(EventContext);
-	} while ((bInterrupted || Status == UE::AssetRegistry::Impl::EGatherStatus::WaitingForEvents) &&
+
+		{
+			TRACE_CPUPROFILER_EVENT_SCOPE_STR("Asset Registry Event Broadcast");
+			Broadcast(EventContext);
+		}
+	} while ((bInterruptedOrShouldProcessDeferredEvents || Status == UE::AssetRegistry::Impl::EGatherStatus::WaitingForEvents) &&
 		(TickStartTime < 0 || (FPlatformTime::Seconds() - TickStartTime) <= UE::AssetRegistry::Impl::MaxSecondsPerFrame));
 }
 
@@ -4310,14 +4493,21 @@ void FAssetRegistryImpl::WaitForGathererIdle(float TimeoutSeconds)
 	}
 }
 
+bool FAssetRegistryImpl::ClassRequiresGameThreadProcessing(const UClass* Class) const
+{
+	// This function is not called. See FAssetDataGatherer::TickInternal for where
+	// it would be called if it were fully implemented.
+	return true;
+}
+
 Impl::EGatherStatus FAssetRegistryImpl::TickGatherer(Impl::FEventContext& EventContext,
-	Impl::FClassInheritanceContext& InheritanceContext, const double TickStartTime, bool& bOutInterrupted,
+	Impl::FClassInheritanceContext& InheritanceContext, Impl::FInterruptionContext& InOutInterruptionContext,
 	TOptional<FAssetsFoundCallback> AssetsFoundCallback)
 {
 	using namespace UE::AssetRegistry::Impl;
+	bool bLocalIsInGameThread = IsInGameThread();
 
 	EGatherStatus OutStatus = EGatherStatus::Complete;
-	bOutInterrupted = false;
 	if (!GlobalGatherer.IsValid())
 	{
 		return OutStatus;
@@ -4345,95 +4535,218 @@ Impl::EGatherStatus FAssetRegistryImpl::TickGatherer(Impl::FEventContext& EventC
 
 	// Gather results from the background search
 	FAssetDataGatherer::FResultContext ResultContext;
-	GlobalGatherer->GetAndTrimSearchResults(BackgroundResults, ResultContext);	
+	GlobalGatherer->GetAndTrimSearchResults(BackgroundResults, ResultContext);
+
 	// Report the search times
 	for (double SearchTime : ResultContext.SearchTimes)
 	{
 		UE_LOG(LogAssetRegistry, Verbose, TEXT("### Background search completed in %0.4f seconds"), SearchTime);
 	}
+
 	const bool bHadAssetsToProcess = BackgroundResults.Assets.Num() > 0 || BackgroundResults.Dependencies.Num() > 0;
 	auto GetNumGatherFromDiskPending = [&ResultContext, this]()
 	{
-		return ResultContext.NumFilesToSearch + ResultContext.NumPathsToSearch + BackgroundResults.Paths.Num() + BackgroundResults.Assets.Num()
-			+ BackgroundResults.Dependencies.Num() + BackgroundResults.CookedPackageNamesWithoutAssetData.Num();
+		return ResultContext.NumFilesToSearch + ResultContext.NumPathsToSearch + BackgroundResults.Paths.Num() + BackgroundResults.Assets.Num() + BackgroundResults.AssetsForGameThread.Num()
+			+ BackgroundResults.Dependencies.Num() + BackgroundResults.DependenciesForGameThread.Num() + BackgroundResults.CookedPackageNamesWithoutAssetData.Num()
+			+ DeferredAssets.Num() + DeferredAssetsForGameThread.Num() + DeferredDependencies.Num() + DeferredDependenciesForGameThread.Num();
 	};
+	
+	auto GetTotalNumPackagesNeedingDependencyCalculation = [this]()
+#if WITH_EDITOR
+	{ return PackagesNeedingDependencyCalculation.Num() + PackagesNeedingDependencyCalculationOnGameThread.Num(); };
+#else
+	{ return 0; };
+#endif
+
+	int32 InitialNumPending = GetNumGatherFromDiskPending() + GetTotalNumPackagesNeedingDependencyCalculation();
 	int32 NumPending = 0;
 	auto CalculateStatus =
-		[bHadAssetsToProcess, &NumPending, &ResultContext, &EventContext, &OutStatus, this]
-		(int32 NumGatherPending, bool bInterrupted)
+		[&NumPending, &InitialNumPending, &ResultContext, &EventContext, &GetTotalNumPackagesNeedingDependencyCalculation,
+		&OutStatus, &InOutInterruptionContext, this]
+		(int32 NumGatherPending)
 	{
 		// Compute total pending, plus highest pending for this run so we can show a good progress bar
-		NumPending = NumGatherPending
-#if WITH_EDITOR
-			+ (PackagesNeedingDependencyCalculation.Num() ? 1 : 0)
-#endif
-			;
+		NumPending = NumGatherPending + (GetTotalNumPackagesNeedingDependencyCalculation() ? 1 : 0);
 		HighestPending = FMath::Max(this->HighestPending, NumPending);
 
-		if (!bInterrupted && !ResultContext.bIsSearching && NumPending == 0)
+		if (!InOutInterruptionContext.WasInterrupted() && !ResultContext.bIsSearching && NumPending == 0)
 		{
 			OutStatus = EGatherStatus::Complete;
 		}
-		else if (!bInterrupted && !ResultContext.bAbleToProgress)
+		else if (!InOutInterruptionContext.WasInterrupted() && !ResultContext.bAbleToProgress)
 		{
 			OutStatus = EGatherStatus::UnableToProgress;
 		}
 		else
 		{
-			OutStatus = EGatherStatus::Active;
+			OutStatus = ResultContext.bAbleToProgress ? EGatherStatus::TickActiveGatherActive : EGatherStatus::TickActiveGatherIdle;
 		}
-	};
-	auto UpdateStatus = [bHadAssetsToProcess, &NumPending, &ResultContext, &EventContext, &OutStatus, this]()
-	{
-		// Notify the status change, only when something changed, or when sending the final result before going idle
-		if (ResultContext.bIsSearching || bHadAssetsToProcess ||
-			(OutStatus == EGatherStatus::Complete && this->GatherStatus != EGatherStatus::Complete))
+		if (OutStatus == Impl::EGatherStatus::TickActiveGatherIdle)
 		{
-			EventContext.ProgressUpdateData.Emplace(
-				HighestPending,					// NumTotalAssets
-				HighestPending - NumPending,	// NumAssetsProcessedByAssetRegistry
-				NumPending / 2,					// NumAssetsPendingDataLoad, divided by 2 because assets are double counted due to dependencies
-				ResultContext.bIsDiscoveringFiles // bIsDiscoveringAssetFiles
-			);
+			// if there's no additional work the gatherer thread can perform, change the status from TickActiveGatherIdle
+			// to TickGameThreadActiveGatherIdle.
+			if (DeferredAssets.Num() == 0
+#if WITH_EDITOR
+				&& PackagesNeedingDependencyCalculation.Num() == 0
+#endif
+				&& BackgroundResults.Assets.Num() == 0
+				&& BackgroundResults.Dependencies.Num() == 0
+				&& BackgroundResults.CookedPackageNamesWithoutAssetData.Num() == 0
+				&& BackgroundResults.Paths.Num() == 0
+				&& DeferredDependencies.Num() == 0)
+			{
+				OutStatus = Impl::EGatherStatus::TickGameThreadActiveGatherIdle;
+			}
 		}
-		this->GatherStatus = OutStatus;
 	};
+	auto UpdateStatus = [bHadAssetsToProcess, &NumPending, &ResultContext, &EventContext, &OutStatus, this]()
+		{
+			// Notify the status change, only when something changed, or when sending the final result before going idle
+			if (ResultContext.bIsSearching || bHadAssetsToProcess ||
+				(OutStatus == EGatherStatus::Complete && this->GatherStatus != EGatherStatus::Complete))
+			{
+				EventContext.ProgressUpdateData.Emplace(
+					HighestPending,					// NumTotalAssets
+					HighestPending - NumPending,	// NumAssetsProcessedByAssetRegistry
+					NumPending / 2,					// NumAssetsPendingDataLoad, divided by 2 because assets are double counted due to dependencies
+					ResultContext.bIsDiscoveringFiles // bIsDiscoveringAssetFiles
+				);
+			}
+			this->GatherStatus = OutStatus;
+		};
 
 	// Add discovered paths
 	if (BackgroundResults.Paths.Num())
 	{
 		LazyStartTimer();
-		PathDataGathered(EventContext, TickStartTime, BackgroundResults.Paths);
+		PathDataGathered(EventContext, BackgroundResults.Paths, InOutInterruptionContext);
+	}
+	if (InOutInterruptionContext.ShouldExitEarly())
+	{
+		CalculateStatus(GetNumGatherFromDiskPending());
+		UpdateStatus();
+		return OutStatus;
 	}
 
-	// Process the asset results
-	if (BackgroundResults.Assets.Num())
+	auto RunAssetSearchDataGathered = [this, &EventContext, &AssetsFoundCallback, &LazyStartTimer, &InOutInterruptionContext]
+	(TMultiMap<FName, TUniquePtr<FAssetData>>& InAssetResults,
+		TMultiMap<FName, TUniquePtr<FAssetData>>& OutDeferredAssetResults)
+		{
+			// Process the asset results
+			if (InAssetResults.Num())
+			{
+				LazyStartTimer();
+				// Mark the first amortize time
+				if (AssetsFoundCallback.IsSet())
+				{
+					TMultiMap<FName, FAssetData*> NonOwningContainer;
+					for (auto Iter = InAssetResults.CreateIterator(); Iter; ++Iter)
+					{
+						NonOwningContainer.Add(Iter.Key(), Iter.Value().Get());
+					}
+					AssetsFoundCallback.GetValue()(NonOwningContainer);
+				}
+
+				AssetSearchDataGathered(EventContext, InAssetResults, OutDeferredAssetResults, InOutInterruptionContext);
+			}
+		};
+	auto RunDependencyDataGathered = [this, &bLocalIsInGameThread, &LazyStartTimer, &InOutInterruptionContext]
+	(TMultiMap<FName, FPackageDependencyData>& DependenciesToProcess,
+		TMultiMap<FName, FPackageDependencyData>& OutDeferredDependencies,
+		TSet<FName>* OutPackagesNeedingDependencyCalculation)
+		{
+			// Add dependencies
+			if (DependenciesToProcess.Num())
+			{
+				LazyStartTimer();
+
+				DependencyDataGathered(DependenciesToProcess, OutDeferredDependencies, OutPackagesNeedingDependencyCalculation, InOutInterruptionContext);
+
+			}
+		};
+
+	bool bRetryAssetGathering = true;
+	bool bHasRetriedDeferredResults = false;
+	while (bRetryAssetGathering)
 	{
-		LazyStartTimer();
-		// Mark the first amortize time
-		if (AssetsFoundCallback.IsSet())
+		bRetryAssetGathering = false;
+
+		// Process the normal results and defer anything that isn't ready
+		RunAssetSearchDataGathered(BackgroundResults.Assets, DeferredAssets);
+		if (InOutInterruptionContext.ShouldExitEarly())
+		{
+			CalculateStatus(GetNumGatherFromDiskPending());
+			UpdateStatus();
+			return OutStatus;
+		}
+
+		if (bLocalIsInGameThread)
 		{
-			AssetsFoundCallback.GetValue()(BackgroundResults.Assets);
+			RunAssetSearchDataGathered(BackgroundResults.AssetsForGameThread, DeferredAssetsForGameThread);
+			if (InOutInterruptionContext.ShouldExitEarly())
+			{
+				CalculateStatus(GetNumGatherFromDiskPending());
+				UpdateStatus();
+				return OutStatus;
+			}
 		}
 
-		AssetSearchDataGathered(EventContext, TickStartTime, BackgroundResults.Assets);
-	}
+		TSet<FName>* PackagesNeedingDependencyCalculationPointer = nullptr;
+#if WITH_EDITOR
+		PackagesNeedingDependencyCalculationPointer = &PackagesNeedingDependencyCalculation;
+#endif
+		RunDependencyDataGathered(BackgroundResults.Dependencies, DeferredDependencies, PackagesNeedingDependencyCalculationPointer);
+		if (InOutInterruptionContext.ShouldExitEarly())
+		{
+			CalculateStatus(GetNumGatherFromDiskPending());
+			UpdateStatus();
+			return OutStatus;
+		}
 
-	// Add dependencies
-	if (BackgroundResults.Dependencies.Num())
-	{
-		LazyStartTimer();
-		DependencyDataGathered(TickStartTime, BackgroundResults.Dependencies);
+		if (bLocalIsInGameThread)
+		{
+#if WITH_EDITOR
+			PackagesNeedingDependencyCalculationPointer = &PackagesNeedingDependencyCalculationOnGameThread;
+#endif
+			RunDependencyDataGathered(BackgroundResults.DependenciesForGameThread, DeferredDependenciesForGameThread, PackagesNeedingDependencyCalculationPointer);
+			if (InOutInterruptionContext.ShouldExitEarly())
+			{
+				CalculateStatus(GetNumGatherFromDiskPending());
+				UpdateStatus();
+				return OutStatus;
+			}
+		}
+
+		// Retry deferred assets if we've finished all the other assets; we need to do this in the current tick
+		// so we avoid spuriously reporting status == UnableToProgress
+		if (!bHasRetriedDeferredResults)
+		{
+			bHasRetriedDeferredResults = true;
+			if (BackgroundResults.Assets.IsEmpty() &&
+				(!bLocalIsInGameThread || BackgroundResults.AssetsForGameThread.IsEmpty()))
+			{
+				if (!DeferredAssets.IsEmpty() || !DeferredDependencies.IsEmpty() ||
+					(bLocalIsInGameThread && (!DeferredAssetsForGameThread.IsEmpty() || !DeferredDependenciesForGameThread.IsEmpty())))
+				{
+					bRetryAssetGathering = true;
+					BackgroundResults.Assets.Append(MoveTemp(DeferredAssets));
+					BackgroundResults.AssetsForGameThread.Append(MoveTemp(DeferredAssetsForGameThread));
+					BackgroundResults.Dependencies.Append(MoveTemp(DeferredDependencies));
+					BackgroundResults.DependenciesForGameThread.Append(MoveTemp(DeferredDependenciesForGameThread));
+					bForceCompletionEvenIfPostLoadsFail = bPreloadingComplete && IsEngineStartupModuleLoadingComplete();
+				}
+			}
+		}
 	}
 
 	// Load cooked packages that do not have asset data
 	if (BackgroundResults.CookedPackageNamesWithoutAssetData.Num())
 	{
 		LazyStartTimer();
-		CookedPackageNamesWithoutAssetDataGathered(EventContext, TickStartTime, BackgroundResults.CookedPackageNamesWithoutAssetData, bOutInterrupted);
-		if (bOutInterrupted)
+		CookedPackageNamesWithoutAssetDataGathered(EventContext, BackgroundResults.CookedPackageNamesWithoutAssetData, InOutInterruptionContext);
+		if (InOutInterruptionContext.ShouldExitEarly())
 		{
-			CalculateStatus(GetNumGatherFromDiskPending(), true /* bInterrupted */);
+			CalculateStatus(GetNumGatherFromDiskPending());
 			UpdateStatus();
 			return OutStatus;
 		}
@@ -4443,7 +4756,13 @@ Impl::EGatherStatus FAssetRegistryImpl::TickGatherer(Impl::FEventContext& EventC
 	if (BackgroundResults.VerseFiles.Num())
 	{
 		LazyStartTimer();
-		VerseFilesGathered(EventContext, TickStartTime, BackgroundResults.VerseFiles);
+		VerseFilesGathered(EventContext, BackgroundResults.VerseFiles, InOutInterruptionContext);
+		if (InOutInterruptionContext.ShouldExitEarly())
+		{
+			CalculateStatus(GetNumGatherFromDiskPending());
+			UpdateStatus();
+			return OutStatus;
+		}
 	}
 
 	// Store blocked files to be reported
@@ -4453,24 +4772,47 @@ Impl::EGatherStatus FAssetRegistryImpl::TickGatherer(Impl::FEventContext& EventC
 		BackgroundResults.BlockedFiles.Reset();
 	}
 
-	// Load Calculated Dependencies when the gather from disk is complete; the full gather is not complete until after this is done
 	int32 NumGatherFromDiskPending = GetNumGatherFromDiskPending();
 #if WITH_EDITOR
+	// Load Calculated Dependencies when the gather from disk is complete; the full gather is not complete until after this is done
 	bool bDiskGatherComplete = !ResultContext.bIsSearching && NumGatherFromDiskPending == 0;
-	if (bDiskGatherComplete && PackagesNeedingDependencyCalculation.Num())
+
+	// We can't do do this work until we've finished startup because modules might add new entries to RegisteredDependencyGathererClasses
+	// as they are loaded.
+	if (bDiskGatherComplete && IsEngineStartupModuleLoadingComplete()
+		&& (PackagesNeedingDependencyCalculation.Num() || PackagesNeedingDependencyCalculationOnGameThread.Num()))
 	{
 		LazyStartTimer();
-		LoadCalculatedDependencies(nullptr, TickStartTime, InheritanceContext, bOutInterrupted);
-		if (bOutInterrupted)
-		{
-			CalculateStatus(NumGatherFromDiskPending, true /* bInterrupted */);
+		// Only assets whose classes have a RegisteredDependencyGathererClasses entry actually need to run through LoadCalculatedDependencies
+		// at all. Furthermore, in that situation, we must always perform the gather on the game thread. This function ensures that happens and
+		// clears out any spurious entries.
+		PruneAndCoalescePackagesRequiringDependencyCalculation(PackagesNeedingDependencyCalculation,
+			PackagesNeedingDependencyCalculationOnGameThread, InOutInterruptionContext);
+		if (InOutInterruptionContext.ShouldExitEarly())
+		{
+			CalculateStatus(NumGatherFromDiskPending);
 			UpdateStatus();
 			return OutStatus;
 		}
+		// As currently implemented, we cannot perform the dependency calculations on a background thread. Once we've
+		// called PruneAndCoalesce, all packages that actually need calculations will be in the OnGameThread container
+		// and all other packages will have been removed.
+		ensure(PackagesNeedingDependencyCalculation.Num() == 0);
+
+		if (PackagesNeedingDependencyCalculationOnGameThread.Num() && bLocalIsInGameThread)
+		{
+			LoadCalculatedDependencies(nullptr, InheritanceContext, &PackagesNeedingDependencyCalculationOnGameThread, InOutInterruptionContext);
+			if (InOutInterruptionContext.ShouldExitEarly())
+			{
+				CalculateStatus(NumGatherFromDiskPending);
+				UpdateStatus();
+				return OutStatus;
+			}
+		}
 	}
 #endif
 
-	CalculateStatus(NumGatherFromDiskPending, false /* bInterrupted */);
+	CalculateStatus(NumGatherFromDiskPending);
 
 	if (OutStatus == EGatherStatus::Complete)
 	{
@@ -4478,7 +4820,8 @@ Impl::EGatherStatus FAssetRegistryImpl::TickGatherer(Impl::FEventContext& EventC
 		{
 			// Finishing the background search is blocked until preloading complete because plugins can be mounted during
 			// startup up until that point, and we need to wait for all the plugins to load before declaring completion.
-			bool bCanCompleteInitialSearch = bPreloadingComplete && IsEngineStartupModuleLoadingComplete();
+			// Only the main thread can know that we're complete because we need to wait until we've broadcast the events
+			bool bCanCompleteInitialSearch = bPreloadingComplete && IsEngineStartupModuleLoadingComplete() && bLocalIsInGameThread;
 
 			if (bCanCompleteInitialSearch)
 			{
@@ -4505,11 +4848,20 @@ Impl::EGatherStatus FAssetRegistryImpl::TickGatherer(Impl::FEventContext& EventC
 	{
 		HighestPending = 0;
 		BackgroundResults.Shrink();
+		DeferredAssets.Shrink();
+		DeferredAssetsForGameThread.Shrink();
+		DeferredDependencies.Shrink();
+		DeferredDependenciesForGameThread.Shrink();
+#if WITH_EDITOR
+		PackagesNeedingDependencyCalculation.Shrink();
+		PackagesNeedingDependencyCalculationOnGameThread.Shrink();
+#endif
 	}
 
 	return OutStatus;
 }
 
+
 void FAssetRegistryImpl::OnInitialSearchCompleted(Impl::FEventContext& EventContext)
 {
 	bInitialSearchCompleted = true;
@@ -4597,24 +4949,52 @@ void FAssetRegistryImpl::TickGatherPackage(Impl::FEventContext& EventContext, co
 	FName PackageFName(PackageName);
 
 	// Gather results from the background search
-	GlobalGatherer->GetPackageResults(BackgroundResults.Assets, BackgroundResults.Dependencies);
+	GlobalGatherer->GetPackageResults(BackgroundResults);
 
-	TArray<FAssetData*> PackageAssets;
+	// The package could be in either the main or the ForGameThread containers but it will only appear in one or the other
+	// Either way, we put it into these two local containers and if we have to defer it, we'll put it into the game thread versions
+	TArray<TUniquePtr<FAssetData>*> PackageAssets;
 	TArray<FPackageDependencyData> PackageDependencyDatas;
-	BackgroundResults.Assets.MultiFind(PackageFName, PackageAssets);
-	BackgroundResults.Assets.Remove(PackageFName);
+	BackgroundResults.Assets.MultiFindPointer(PackageFName, PackageAssets);
+	BackgroundResults.AssetsForGameThread.MultiFindPointer(PackageFName, PackageAssets);
+	// We can't remove the assets until we've finished the transfer into the PackageAssetsMap below
 	BackgroundResults.Dependencies.MultiFind(PackageFName, PackageDependencyDatas);
 	BackgroundResults.Dependencies.Remove(PackageFName);
+	BackgroundResults.DependenciesForGameThread.MultiFind(PackageFName, PackageDependencyDatas);
+	BackgroundResults.DependenciesForGameThread.Remove(PackageFName);
+
 	if (PackageAssets.Num() > 0)
 	{
 		LazyStartTimer();
-		TMultiMap<FName, FAssetData*> PackageAssetsMap;
+		TMultiMap<FName, TUniquePtr<FAssetData>> PackageAssetsMap;
 		PackageAssetsMap.Reserve(PackageAssets.Num());
-		for (FAssetData* PackageAsset : PackageAssets)
+		for (TUniquePtr<FAssetData>* PackageAsset : PackageAssets)
 		{
-			PackageAssetsMap.Add(PackageFName, PackageAsset);
+			PackageAssetsMap.Add(PackageFName, MoveTemp(*PackageAsset));
+		}
+		// Ownership transfer is now complete so remove these packages from the results arrays
+		BackgroundResults.Assets.Remove(PackageFName);
+		BackgroundResults.AssetsForGameThread.Remove(PackageFName);
+
+		TMultiMap<FName, TUniquePtr<FAssetData>> DeferredPackageAssetsMap;
+		UE::AssetRegistry::Impl::FInterruptionContext InterruptionContext(-1., -1.);
+		AssetSearchDataGathered(EventContext, PackageAssetsMap, DeferredPackageAssetsMap, InterruptionContext);
+		if (DeferredPackageAssetsMap.Num())
+		{
+			UE_LOG(LogAssetRegistry, Warning, TEXT("Attempted to add package '%s' to the registry before its UClass was available. \
+Could not execute PostLoadAssetRegistryTags. We will try again later. Until then, dependency data will also be unavailable."),
+				*PackageName);
+			FDebug::DumpStackTraceToLog(ELogVerbosity::Warning);
+			DeferredAssetsForGameThread.Append(MoveTemp(DeferredPackageAssetsMap));
+			// If we are deferring this data we won't process the dependency data below anyway (we'll early out of DependencyDataGathered)
+			// so put the dependency data back into the BackgroundResults.DependenciesForGameThread which is where we will expect
+			// to find it when we reprocess the DeferredAssetsForGameThread after clearing the rest of the results queue.
+			for (FPackageDependencyData& Data : PackageDependencyDatas)
+			{
+				BackgroundResults.DependenciesForGameThread.Add(PackageFName, Data);
+			}
+			PackageDependencyDatas.Empty();
 		}
-		AssetSearchDataGathered(EventContext, -1., PackageAssetsMap);
 	}
 	if (PackageDependencyDatas.Num() > 0)
 	{
@@ -4625,27 +5005,28 @@ void FAssetRegistryImpl::TickGatherPackage(Impl::FEventContext& EventContext, co
 		{
 			PackageDependencyDatasMap.Add(PackageFName, MoveTemp(DependencyData));
 		}
-		DependencyDataGathered(-1., PackageDependencyDatasMap);
+		TSet<FName>* OutPackagesNeedingDependencyCalculation = nullptr;
+#if WITH_EDITOR
+		OutPackagesNeedingDependencyCalculation = &PackagesNeedingDependencyCalculation;
+#endif
+		UE::AssetRegistry::Impl::FInterruptionContext InterruptionContext(-1., -1.);
+		DependencyDataGathered(PackageDependencyDatasMap, DeferredDependenciesForGameThread, OutPackagesNeedingDependencyCalculation, InterruptionContext);
 	}
 }
 
 #if WITH_EDITOR
-void FAssetRegistryImpl::LoadCalculatedDependencies(TArray<FName>* AssetPackageNamesToCalculate, double TickStartTime,
-	Impl::FClassInheritanceContext& InheritanceContext, bool& bOutInterrupted)
+void FAssetRegistryImpl::LoadCalculatedDependencies(TArray<FName>* AssetPackageNamesToCalculate, 
+	Impl::FClassInheritanceContext& InheritanceContext, TSet<FName>* InPackagesNeedingDependencyCalculation,
+	Impl::FInterruptionContext& InOutInterruptionContext)
 {
-	auto CheckForTimeUp = [&TickStartTime, &bOutInterrupted](bool bHadActivity)
+	auto CheckForTimeUp = [&InOutInterruptionContext](bool bHadActivity)
 	{
 		// Only Check TimeUp when we found something to do, otherwise we waste time calling FPlatformTime::Seconds
 		if (!bHadActivity)
 		{
 			return false;
 		}
-		if (TickStartTime >= 0 && (FPlatformTime::Seconds() - TickStartTime) >= UE::AssetRegistry::Impl::MaxSecondsPerFrame)
-		{
-			bOutInterrupted = true;
-			return true;
-		}
-		return false;
+		return InOutInterruptionContext.ShouldExitEarly();
 	};
 
 	RebuildAssetDependencyGathererMapIfNeeded();
@@ -4654,10 +5035,10 @@ void FAssetRegistryImpl::LoadCalculatedDependencies(TArray<FName>* AssetPackageN
 	{
 		for (FName PackageName : *AssetPackageNamesToCalculate)
 		{
-			// We do not remove the package from PackagesNeedingDependencyCalculation, because
+			// We do not remove the package from InPackagesNeedingDependencyCalculation, because
 			// we are only calculating an interim result when AssetsToCalculate is non-null
 			// We will run again on each of these PackageNames when TickGatherer finishes gathering all dependencies
-			if (PackagesNeedingDependencyCalculation.Contains(PackageName))
+			if (InPackagesNeedingDependencyCalculation->Contains(PackageName))
 			{
 				bool bHadActivity;
 				LoadCalculatedDependencies(PackageName, InheritanceContext, bHadActivity);
@@ -4670,7 +5051,7 @@ void FAssetRegistryImpl::LoadCalculatedDependencies(TArray<FName>* AssetPackageN
 	}
 	else
 	{
-		for (TSet<FName>::TIterator It = PackagesNeedingDependencyCalculation.CreateIterator(); It; ++It)
+		for (TSet<FName>::TIterator It = InPackagesNeedingDependencyCalculation->CreateIterator(); It; ++It)
 		{
 			bool bHadActivity;
 			LoadCalculatedDependencies(*It, InheritanceContext, bHadActivity);
@@ -4680,7 +5061,7 @@ void FAssetRegistryImpl::LoadCalculatedDependencies(TArray<FName>* AssetPackageN
 				return;
 			}
 		}
-		check(PackagesNeedingDependencyCalculation.IsEmpty());
+		check(InPackagesNeedingDependencyCalculation->IsEmpty());
 	}
 }
 
@@ -4696,6 +5077,8 @@ void FAssetRegistryImpl::LoadCalculatedDependencies(FName PackageName,
 		return CompiledFilter;
 	};
 
+	FReadScopeLock GathererClassScopeLock(RegisteredDependencyGathererClassesLock);
+
 	TArray<UE::AssetDependencyGatherer::Private::FRegisteredAssetDependencyGatherer*, TInlineAllocator<2>> Gatherers;
 	for (const FAssetData* AssetData : State.GetAssetsByPackageName(PackageName))
 	{
@@ -4778,7 +5161,7 @@ void UAssetRegistryImpl::Serialize(FArchive& Ar)
 	UE::AssetRegistry::Impl::FEventContext EventContext;
 	{
 		LLM_SCOPE(ELLMTag::AssetRegistry);
-		FWriteScopeLock InterfaceScopeLock(InterfaceLock);
+		UE::AssetRegistry::FInterfaceWriteScopeLock InterfaceScopeLock(InterfaceLock);
 		GuardedData.Serialize(Ar, EventContext);
 	}
 	Broadcast(EventContext);
@@ -4814,7 +5197,7 @@ void UAssetRegistryImpl::AppendState(const FAssetRegistryState& InState)
 	UE::AssetRegistry::Impl::FEventContext EventContext;
 	{
 		LLM_SCOPE(ELLMTag::AssetRegistry);
-		FWriteScopeLock InterfaceScopeLock(InterfaceLock);
+		UE::AssetRegistry::FInterfaceWriteScopeLock InterfaceScopeLock(InterfaceLock);
 		GuardedData.AppendState(EventContext, InState);
 	}
 
@@ -4886,7 +5269,7 @@ SIZE_T UAssetRegistryImpl::GetAllocatedSize(bool bLogDetailed) const
 	SIZE_T StaticSize = 0;
 	SIZE_T SearchSize = 0;
 	{
-		FReadScopeLock InterfaceScopeLock(InterfaceLock);
+		UE::AssetRegistry::FInterfaceReadScopeLock InterfaceScopeLock(InterfaceLock);
 		GuardedData.GetAllocatedSize(bLogDetailed, StateSize, StaticSize, SearchSize);
 		StaticSize += sizeof(UAssetRegistryImpl);
 #if WITH_EDITOR
@@ -5006,7 +5389,7 @@ void UAssetRegistryImpl::InitializeTemporaryAssetRegistryState(FAssetRegistrySta
 {
 	using FAssetDataMap = UE::AssetRegistry::Private::FAssetDataMap;
 
-	FReadScopeLock InterfaceScopeLock(InterfaceLock);
+	UE::AssetRegistry::FInterfaceReadScopeLock InterfaceScopeLock(InterfaceLock);
 	const FAssetRegistryState& State = GuardedData.GetState();
 	if (!RequiredPackages.IsEmpty() || !RemovePackages.IsEmpty())
 	{
@@ -5033,14 +5416,14 @@ void UAssetRegistryImpl::InitializeTemporaryAssetRegistryState(FAssetRegistrySta
 #if ASSET_REGISTRY_STATE_DUMPING_ENABLED
 void UAssetRegistryImpl::DumpState(const TArray<FString>& Arguments, TArray<FString>& OutPages, int32 LinesPerPage) const
 {
-	FReadScopeLock InterfaceScopeLock(InterfaceLock);
+	UE::AssetRegistry::FInterfaceReadScopeLock InterfaceScopeLock(InterfaceLock);
 	GuardedData.GetState().Dump(Arguments, OutPages, LinesPerPage);
 }
 #endif
 
 const FAssetRegistryState* UAssetRegistryImpl::GetAssetRegistryState() const
 {
-	FReadScopeLock InterfaceScopeLock(InterfaceLock);
+	UE::AssetRegistry::FInterfaceReadScopeLock InterfaceScopeLock(InterfaceLock);
 	PRAGMA_DISABLE_DEPRECATION_WARNINGS
 	return &GuardedData.GetState();
 	PRAGMA_ENABLE_DEPRECATION_WARNINGS
@@ -5048,13 +5431,13 @@ const FAssetRegistryState* UAssetRegistryImpl::GetAssetRegistryState() const
 
 TSet<FName> UAssetRegistryImpl::GetCachedEmptyPackagesCopy() const
 {
-	FReadScopeLock InterfaceScopeLock(InterfaceLock);
+	UE::AssetRegistry::FInterfaceReadScopeLock InterfaceScopeLock(InterfaceLock);
 	return GuardedData.GetCachedEmptyPackages();
 }
 
 const TSet<FName>& UAssetRegistryImpl::GetCachedEmptyPackages() const
 {
-	FReadScopeLock InterfaceScopeLock(InterfaceLock);
+	UE::AssetRegistry::FInterfaceReadScopeLock InterfaceScopeLock(InterfaceLock);
 	PRAGMA_DISABLE_DEPRECATION_WARNINGS
 	return GuardedData.GetCachedEmptyPackages();
 	PRAGMA_ENABLE_DEPRECATION_WARNINGS
@@ -5062,7 +5445,7 @@ const TSet<FName>& UAssetRegistryImpl::GetCachedEmptyPackages() const
 
 bool UAssetRegistryImpl::ContainsTag(FName TagName) const
 {
-	FReadScopeLock InterfaceScopeLock(InterfaceLock);
+	UE::AssetRegistry::FInterfaceReadScopeLock InterfaceScopeLock(InterfaceLock);
 	return GuardedData.GetState().GetTagToAssetDatasMap().Contains(TagName);
 }
 
@@ -5166,7 +5549,10 @@ void FAssetRegistryImpl::ScanPathsSynchronous(Impl::FScanPathContext& Context)
 {
 	LLM_SCOPE(ELLMTag::AssetRegistry);
 
-	ConstructGatherer();
+	if (!TryConstructGathererIfNeeded())
+	{
+		return;
+	}
 	FAssetDataGatherer& Gatherer = *GlobalGatherer;
 
 	// Add a cache file for any not-yet-scanned dirs
@@ -5234,8 +5620,23 @@ void FAssetRegistryImpl::ScanPathsSynchronous(Impl::FScanPathContext& Context)
 	Gatherer.ScanPathsSynchronous(Context.LocalPaths, Context.bForceRescan, Context.bIgnoreDenyListScanFilters, CacheFilename, Context.PackageDirs);
 	TArray<FName> FoundAssetPackageNames;
 
+	auto IsInRequestedDir = [&Context](const FAssetData& AssetData)
+	{
+		TStringBuilder<128> PackageNameStr;
+		AssetData.PackageName.ToString(PackageNameStr);
+		FStringView PackageName(PackageNameStr.ToString(), PackageNameStr.Len());
+
+		for (const FString& RequestedPackageDir : Context.PackageDirs)
+		{
+			if (FPathViews::IsParentPathOf(RequestedPackageDir, PackageName))
+			{
+				return true;
+			}
+		}
+		return false;
+	};
 	auto AssetsFoundCallback =
-		[&Context, &FoundAssetPackageNames, &OldAssetsToRemove, this]
+		[&IsInRequestedDir, &Context, &FoundAssetPackageNames, &OldAssetsToRemove, this]
 		(const TMultiMap<FName, FAssetData*>& InFoundAssets)
 	{
 		Context.NumFoundAssets = InFoundAssets.Num();
@@ -5254,14 +5655,7 @@ void FAssetRegistryImpl::ScanPathsSynchronous(Impl::FScanPathContext& Context)
 			AssetData->PackageName.ToString(PackageNameStr);
 			FStringView PackageName(PackageNameStr.ToString(), PackageNameStr.Len());
 
-			for (const FString& RequestedPackageDir : Context.PackageDirs)
-			{
-				if (FPathViews::IsParentPathOf(RequestedPackageDir, PackageName))
-				{
-					bIsInRequestedPaths = true;
-					break;
-				}
-			}
+			bIsInRequestedPaths = IsInRequestedDir(*AssetData);
 
 			if (!bIsInRequestedPaths)
 			{
@@ -5293,11 +5687,60 @@ void FAssetRegistryImpl::ScanPathsSynchronous(Impl::FScanPathContext& Context)
 		}
 	};
 
-	bool bUnusedInterrupted;
-	Context.Status = TickGatherer(Context.EventContext, Context.InheritanceContext, -1., bUnusedInterrupted,
+	Impl::FInterruptionContext InterruptionContext(-1., -1.);
+	Context.Status = TickGatherer(Context.EventContext, Context.InheritanceContext,	InterruptionContext,
 		FAssetsFoundCallback(AssetsFoundCallback));
+
+	// Temporary hack/partial solution. The expectation is that this function will return cause all assets
+	// under the specified directories to be ingested into the registry. However, one of the early steps
+	// in ingestion is an attempt to PostLoadAssetRegistryTags. This step requires that we already have loaded
+	// the AssetClass UClass for an asset. That may not have happened yet. In the past, we would just have skipped
+	// over that step and continued, but now we defer the asset for processing at a later time. However, that means
+	// that after running TickGatherer, even without timeslicing, our end state might be that only some assets have
+	// been scanned and others have been deferred and so would be unavailable to subsequent queries. Ideally we would
+	// solve this by loading the classes that these assets depend on. Instead, we are deferring that task and for now
+	// we manually identify any deferred assets that fall under the paths we are scanning and ask the asset registry
+	// to process them ignoring any failures of TryPostLoadAssetRegistryTags. We then run a second full Tick to finish
+	// out their processing. See UE-210249 for the desired fix.
+
+	{
+		// Find any assets that were deferred but fall into the paths we are interested in. Extract them from 
+		// the DeferredAssets and DeferredAssetsForGameThread containers
+
+		TMultiMap<FName, TUniquePtr<FAssetData>> CollectedDeferredAssets;
+		for (auto Iter = DeferredAssets.CreateIterator(); Iter; ++Iter)
+		{
+			if (IsInRequestedDir(*Iter.Value()))
+			{
+				CollectedDeferredAssets.Add(MoveTemp(*Iter));
+				Iter.RemoveCurrent();
+			}
+		}
+		for (auto Iter = DeferredAssetsForGameThread.CreateIterator(); Iter; ++Iter)
+		{
+			if (IsInRequestedDir(*Iter.Value()))
+			{
+				CollectedDeferredAssets.Add(MoveTemp(*Iter));
+				Iter.RemoveCurrent();
+			}
+		}
+		// Force AssetSearchDataGathered to process these assets, skipping the PostLoadAssetRegistryTags if needed
+		const bool bOldForceCompletionEvenIfPostLoadsFail = bForceCompletionEvenIfPostLoadsFail;
+		bForceCompletionEvenIfPostLoadsFail = true;
+		// We don't call the AssetsFoundCallback here because even for deferred assets it will already have been called.
+		// We pass DeferredAssetsForGameThread as the OutDeferred parameter, but the expectation is that nothing will be deferred
+		AssetSearchDataGathered(Context.EventContext, CollectedDeferredAssets, DeferredAssetsForGameThread, InterruptionContext);
+		// All of the assets we collected should have been processed.
+		ensure(CollectedDeferredAssets.Num() == 0);
+		bForceCompletionEvenIfPostLoadsFail = bOldForceCompletionEvenIfPostLoadsFail;
+		// Now run one more tick to perform any subsequent processing required for these assets beyond AssetSearchDataGathered
+		Context.Status = TickGatherer(Context.EventContext, Context.InheritanceContext, InterruptionContext,
+			FAssetsFoundCallback(AssetsFoundCallback));
+	}
+
 #if WITH_EDITOR
-	LoadCalculatedDependencies(&FoundAssetPackageNames, -1., Context.InheritanceContext, bUnusedInterrupted);
+	LoadCalculatedDependencies(&FoundAssetPackageNames, Context.InheritanceContext, &PackagesNeedingDependencyCalculation, InterruptionContext);
+	LoadCalculatedDependencies(&FoundAssetPackageNames, Context.InheritanceContext, &PackagesNeedingDependencyCalculationOnGameThread, InterruptionContext);
 #endif
 	for (FSoftObjectPath& OldAssetToRemove : OldAssetsToRemove)
 	{
@@ -5337,13 +5780,15 @@ bool IsPathMounted(const FString& Path, const TSet<FString>& MountPointsNoTraili
 }
 
 #if WITH_EDITOR
-void FAssetRegistryImpl::PostLoadAssetRegistryTags(FAssetData* AssetData)
+bool FAssetRegistryImpl::TryPostLoadAssetRegistryTags(FAssetData* AssetData)
 {
 	check(AssetData);
+	bool CouldPostLoadAssetRegistryTags = true;
 	if (AssetData->TagsAndValues.Num())
 	{
 		FTopLevelAssetPath AssetClassPath = AssetData->AssetClassPath;
 		UClass* AssetClass = FindObject<UClass>(AssetClassPath, true);
+
 		while (!AssetClass)
 		{
 			// this is probably a blueprint that has not yet been loaded, try to find its native base class
@@ -5355,16 +5800,82 @@ void FAssetRegistryImpl::PostLoadAssetRegistryTags(FAssetData* AssetData)
 			}
 			else
 			{
-				break;
+				FTopLevelAssetPath LastAssetClassPath = AssetClassPath;
+				// Maybe it's a redirector
+				FSoftObjectPath RedirectedPath = GRedirectCollector.GetAssetPathRedirection(FSoftObjectPath(AssetClassPath));
+				if (RedirectedPath.IsValid())
+				{
+					AssetClassPath = RedirectedPath.GetAssetPath();
+				}
+				else
+				{
+					FCoreRedirectObjectName NewName = FCoreRedirects::GetRedirectedName(ECoreRedirectFlags::Type_Class, FCoreRedirectObjectName(AssetClassPath));
+					if (NewName.IsValid())
+					{
+						AssetClassPath = NewName.ToString();
+					}
+				}
+				
+				if (AssetClassPath != LastAssetClassPath && !AssetClassPath.IsNull())
+				{
+					AssetClass = FindObject<UClass>(AssetClassPath, true);
+				}
+				else
+				{
+					CouldPostLoadAssetRegistryTags = false;
+					break;
+				}
 			}
 		}
 
+		// Now identify the most derived native class in the class hierarchy
+		if (AssetClass)
+		{
+			while (!AssetClass->HasAnyClassFlags(CLASS_Native))
+			{
+				AssetClass = AssetClass->GetSuperClass();
+			}
+		}
+
+		bool MakeFinalChecks = false;
+		if (bForceCompletionEvenIfPostLoadsFail && bPreloadingComplete && IsEngineStartupModuleLoadingComplete())
+		{
+			// Okay, we think we're done loading and now we need to make some expensive final checks to try to either
+			// track down the classes for fixup or just give up
+			MakeFinalChecks = true;
+		}
+		if (!AssetClass && bForceCompletionEvenIfPostLoadsFail)
+		{
+			if (MakeFinalChecks)
+			{
+				FString Reason;
+				if (AssetClassPath.ToString().StartsWith(TEXT("/Script/")))
+				{
+					Reason = TEXT("The missing class is native--perhaps a CoreRedirector is missing?");
+				}
+				else
+				{
+					if (State.GetAssetPackageData(AssetClassPath.GetPackageName()) == nullptr)
+					{
+						Reason = TEXT("The class is missing on disk or could not be loaded. Perhaps it has been deleted from perforce and the referencing object is broken?");
+					}
+				}
+				//@TODO this should become a Warning once UE-209846 is finished
+				UE_LOG(LogAssetRegistry, Verbose, TEXT("Unable to PostLoadAssetRegistryTags for '%s' because ancestor class '%s' cannot be found. %s"), 
+					*AssetData->GetObjectPathString(), *AssetClassPath.ToString(), *Reason);
+			}
+
+			// Force this so that we can move on
+			CouldPostLoadAssetRegistryTags = true;
+		}
+
 		if (AssetClass)
 		{
 			if (UObject* ClassCDO = AssetClass->GetDefaultObject(false))
 			{
 				TArray<UObject::FAssetRegistryTag> TagsToModify;
-				ClassCDO->PostLoadAssetRegistryTags(*AssetData, TagsToModify);
+				UObject::FPostLoadAssetRegistryTagsContext Context(*AssetData, TagsToModify);
+				ClassCDO->ThreadedPostLoadAssetRegistryTags(Context);
 				if (TagsToModify.Num())
 				{
 					FAssetDataTagMap TagsAndValues = AssetData->TagsAndValues.CopyMap();
@@ -5385,6 +5896,7 @@ void FAssetRegistryImpl::PostLoadAssetRegistryTags(FAssetData* AssetData)
 			}
 		}
 	}
+	return CouldPostLoadAssetRegistryTags;
 }
 #endif
 
@@ -5500,10 +6012,10 @@ bool FAssetRegistryImpl::ShouldSkipGatheredAsset(FAssetData& AssetData)
 	return false;
 }
 
-void FAssetRegistryImpl::AssetSearchDataGathered(Impl::FEventContext& EventContext, const double TickStartTime, TMultiMap<FName, FAssetData*>& AssetResults)
+void FAssetRegistryImpl::AssetSearchDataGathered(Impl::FEventContext& EventContext, TMultiMap<FName, TUniquePtr<FAssetData>>& AssetResults,
+	TMultiMap<FName, TUniquePtr<FAssetData>>& OutDeferredAssetResults, Impl::FInterruptionContext& InOutInterruptionContext)
 {
 	TRACE_CPUPROFILER_EVENT_SCOPE(AssetSearchDataGathered);
-	const bool bFlushFullBuffer = TickStartTime < 0;
 
 	// Refreshes ClassGeneratorNames if out of date due to module load
 	CollectCodeGeneratorClasses();
@@ -5518,11 +6030,32 @@ void FAssetRegistryImpl::AssetSearchDataGathered(Impl::FEventContext& EventConte
 		MountPoints.Append(MoveTemp(MountPointsArray));
 	}
 
+#if WITH_EDITOR
+	// This ensures we can search for classes inside PostLoadAssetRegistryTags. We take the lock once out here to reduce overhead
+	FGCScopeGuard Guard;
+#endif
+
+	TSet<FTopLevelAssetPath> MissingClasses;
+	bool bInterrupted = false;
+	int64 IterationCounter = 0;
+
 	// Add the found assets
-	for (TMultiMap<FName, FAssetData*>::TIterator Iter(AssetResults); Iter; ++Iter)
+	for (TMultiMap<FName, TUniquePtr<FAssetData>>::TIterator Iter(AssetResults); Iter && !bInterrupted; ++Iter)
 	{
+		ON_SCOPE_EXIT
+		{
+			// ShouldExitEarly calls FPlatformTime::Seconds which isn't super cheap
+			// Since we can spin very quickly in this loop, avoid checking every single iteration
+			if ((++IterationCounter % 10) == 0)
+			{
+				// Check to see if we have run out of time in this tick
+				bInterrupted = InOutInterruptionContext.ShouldExitEarly();
+			}
+		};
+
 		// Delete or take ownership of the BackgroundResult; it was originally new'd by an FPackageReader
-		TUniquePtr<FAssetData> BackgroundResult(Iter.Value());
+		TUniquePtr<FAssetData> BackgroundResult(MoveTemp(Iter.Value()));
+		FName BackgroundAssetPackageName = Iter.Key();
 		CA_ASSUME(BackgroundResult.Get() != nullptr);
 		Iter.RemoveCurrent();
 
@@ -5556,7 +6089,12 @@ void FAssetRegistryImpl::AssetSearchDataGathered(Impl::FEventContext& EventConte
 			if (ensure(ExistingAssetData != BackgroundResult.Get()))
 			{
 #if WITH_EDITOR
-				PostLoadAssetRegistryTags(BackgroundResult.Get());
+				bool CouldPostLoad = TryPostLoadAssetRegistryTags(BackgroundResult.Get());
+				if (!CouldPostLoad)
+				{
+					OutDeferredAssetResults.Add(BackgroundAssetPackageName, MoveTemp(BackgroundResult));
+					continue;
+				}
 				if (AssetDataObjectPathsUpdatedOnLoad.Contains(BackgroundResult->GetSoftObjectPath()))
 				{
 					// If the current AssetData came from a loaded asset, don't overwrite it with the new one from disk
@@ -5591,7 +6129,12 @@ void FAssetRegistryImpl::AssetSearchDataGathered(Impl::FEventContext& EventConte
 #endif
 
 #if WITH_EDITOR
-				PostLoadAssetRegistryTags(BackgroundResult.Get());
+				bool CouldPostLoad = TryPostLoadAssetRegistryTags(BackgroundResult.Get());
+				if (!CouldPostLoad)
+				{
+					OutDeferredAssetResults.Add(BackgroundAssetPackageName, MoveTemp(BackgroundResult));
+					continue;
+				}
 #endif
 				AddAssetData(EventContext, BackgroundResult.Release());
 			}
@@ -5606,19 +6149,12 @@ void FAssetRegistryImpl::AssetSearchDataGathered(Impl::FEventContext& EventConte
 		{
 			UE_LOG(LogAssetRegistry, Warning, TEXT("AssetRegistry: An asset has been loaded with an invalid mount point: '%s', Mount Point: '%s'"), *BackgroundResult->GetObjectPathString(), *PackagePathString)
 		}
-
-		// Check to see if we have run out of time in this tick
-		if (!bFlushFullBuffer && (FPlatformTime::Seconds() - TickStartTime) > Impl::MaxSecondsPerFrame)
-		{
-			return;
-		}
 	}
 }
 
-void FAssetRegistryImpl::PathDataGathered(Impl::FEventContext& EventContext, const double TickStartTime, TRingBuffer<FString>& PathResults)
+void FAssetRegistryImpl::PathDataGathered(Impl::FEventContext& EventContext, TRingBuffer<FString>& PathResults, Impl::FInterruptionContext& InOutInterruptionContext)
 {
 	TRACE_CPUPROFILER_EVENT_SCOPE(PathDataGathered);
-	const bool bFlushFullBuffer = TickStartTime < 0;
 
 	TSet<FString> MountPoints;
 	FString PackageRoot;
@@ -5646,22 +6182,46 @@ void FAssetRegistryImpl::PathDataGathered(Impl::FEventContext& EventContext, con
 		}
 
 		// Check to see if we have run out of time in this tick
-		if (!bFlushFullBuffer && (FPlatformTime::Seconds() - TickStartTime) > Impl::MaxSecondsPerFrame)
+		if (InOutInterruptionContext.ShouldExitEarly())
 		{
 			return;
 		}
 	}
 }
 
-void FAssetRegistryImpl::DependencyDataGathered(const double TickStartTime, TMultiMap<FName, FPackageDependencyData>& DependsResults)
+void FAssetRegistryImpl::DependencyDataGathered(TMultiMap<FName, FPackageDependencyData>& DependsResults, 
+	TMultiMap<FName, FPackageDependencyData>& OutDeferredDependencyResults, 
+	TSet<FName>* OutPackagesNeedingDependencyCalculation, Impl::FInterruptionContext& InOutInterruptionContext)
 {
 	TRACE_CPUPROFILER_EVENT_SCOPE(DependencyDataGathered);
 	using namespace UE::AssetRegistry;
-	const bool bFlushFullBuffer = TickStartTime < 0;
+
+	// This ensures we can call FindPackage below from a background thread
+	FGCScopeGuard Guard;
 
 	TMap<FName, FName> CachedDepToRedirect;
-	for (TMultiMap<FName, FPackageDependencyData>::TIterator Iter(DependsResults); Iter; ++Iter)
+	bool bInterrupted = false;
+	int64 IterationCounter = 0;
+	for (TMultiMap<FName, FPackageDependencyData>::TIterator Iter(DependsResults); Iter && !bInterrupted; ++Iter)
 	{
+		ON_SCOPE_EXIT
+		{
+			// ShouldExitEarly calls FPlatformTime::Seconds which isn't super cheap
+			// Since we can spin very quickly in this loop, avoid checking every single iteration
+			if ((++IterationCounter % 10) == 0)
+			{
+				// Check to see if we have run out of time in this tick
+				bInterrupted = InOutInterruptionContext.ShouldExitEarly();
+			}
+		};
+
+		if (DeferredAssets.Contains(Iter.Key()) || DeferredAssetsForGameThread.Contains(Iter.Key()))
+		{
+			OutDeferredDependencyResults.Add(MoveTemp(*Iter));
+			Iter.RemoveCurrent();
+			// Not ready to process this package yet
+			continue;
+		}
 		FPackageDependencyData Result = MoveTemp(Iter.Value());
 		Iter.RemoveCurrent();
 
@@ -5677,7 +6237,7 @@ void FAssetRegistryImpl::DependencyDataGathered(const double TickStartTime, TMul
 		{
 			FDependsNode* Node = State.CreateOrFindDependsNode(Result.PackageName);
 #if WITH_EDITOR
-			PackagesNeedingDependencyCalculation.Add(Result.PackageName);
+			OutPackagesNeedingDependencyCalculation->Add(Result.PackageName);
 #endif
 
 			// We will populate the node dependencies below. Empty the set here in case this file was already read
@@ -5775,20 +6335,13 @@ void FAssetRegistryImpl::DependencyDataGathered(const double TickStartTime, TMul
 			}
 			Node->SetIsDependenciesInitialized(true);
 		}
-
-		// Check to see if we have run out of time in this tick
-		if (!bFlushFullBuffer && (FPlatformTime::Seconds() - TickStartTime) > Impl::MaxSecondsPerFrame)
-		{
-			return;
-		}
 	}
 }
 
 void FAssetRegistryImpl::CookedPackageNamesWithoutAssetDataGathered(Impl::FEventContext& EventContext,
-	const double TickStartTime, TRingBuffer<FString>& CookedPackageNamesWithoutAssetDataResults, bool& bOutInterrupted)
+	TRingBuffer<FString>& CookedPackageNamesWithoutAssetDataResults, Impl::FInterruptionContext& InOutInterruptionContext)
 {
 	TRACE_CPUPROFILER_EVENT_SCOPE(CookedPackageNamesWithoutAssetDataGathered);
-	bOutInterrupted = false;
 
 	struct FConfigValue
 	{
@@ -5813,10 +6366,14 @@ void FAssetRegistryImpl::CookedPackageNamesWithoutAssetDataGathered(Impl::FEvent
 			// Loading the entire package will make all of its assets searchable through the in-memory scanning performed by GetAssets
 			EventContext.RequiredLoads.Add(CookedPackageNamesWithoutAssetDataResults.PopFrontValue());
 		}
-		if (TickStartTime >= 0)
+		// Avoid marking the scan complete before we have loaded all the relevant assets. By interrupting here
+		// we intend to ensure that the event context is processed, triggering a LoadPackage, and then a ProcessLoadedAssetsToUpdateCache,
+		// and only then resume scanning from disk. However, in the current multithreaded implementation this is not guaranteed
+		// as only the main thread broadcasts events but the background thread might come around for another time slice before
+		// the main thread does so.  UE-209843
+		if (InOutInterruptionContext.IsTimeSlicingEnabled())
 		{
-			// If the tick is time-limited, signal an interruption now to process the loads before proceeding with any other tick steps
-			bOutInterrupted = true;
+			InOutInterruptionContext.RequestEarlyExit();
 			return;
 		}
 	}
@@ -5829,9 +6386,11 @@ void FAssetRegistryImpl::CookedPackageNamesWithoutAssetDataGathered(Impl::FEvent
 	}
 }
 
-void FAssetRegistryImpl::VerseFilesGathered(Impl::FEventContext& EventContext, const double TickStartTime, TRingBuffer<FName>& VerseResults)
+void FAssetRegistryImpl::VerseFilesGathered(Impl::FEventContext& EventContext, TRingBuffer<FName>& VerseResults, 
+	Impl::FInterruptionContext& InOutInterruptionContext)
 {
 	TRACE_CPUPROFILER_EVENT_SCOPE(VerseFilesGathered);
+
 	while (VerseResults.Num() > 0)
 	{
 		FName VerseFilePath = VerseResults.PopFrontValue();
@@ -5839,7 +6398,7 @@ void FAssetRegistryImpl::VerseFilesGathered(Impl::FEventContext& EventContext, c
 		AddVerseFile(EventContext, VerseFilePath);
 
 		// Check to see if we have run out of time in this tick
-		if (TickStartTime >= 0 && (FPlatformTime::Seconds() - TickStartTime) > Impl::MaxSecondsPerFrame)
+		if (InOutInterruptionContext.ShouldExitEarly())
 		{
 			return;
 		}
@@ -5891,6 +6450,23 @@ bool FAssetRegistryImpl::RemoveAssetPath(Impl::FEventContext& EventContext, FNam
 
 void FAssetRegistryImpl::AddAssetData(Impl::FEventContext& EventContext, FAssetData* AssetData)
 {
+	// Make sure to consider redirections!
+#if WITH_EDITOR
+	{
+		if (AssetData->IsRedirector())
+		{
+			const FName DestinationObjectFName(TEXT("DestinationObject"));
+			FString RedirectDestinationString;
+			AssetData->GetTagValue(DestinationObjectFName, RedirectDestinationString);
+			FSoftObjectPath RedirectDestination = RedirectDestinationString;
+			if (!RedirectDestination.IsNull())
+			{
+				GRedirectCollector.AddAssetPathRedirection(AssetData->GetSoftObjectPath(), RedirectDestination);
+			}
+		}
+	}
+#endif
+
 	State.AddAssetData(AssetData);
 
 	if (!ShouldSkipAsset(AssetData->AssetClassPath, AssetData->PackageFlags))
@@ -6181,7 +6757,7 @@ void UAssetRegistryImpl::OnDirectoryChanged(const TArray<FFileChangeData>& FileC
 	bool bInitialSearchCompleted;
 	{
 		LLM_SCOPE(ELLMTag::AssetRegistry);
-		FWriteScopeLock InterfaceScopeLock(InterfaceLock);
+		UE::AssetRegistry::FInterfaceWriteScopeLock InterfaceScopeLock(InterfaceLock);
 		bInitialSearchStarted = GuardedData.IsInitialSearchStarted();
 		bInitialSearchCompleted = GuardedData.IsInitialSearchCompleted();
 		UE::AssetRegistry::Impl::FClassInheritanceContext InheritanceContext;
@@ -6512,12 +7088,12 @@ void FAssetRegistryImpl::OnDirectoryRescanRequired(Impl::FEventContext& EventCon
 void UAssetRegistryImpl::OnAssetLoaded(UObject *AssetLoaded)
 {
 	LLM_SCOPE(ELLMTag::AssetRegistry);
-	FWriteScopeLock InterfaceScopeLock(InterfaceLock);
+	UE::AssetRegistry::FInterfaceWriteScopeLock InterfaceScopeLock(InterfaceLock);
 	GuardedData.AddLoadedAssetToProcess(*AssetLoaded);
 }
 
 void UAssetRegistryImpl::ProcessLoadedAssetsToUpdateCache(UE::AssetRegistry::Impl::FEventContext& EventContext,
-	const double TickStartTime, UE::AssetRegistry::Impl::EGatherStatus Status)
+	UE::AssetRegistry::Impl::EGatherStatus Status, UE::AssetRegistry::Impl::FInterruptionContext& InOutInterruptionContext)
 {
 	// Note this function can be reentered due to arbitrary code execution in construction of FAssetData
 	if (!IsInGameThread())
@@ -6527,8 +7103,7 @@ void UAssetRegistryImpl::ProcessLoadedAssetsToUpdateCache(UE::AssetRegistry::Imp
 	}
 
 	// Early exit to save cputime if we're still processing cache data
-	const bool bFlushFullBuffer = TickStartTime < 0;
-	if (Status == UE::AssetRegistry::Impl::EGatherStatus::Active && !bFlushFullBuffer)
+	if (IsTickActive(Status) && InOutInterruptionContext.IsTimeSlicingEnabled())
 	{
 		return;
 	}
@@ -6539,7 +7114,7 @@ void UAssetRegistryImpl::ProcessLoadedAssetsToUpdateCache(UE::AssetRegistry::Imp
 
 	{
 		LLM_SCOPE(ELLMTag::AssetRegistry);
-		FWriteScopeLock InterfaceScopeLock(InterfaceLock);
+		UE::AssetRegistry::FInterfaceWriteScopeLock InterfaceScopeLock(InterfaceLock);
 		GuardedData.GetProcessLoadedAssetsBatch(BatchObjects, BatchSize, bUpdateDiskCacheAfterLoad);
 		if (BatchObjects.Num() == 0)
 		{
@@ -6568,8 +7143,7 @@ void UAssetRegistryImpl::ProcessLoadedAssetsToUpdateCache(UE::AssetRegistry::Imp
 				EAssetRegistryTagsCaller::AssetRegistryLoad));
 
 			// Check to see if we have run out of time in this tick
-			if (!bFlushFullBuffer &&
-				(FPlatformTime::Seconds() - TickStartTime) > UE::AssetRegistry::Impl::MaxSecondsPerFrame)
+			if (InOutInterruptionContext.ShouldExitEarly())
 			{
 				bTimedOut = true;
 				break;
@@ -6577,7 +7151,7 @@ void UAssetRegistryImpl::ProcessLoadedAssetsToUpdateCache(UE::AssetRegistry::Imp
 		}
 
 		LLM_SCOPE(ELLMTag::AssetRegistry);
-		FWriteScopeLock InterfaceScopeLock(InterfaceLock);
+		UE::AssetRegistry::FInterfaceWriteScopeLock InterfaceScopeLock(InterfaceLock);
 		GuardedData.PushProcessLoadedAssetsBatch(EventContext, BatchAssetDatas,
 			TArrayView<const UObject*>(BatchObjects).Slice(Index, CurrentBatchSize-Index));
 		if (bTimedOut)
@@ -6702,7 +7276,7 @@ void UAssetRegistryImpl::ScanModifiedAssetFiles(const TArray<FString>& InFilePat
 	UE::AssetRegistry::Impl::FEventContext EventContext;
 	{
 		LLM_SCOPE(ELLMTag::AssetRegistry);
-		FWriteScopeLock InterfaceScopeLock(InterfaceLock);
+		UE::AssetRegistry::FInterfaceWriteScopeLock InterfaceScopeLock(InterfaceLock);
 		UE::AssetRegistry::Impl::FClassInheritanceContext InheritanceContext;
 		UE::AssetRegistry::Impl::FClassInheritanceBuffer InheritanceBuffer;
 		GetInheritanceContextWithRequiredLock(InterfaceScopeLock, InheritanceContext, InheritanceBuffer);
@@ -6714,7 +7288,8 @@ void UAssetRegistryImpl::ScanModifiedAssetFiles(const TArray<FString>& InFilePat
 	// but in-memory results will override the on-disk results we just scanned,
 	// and our in-memory results might be out of date due to being queued but not yet processed.
 	// So ProcessLoadedAssetsToUpdateCache before returning to make sure results are up to date.
-	ProcessLoadedAssetsToUpdateCache(EventContext, -1., UE::AssetRegistry::Impl::EGatherStatus::Complete);
+	UE::AssetRegistry::Impl::FInterruptionContext InterruptionContext(-1., -1.);
+	ProcessLoadedAssetsToUpdateCache(EventContext, UE::AssetRegistry::Impl::EGatherStatus::Complete, InterruptionContext);
 #endif
 
 	Broadcast(EventContext);
@@ -6818,7 +7393,7 @@ void UAssetRegistryImpl::OnContentPathMounted(const FString& InAssetPath, const
 	UE::AssetRegistry::Impl::FEventContext EventContext;
 	{
 		LLM_SCOPE(ELLMTag::AssetRegistry);
-		FWriteScopeLock InterfaceScopeLock(InterfaceLock);
+		UE::AssetRegistry::FInterfaceWriteScopeLock InterfaceScopeLock(InterfaceLock);
 		UE::AssetRegistry::Impl::FClassInheritanceContext InheritanceContext;
 		UE::AssetRegistry::Impl::FClassInheritanceBuffer InheritanceBuffer;
 		GetInheritanceContextWithRequiredLock(InterfaceScopeLock, InheritanceContext, InheritanceBuffer);
@@ -6898,7 +7473,7 @@ void UAssetRegistryImpl::OnContentPathDismounted(const FString& InAssetPath, con
 	UE::AssetRegistry::Impl::FEventContext EventContext;
 	{
 		LLM_SCOPE(ELLMTag::AssetRegistry);
-		FWriteScopeLock InterfaceScopeLock(InterfaceLock);
+		UE::AssetRegistry::FInterfaceWriteScopeLock InterfaceScopeLock(InterfaceLock);
 		GuardedData.OnContentPathDismounted(EventContext, InAssetPath, AssetPathNoTrailingSlash, FileSystemPath);
 
 		// Stop listening for directory changes in this content path
@@ -7009,7 +7584,7 @@ void UAssetRegistryImpl::SetTemporaryCachingMode(bool bEnable)
 {
 	checkf(IsInGameThread(), TEXT("Changing Caching mode is only available on the game thread because it affects behavior on all threads"));
 	LLM_SCOPE(ELLMTag::AssetRegistry);
-	FWriteScopeLock InterfaceScopeLock(InterfaceLock);
+	UE::AssetRegistry::FInterfaceWriteScopeLock InterfaceScopeLock(InterfaceLock);
 	GuardedData.SetTemporaryCachingMode(bEnable);
 }
 
@@ -7037,7 +7612,7 @@ void UAssetRegistryImpl::SetTemporaryCachingModeInvalidated()
 {
 	checkf(IsInGameThread(), TEXT("Invalidating temporary cache is only available on the game thread because it affects behavior on all threads"));
 	LLM_SCOPE(ELLMTag::AssetRegistry);
-	FWriteScopeLock InterfaceScopeLock(InterfaceLock);
+	UE::AssetRegistry::FInterfaceWriteScopeLock InterfaceScopeLock(InterfaceLock);
 	GuardedData.SetTemporaryCachingModeInvalidated();
 }
 
@@ -7053,7 +7628,7 @@ void FAssetRegistryImpl::SetTemporaryCachingModeInvalidated()
 
 bool UAssetRegistryImpl::GetTemporaryCachingMode() const
 {
-	FReadScopeLock InterfaceScopeLock(InterfaceLock);
+	UE::AssetRegistry::FInterfaceReadScopeLock InterfaceScopeLock(InterfaceLock);
 	return GuardedData.IsTempCachingEnabled();
 }
 
@@ -7160,7 +7735,7 @@ void FAssetRegistryImpl::UpdateInheritanceBuffer(Impl::FClassInheritanceBuffer&
 
 }
 
-void UAssetRegistryImpl::GetInheritanceContextWithRequiredLock(FRWScopeLock& InOutScopeLock,
+void UAssetRegistryImpl::GetInheritanceContextWithRequiredLock(UE::AssetRegistry::FInterfaceRWScopeLock& InOutScopeLock,
 	UE::AssetRegistry::Impl::FClassInheritanceContext& InheritanceContext,
 	UE::AssetRegistry::Impl::FClassInheritanceBuffer& StackBuffer)
 {
@@ -7191,7 +7766,7 @@ void UAssetRegistryImpl::GetInheritanceContextWithRequiredLock(FRWScopeLock& InO
 		InheritanceContext, StackBuffer);
 }
 
-void UAssetRegistryImpl::GetInheritanceContextWithRequiredLock(FWriteScopeLock& InOutScopeLock,
+void UAssetRegistryImpl::GetInheritanceContextWithRequiredLock(UE::AssetRegistry::FInterfaceWriteScopeLock& InOutScopeLock,
 	UE::AssetRegistry::Impl::FClassInheritanceContext& InheritanceContext,
 	UE::AssetRegistry::Impl::FClassInheritanceBuffer& StackBuffer)
 {
@@ -7265,6 +7840,20 @@ bool UAssetRegistryImpl::IsDirAlreadyWatchedByRootWatchers(const FString& Direct
 
 #endif
 
+void UAssetRegistryImpl::RequestPauseBackgroundProcessing()
+{
+#if WITH_EDITOR
+	GuardedData.RequestPauseBackgroundProcessing();
+#endif
+}
+
+void UAssetRegistryImpl::RequestResumeBackgroundProcessing()
+{
+#if WITH_EDITOR
+	GuardedData.RequestResumeBackgroundProcessing();
+#endif
+}
+
 namespace UE::AssetRegistry
 {
 
@@ -7363,6 +7952,34 @@ void FAssetRegistryImpl::GetSubClasses_Recursive(Impl::FClassInheritanceContext&
 	}
 }
 
+
+#if WITH_EDITOR
+void FAssetRegistryImpl::RequestPauseBackgroundProcessing()
+{
+	if (GlobalGatherer.IsValid())
+	{
+		GlobalGatherer->PauseProcessing();
+	}
+}
+
+void FAssetRegistryImpl::RequestResumeBackgroundProcessing()
+{
+	if (GlobalGatherer.IsValid())
+	{
+		GlobalGatherer->ResumeProcessing();
+	}
+}
+
+bool FAssetRegistryImpl::IsBackgroundProcessingPaused() const
+{
+	if (GlobalGatherer.IsValid())
+	{
+		return GlobalGatherer->IsProcessingPauseRequested();
+	}
+	return true;
+}
+#endif
+
 }
 
 #if WITH_EDITOR
@@ -7405,7 +8022,7 @@ void UAssetRegistryImpl::SetManageReferences(const TMultiMap<FAssetIdentifier, F
 	// For performance reasons we call the ShouldSetManager callback when inside the lock. Licensee UAssetManagers
 	// are responsible for not calling AssetRegistry functions from ShouldSetManager as that would create a deadlock
 	LLM_SCOPE(ELLMTag::AssetRegistry);
-	FWriteScopeLock InterfaceScopeLock(InterfaceLock);
+	UE::AssetRegistry::FInterfaceWriteScopeLock InterfaceScopeLock(InterfaceLock);
 	GuardedData.SetManageReferences(ManagerMap, bClearExisting, RecurseType, ExistingManagedNodes, ShouldSetManager);
 }
 
@@ -7592,7 +8209,7 @@ bool UAssetRegistryImpl::SetPrimaryAssetIdForObjectPath(const FSoftObjectPath& O
 	bool bResult;
 	{
 		LLM_SCOPE(ELLMTag::AssetRegistry);
-		FWriteScopeLock InterfaceScopeLock(InterfaceLock);
+		UE::AssetRegistry::FInterfaceWriteScopeLock InterfaceScopeLock(InterfaceLock);
 		bResult = GuardedData.SetPrimaryAssetIdForObjectPath(EventContext, ObjectPath, PrimaryAssetId);
 	}
 	Broadcast(EventContext);
@@ -7739,6 +8356,63 @@ bool FAssetRegistryImpl::ShouldSkipAsset(const UObject* InAsset) const
 #endif
 }
 
+#if WITH_EDITOR
+void FAssetRegistryImpl::PruneAndCoalescePackagesRequiringDependencyCalculation(TSet<FName>& BackgroundPackages,
+	TSet<FName>& GameThreadPackages, Impl::FInterruptionContext& InOutInterruptionContext)
+{
+	RebuildAssetDependencyGathererMapIfNeeded();
+
+	FReadScopeLock GathererClassScopeLock(RegisteredDependencyGathererClassesLock);
+
+	// In many cases, this loop will be tight. If so, we don't want to spend a bunch of time checking whether we've
+	// run out of processing time. So only check every N iterations.
+	uint64 IterationCounter = 0;
+	auto ProcessSet = [&InOutInterruptionContext, &IterationCounter, this](TSet<FName>& SourceSet,
+		TSet<FName>* OptDestinationSet)->void
+		{
+			for (auto Iter = SourceSet.CreateIterator(); Iter; ++Iter)
+			{
+				bool HasAnyRegisteredDependencyGatherers = false;
+				for (const FAssetData* AssetData : State.GetAssetsByPackageName(*Iter))
+				{
+					if (RegisteredDependencyGathererClasses.Contains(AssetData->AssetClassPath))
+					{
+						HasAnyRegisteredDependencyGatherers = true;
+						break;
+					}
+				}
+
+				// If we need to process this asset and we have a destination set, move it there
+				if ((OptDestinationSet != nullptr) && HasAnyRegisteredDependencyGatherers)
+				{
+					OptDestinationSet->Add(*Iter);
+					Iter.RemoveCurrent();
+				}
+				else if (!HasAnyRegisteredDependencyGatherers)
+				{
+					// If we don't have to process this asset, remove it from whichever list it is in
+					Iter.RemoveCurrent();
+				}
+
+				if ((++IterationCounter % 50) == 0)
+				{
+					if (InOutInterruptionContext.ShouldExitEarly())
+					{
+						return;
+					}
+				}
+			}
+		};
+
+	ProcessSet(GameThreadPackages, nullptr);
+	if (InOutInterruptionContext.ShouldExitEarly())
+	{
+		return;
+	}
+	ProcessSet(BackgroundPackages, &GameThreadPackages);
+}
+#endif 
+
 namespace Impl
 {
 
@@ -7785,7 +8459,7 @@ void FEventContext::Append(FEventContext&& Other)
 
 void UAssetRegistryImpl::ReadLockEnumerateTagToAssetDatas(TFunctionRef<void(FName TagName, const TArray<const FAssetData*>& Assets)> Callback) const
 {
-	FReadScopeLock InterfaceScopeLock(InterfaceLock);
+	UE::AssetRegistry::FInterfaceReadScopeLock InterfaceScopeLock(InterfaceLock);
 	for (const TPair<FName, const TArray<const FAssetData*>>& Pair : GuardedData.GetState().GetTagToAssetDatasMap())
 	{
 		Callback(Pair.Key, Pair.Value);
@@ -7800,7 +8474,7 @@ void UAssetRegistryImpl::Broadcast(UE::AssetRegistry::Impl::FEventContext& Event
 		// By contract events (and PackageLoads) can only be sent on the game thread; some legacy systems depend on 
 		// this and are not threadsafe. If we're not in the game thread, defer all events in the EventContext
 		// instead of broadcasting them on this thread
-		FWriteScopeLock InterfaceScopeLock(InterfaceLock);
+		UE::AssetRegistry::FInterfaceWriteScopeLock InterfaceScopeLock(InterfaceLock);
 		check(&EventContext != &DeferredEvents); // Only the GameThread should be calling Broadcast on DeferredEvents
 		DeferredEvents.Append(MoveTemp(EventContext));
 		EventContext.Clear();
@@ -7964,7 +8638,7 @@ void UAssetRegistryImpl::Broadcast(UE::AssetRegistry::Impl::FEventContext& Event
 	{
 		if (&EventContext != &DeferredEvents)
 		{
-			FWriteScopeLock InterfaceScopeLock(InterfaceLock);
+			UE::AssetRegistry::FInterfaceWriteScopeLock InterfaceScopeLock(InterfaceLock);
 			// Do not send the file loaded event yet if there are still deferred events and pass the flag on instead
 			if (!DeferredEvents.IsEmpty())
 			{
@@ -8080,7 +8754,6 @@ UAssetRegistryImpl::FFileLoadProgressUpdatedEvent& UAssetRegistryImpl::OnFileLoa
 	return FileLoadProgressUpdatedEvent;
 }
 
-
 namespace UE::AssetRegistry
 {
 const FAssetData* GetMostImportantAsset(TConstArrayView<const FAssetData*> PackageAssetDatas, EGetMostImportantAssetFlags InFlags)
@@ -8252,4 +8925,21 @@ bool ShouldSearchAllAssetsAtStart()
 	return bSearchAllAssetsAtStart;
 }
 
+namespace Impl
+{
+
+bool FInterruptionContext::ShouldExitEarly()
+{
+	if (EarlyExitCallback && EarlyExitCallback())
+	{
+		OutInterrupted = true;
+	}
+	else if (TickStartTime > 0 && ((FPlatformTime::Seconds() - TickStartTime) > MaxRunningTime))
+	{
+		OutInterrupted = true;
+	}
+	return OutInterrupted;
+}
+
+}
 } // namespace AssetRegistry
diff --git a/Engine/Source/Runtime/AssetRegistry/Private/AssetRegistry.h b/Engine/Source/Runtime/AssetRegistry/Private/AssetRegistry.h
--- a/Engine/Source/Runtime/AssetRegistry/Private/AssetRegistry.h
+++ b/Engine/Source/Runtime/AssetRegistry/Private/AssetRegistry.h
@@ -7,7 +7,6 @@
 
 #include "AssetRegistry.generated.h"
 
-class FRWScopeLock;
 namespace UE::AssetRegistry::Premade { struct FAsyncConsumer; }
 
 /**
@@ -201,6 +200,8 @@ public:
 
 	virtual bool IsPathBeautificationNeeded(const FString& InAssetPath) const override;
 
+	UE::AssetRegistry::Impl::EGatherStatus TickOnBackgroundThread();
+
 protected:
 	virtual void SetManageReferences(const TMultiMap<FAssetIdentifier, FAssetIdentifier>& ManagerMap, bool bClearExisting, UE::AssetRegistry::EDependencyCategory RecurseType, TSet<FDependsNode*>& ExistingManagedNodes, ShouldSetManagerPredicate ShouldSetManager = nullptr) override;
 	virtual bool SetPrimaryAssetIdForObjectPath(const FSoftObjectPath& ObjectPath, FPrimaryAssetId PrimaryAssetId) override;
@@ -256,7 +257,7 @@ private:
 #if WITH_EDITOR
 	/** Create FAssetData from any loaded UObject assets and store the updated AssetData in the state */
 	void ProcessLoadedAssetsToUpdateCache(UE::AssetRegistry::Impl::FEventContext& EventContext,
-		const double TickStartTime, UE::AssetRegistry::Impl::EGatherStatus Status);
+		UE::AssetRegistry::Impl::EGatherStatus Status, UE::AssetRegistry::Impl::FInterruptionContext& InOutInterruptionContext);
 #endif
 	/**
 	 * Remain under the given lock and return an InheritanceContext based on the appropriate choice of the persistent
@@ -264,10 +265,10 @@ private:
 	 * before being used. If the buffer needs to be updated and its the persistent buffer (which is protected data),
 	 * convert the given lock to a write lock if not one already.
 	 */
-	void GetInheritanceContextWithRequiredLock(FRWScopeLock& InOutScopeLock,
+	void GetInheritanceContextWithRequiredLock(UE::AssetRegistry::FInterfaceRWScopeLock& InOutScopeLock,
 		UE::AssetRegistry::Impl::FClassInheritanceContext& InheritanceContext,
 		UE::AssetRegistry::Impl::FClassInheritanceBuffer& StackBuffer);
-	void GetInheritanceContextWithRequiredLock(FWriteScopeLock& InOutScopeLock,
+	void GetInheritanceContextWithRequiredLock(UE::AssetRegistry::FInterfaceWriteScopeLock& InOutScopeLock,
 		UE::AssetRegistry::Impl::FClassInheritanceContext& InheritanceContext,
 		UE::AssetRegistry::Impl::FClassInheritanceBuffer& StackBuffer);
 	void GetInheritanceContextAfterVerifyingLock(uint64 CurrentGeneratorClassesVersionNumber,
@@ -290,12 +291,33 @@ private:
 	bool IsDirAlreadyWatchedByRootWatchers(const FString& Directory) const;
 #endif
 
+	/** Request to pause or resume background processing of scan results.
+	 *  This can be used to allow a priority thread to perform along sequence of operations
+	 *  without having to contend with the background thread for data access
+	 */
+	virtual void RequestPauseBackgroundProcessing();
+	virtual void RequestResumeBackgroundProcessing();
+	bool IsBackgroundProcessingPaused() const 
+	{ 
+#if WITH_EDITOR 
+		return GuardedData.IsBackgroundProcessingPaused(); 
+#else
+		return true;
+#endif
+	}
+
 private:
 
 	UE::AssetRegistry::FAssetRegistryImpl GuardedData;
 
 	/** Lock guarding the GuardedData */
-	mutable FRWLock InterfaceLock;
+	mutable UE::AssetRegistry::Private::FRWLockWithPriority InterfaceLock;
+
+	/** This lock doesn't strictly protect any data (the InterfaceLock does that). Instead, 
+	 *	it is used to let the main thread know when the gatherer thread is doing processing work
+	 *  so that the main thread does not end up blocking on the InterfaceLock in Tick(). 
+	 */
+	FCriticalSection GatheredDataProcessingLock;
 
 #if WITH_EDITOR
 	/** Handles to all registered OnDirectoryChanged delegates */
diff --git a/Engine/Source/Runtime/AssetRegistry/Private/AssetRegistryImpl.h b/Engine/Source/Runtime/AssetRegistry/Private/AssetRegistryImpl.h
--- a/Engine/Source/Runtime/AssetRegistry/Private/AssetRegistryImpl.h
+++ b/Engine/Source/Runtime/AssetRegistry/Private/AssetRegistryImpl.h
@@ -32,6 +32,10 @@ namespace UE::AssetRegistry::Premade { struct FAsyncConsumer; }
 namespace UE::AssetRegistry::Impl { struct FInitializeContext; }
 namespace UE::AssetRegistry::Impl { struct FScanPathContext; }
 namespace UE::AssetRegistry::Impl { struct FClassInheritanceContext; }
+namespace UE::AssetRegistry { template<typename TScopeLockType> class TRWScopeLockWithPriority; }
+namespace UE::AssetRegistry { using FInterfaceReadScopeLock = TRWScopeLockWithPriority<FReadScopeLock>; }
+namespace UE::AssetRegistry { using FInterfaceWriteScopeLock = TRWScopeLockWithPriority<FWriteScopeLock>; }
+namespace UE::AssetRegistry { using FInterfaceRWScopeLock = class FRWScopeLockWithPriority; }
 
 #if WITH_EDITOR
 namespace UE::AssetDependencyGatherer::Private { class FRegisteredAssetDependencyGatherer; }
@@ -50,10 +54,11 @@ struct FAsyncConsumer
 	/** Sets the Consumer into need-to-wait mode; Wait will block until Consume is called. */
 	void PrepareForConsume();
 	/**
-	 * Does not return until after the Premade AssetRegistryState has been added into the target AssetRegistry, or it has been decided not to be used.
-	 * For performance reasons, this must be called within the WriteScopeLock, and the caller must handle the possibility that it leaves and reenters the lock.
+	 * Does not return until after the Premade AssetRegistryState has been added into the target AssetRegistry, or it has 
+	 * been decided not to be used. For performance reasons, this must be called within the WriteScopeLock, and the caller 
+	 * must handle the possibility that it leaves and reenters the lock.
 	 */
-	void Wait(UAssetRegistryImpl& UARI, FWriteScopeLock& ScopeLock);
+	void Wait(UAssetRegistryImpl& UARI, FInterfaceWriteScopeLock& ScopeLock);
 	/** Callback from the async thread to consume the Premade ARState */
 	void Consume(UAssetRegistryImpl& UARI, UE::AssetRegistry::Impl::FEventContext& EventContext, ELoadResult LoadResult, FAssetRegistryState&& ARState);
 
@@ -71,6 +76,34 @@ private:
 namespace UE::AssetRegistry
 {
 
+namespace Private
+{
+	/** Custom lock class that can prioritize waiters upon request. This is used to allow threads with critical requests 
+	 *  to request that the gatherer thread pause its work and allow the higher pri threads to jump in. Note that for this
+	 *  to work correctly it must be used with the associated TPriorityScopeLock types 
+	 *  (TScopeLockWithPriority and FRWScopeLockWithPriority)
+	 */
+	class FRWLockWithPriority : public FRWLock
+	{
+	public:
+		bool HasWaiters() const { return HighPriorityWaitersCount.load(std::memory_order_relaxed) > 0; }
+	
+	private:
+		friend FInterfaceReadScopeLock;
+		friend FInterfaceWriteScopeLock;
+		friend FInterfaceRWScopeLock;
+		std::atomic<int32> HighPriorityWaitersCount;
+	};
+
+	enum ELockPriority : uint8
+	{
+		PriorityLow,
+		PriorityHigh
+	};
+
+	using FInterfaceRWLock = FRWLockWithPriority;
+}
+
 namespace Impl
 {
 	/** Container for global class inheritance data; the AssetRegistry may use a persistent buffer or a stack local buffer. */
@@ -94,14 +127,16 @@ namespace Impl
 	};
 
 	/** Status of gathering, returned from the Tick function */
-	enum EGatherStatus : uint8
+	enum class EGatherStatus : uint8
 	{
-		Active,
+		TickActiveGatherActive,
+		TickActiveGatherIdle,
+		TickGameThreadActiveGatherIdle,
 		Complete,
 		UnableToProgress,
 		WaitingForEvents,
 	};
-
+	
 	/** Affects how rules are applied to improve loading/runtime performance */
 	enum EPerformanceMode : uint8
 	{
@@ -111,8 +146,52 @@ namespace Impl
 		// Not changing, optimize for runtime queries
 		MostlyStatic,
 	};
+
+	/** Provides handling for time slicing during TickGatherer */
+	struct FInterruptionContext
+	{
+	public:
+		typedef TFunction<bool(void)> ShouldExitEarlyCallbackType;
+
+		explicit FInterruptionContext(double InTickStartTime, double InMaxRunningTime)
+			: TickStartTime(InTickStartTime)
+			, MaxRunningTime(InMaxRunningTime) {}
+
+		FInterruptionContext(double InTickStartTime, double InMaxRunningTime, ShouldExitEarlyCallbackType& Callback)
+			: TickStartTime(InTickStartTime)
+			, MaxRunningTime(InMaxRunningTime)
+			, EarlyExitCallback(Callback) {}
+
+		double GetTickStartTime() const { return TickStartTime; }
+		bool IsTimeSlicingEnabled() const { return TickStartTime > 0; }
+		bool WasInterrupted() const { return OutInterrupted; }
+		bool ShouldExitEarly();
+		void RequestEarlyExit() { OutInterrupted = true; }
+
+	private:
+		// A negative value disables time slicing
+		const double TickStartTime;
+		// The maximum time should allow before interruption. If TickStartTime is negative, this is ignored
+		const double MaxRunningTime;
+		// If provided, this is always checked
+		ShouldExitEarlyCallbackType EarlyExitCallback;
+		// True if we ran out of time, EarlyExitCallback returned true, or RequestEarlyExit() was called
+		bool OutInterrupted = false;
+	};
+}
 }
 
+/** Returns true if ANY work remains to be done. This work might require the game thread. */
+FORCEINLINE bool IsTickActive(UE::AssetRegistry::Impl::EGatherStatus Status) 
+{ 
+	return (Status == UE::AssetRegistry::Impl::EGatherStatus::TickGameThreadActiveGatherIdle) 
+		|| (Status == UE::AssetRegistry::Impl::EGatherStatus::TickActiveGatherActive)
+		|| (Status == UE::AssetRegistry::Impl::EGatherStatus::TickActiveGatherIdle);
+}
+
+namespace UE::AssetRegistry
+{
+
 /**
  * Threading helper class for UAssetRegistryImpl that holds all of the data.
  * This class is not threadsafe; all const functions must be called from within a ReadLock and all non-const 
@@ -128,7 +207,6 @@ public:
 	FAssetRegistryImpl();
 	/** Construct the AssetRegistryImpl, including initial scans if applicable. */
 	void Initialize(Impl::FInitializeContext& Context);
-	void OnEnginePreExit();
 
 	// Helpers for functions of the same name from UAssetRegistryImpl
 
@@ -195,8 +273,9 @@ public:
 	typedef TFunctionRef<void(const TMultiMap<FName, FAssetData*>&)> FAssetsFoundCallback;
 	/** Consume any results from the gatherer and return its status */
 	Impl::EGatherStatus TickGatherer(Impl::FEventContext& EventContext,
-		Impl::FClassInheritanceContext& InheritanceContext, const double TickStartTime, bool& bOutInterrupted,
+		Impl::FClassInheritanceContext& InheritanceContext, Impl::FInterruptionContext& InOutInterruptionContext,
 		TOptional<FAssetsFoundCallback> AssetsFoundCallback = TOptional<FAssetsFoundCallback>());
+
 	/** Send a log message with the search statistics. 
 	 *  StartTime is used to report wall clock search time in the case of background scan
 	 */
@@ -211,12 +290,20 @@ public:
 	void PushProcessLoadedAssetsBatch(Impl::FEventContext& EventContext,
 		TArrayView<FAssetData> LoadedAssetDatas, TArrayView<const UObject*> UnprocessedFromBatch);
 	/** Call LoadCalculatedDependencies on each Package updated after the last LoadCalculatedDependencies. */
-	void LoadCalculatedDependencies(TArray<FName>* AssetPackageNamessToCalculate, double TickStartTime, 
-		Impl::FClassInheritanceContext& InheritanceContext, bool& bOutInterrupted);
-	/**
-	 * Look for a CalculatedDependencies function registered for the asset(s) in the given package
-	 * and call that function to add calculated dependencies. Calculated dependencies are added only after
-	 * all normal dependencies gathered from the AssetRegistry data stored in the package have been loaded.
+	void LoadCalculatedDependencies(TArray<FName>* AssetPackageNamessToCalculate, Impl::FClassInheritanceContext& InheritanceContext, 
+		TSet<FName>* InPackagesNeedingDependencyCalculation, Impl::FInterruptionContext& InOutInterruptionContext);
+	/** For each package in BackgroundPackages, if there is a registered gatherer for its asset class, this function will 
+	 *  move the package to GameThreadPackages otherwise it will be removed from the set. Any entries in GameThreadPackages that
+	 *  do not have a corresponding registered gatherer will likewise be removed. The end condition for this function is that 
+	 *  BackgroundPackages is empty and GameThreadPackages contains only those packages which require LoadCalculatedDependencies to
+	 *  be called on them
+	 */
+	void PruneAndCoalescePackagesRequiringDependencyCalculation(TSet<FName>& BackgroundPackages, TSet<FName>& GameThreadPackages,
+		Impl::FInterruptionContext& InOutInterruptionContext);
+
+	/** Look for a CalculatedDependencies function registered for the asset(s) in the given package
+	 *  and call that function to add calculated dependencies. Calculated dependencies are added only after
+	 *  all normal dependencies gathered from the AssetRegistry data stored in the package have been loaded.
 	 */
 	void LoadCalculatedDependencies(FName PackageName, Impl::FClassInheritanceContext& InheritanceContext,
 		bool& bOutHadActivity);
@@ -281,10 +368,22 @@ public:
 
 	/**
 	 * Block until the Premade AssetRegistry finishes loading, if there is one still loading.
-	 * For performance reasons, this must be called within the WriteScopeLock, and the caller must handle the possibility that it leaves and reenters the lock.
+	 * For performance reasons, this must be called within the WriteScopeLock, and the caller 
+	 * must handle the possibility that it leaves and reenters the lock.
 	 */
-	void ConditionalLoadPremadeAssetRegistry(UAssetRegistryImpl& UARI, UE::AssetRegistry::Impl::FEventContext& EventContext, FWriteScopeLock& ScopeLock);
+	void ConditionalLoadPremadeAssetRegistry(UAssetRegistryImpl& UARI, 
+		UE::AssetRegistry::Impl::FEventContext& EventContext, FInterfaceWriteScopeLock& ScopeLock);
 
+#if WITH_EDITOR
+	/** Request to pause or resume background processing of scan results.
+	 *  This can be used to allow a priority thread to perform along sequence of operations
+	 *  without having to contend with the background thread for data access
+	 */
+	void RequestPauseBackgroundProcessing();
+	void RequestResumeBackgroundProcessing();
+	bool IsBackgroundProcessingPaused() const;
+
+#endif
 private:
 
 	/**
@@ -298,15 +397,19 @@ private:
 		bool& bOutRedirectorsNeedSubscribe);
 	/** If collecting dependencies, create an FAssetPackageData for every script package, to make sure dependencies can find a data for them. */
 	void ReadScriptPackages();
-	/* Construct the gatherer if it does not already exist */
-	void ConstructGatherer();
+	/* Try to construct the gatherer if it does not already exist. If that is not possible (e.g., because we are shutting down)
+	 * return false. If it was successfully created or already existed, return true.
+	 */
+	bool TryConstructGathererIfNeeded();
 	/**  Called to set up timing variables and launch the on-constructor SearchAllAssets async call */
 	void SearchAllAssetsInitialAsync(Impl::FEventContext& EventContext, Impl::FClassInheritanceContext& InheritanceContext);
 	/**
 	 * Called every tick to when data is retrieved by the background asset search.
 	 * If TickStartTime is < 0, the entire list of gathered assets will be cached. Also used in sychronous searches
+	 * The DeferredResults array contains assets that were not ready for processing due to required UClass's not being loaded
 	 */
-	void AssetSearchDataGathered(Impl::FEventContext& EventContext, const double TickStartTime, TMultiMap<FName, FAssetData*>& AssetResults);
+	void AssetSearchDataGathered(Impl::FEventContext& EventContext, TMultiMap<FName, TUniquePtr<FAssetData>>& AssetResults, 
+		TMultiMap<FName, TUniquePtr<FAssetData>>& OutDeferredResults, Impl::FInterruptionContext& InOutInterruptionContext);
 	/** Validate assets gathered from disk before adding them to the AssetRegistry. */
 	bool ShouldSkipGatheredAsset(FAssetData& AssetData);
 
@@ -314,17 +417,19 @@ private:
 	 * Called every tick when data is retrieved by the background path search.
 	 * If TickStartTime is < 0, the entire list of gathered assets will be cached. Also used in sychronous searches
 	 */
-	void PathDataGathered(Impl::FEventContext& EventContext, const double TickStartTime, TRingBuffer<FString>& PathResults);
+	void PathDataGathered(Impl::FEventContext& EventContext, TRingBuffer<FString>& PathResults, Impl::FInterruptionContext& InOutInterruptionContext);
 
 	/** Called every tick when data is retrieved by the background dependency search */
-	void DependencyDataGathered(const double TickStartTime, TMultiMap<FName, FPackageDependencyData>& DependsResults);
+	void DependencyDataGathered(TMultiMap<FName, FPackageDependencyData>& DependsResults, 
+		TMultiMap<FName, FPackageDependencyData>& OutDeferredDependencyResults,	TSet<FName>* OutPackagesNeedingDependencyCalculation, 
+		Impl::FInterruptionContext& InOutInterruptionContext);
 
 	/** Called every tick when data is retrieved by the background search for cooked packages that do not contain asset data */
-	void CookedPackageNamesWithoutAssetDataGathered(Impl::FEventContext& EventContext, const double TickStartTime,
-		TRingBuffer<FString>& CookedPackageNamesWithoutAssetDataResults, bool& bOutInterrupted);
+	void CookedPackageNamesWithoutAssetDataGathered(Impl::FEventContext& EventContext, 
+		TRingBuffer<FString>& CookedPackageNamesWithoutAssetDataResults, Impl::FInterruptionContext& InOutInterruptionContext);
 
 	/** Called every tick when data is retrieved by the background dependency search */
-	void VerseFilesGathered(Impl::FEventContext& EventContext, const double TickStartTime, TRingBuffer<FName>& VerseResults);
+	void VerseFilesGathered(Impl::FEventContext& EventContext, TRingBuffer<FName>& VerseResults, Impl::FInterruptionContext& InOutInterruptionContext);
 
 	/** Adds the asset data to the lookup maps */
 	void AddAssetData(Impl::FEventContext& EventContext, FAssetData* AssetData);
@@ -341,10 +446,11 @@ private:
 
 #if WITH_EDITOR
 	/** 
-	 * Calls PostLoadAssetRegistryTags on the CDO of the asset class this data represents 
+	 * Calls PostLoadAssetRegistryTags on the CDO of the asset class this data represents.
 	 * @param AssetData Existing asset data
+	 * @return Returns false if the required parent UClass is not yet available
 	 */
-	void PostLoadAssetRegistryTags(FAssetData* AssetData);
+	bool TryPostLoadAssetRegistryTags(FAssetData* AssetData);
 
 	/** Update Redirect collector with redirects loaded from asset registry */
 	void UpdateRedirectCollector();
@@ -373,6 +479,19 @@ private:
 	void UpdatePersistentMountPoints();
 	void OnInitialSearchCompleted(Impl::FEventContext& EventContext);
 
+	/** This function is not called and, currently, will always return 'true'. See the commented out call in
+	 *	FAssetDataGatherer::TickInternal for further details on how it could be used. 
+	 */
+	friend FAssetDataGatherer;
+	bool ClassRequiresGameThreadProcessing(const UClass* Class) const;
+
+	friend UAssetRegistryImpl;
+	/** This exists purely for use during shutdown to enable the UAssetRegistryImpl to
+	 *  avoid waiting for the gatherer to terminate while holding the interface lock.
+	 *  See UAssetRegistryImpl::OnEnginePreExit 
+	 */
+	TUniquePtr<FAssetDataGatherer>& AccessGlobalGatherer() { return GlobalGatherer; }
+
 private:
 
 	/** Internal state of the cached asset registry */
@@ -410,9 +529,15 @@ private:
 	/** Async task that gathers asset information from disk */
 	TUniquePtr<FAssetDataGatherer> GlobalGatherer;
 
-	/** Lists of results from the background thread that are waiting to get processed by the main thread */
+	/** Lists of results from the gatherer thread that are waiting to get processed */
 	FAssetDataGatherer::FResults BackgroundResults;
 
+	/** Assets and dependencies that are not ready for processing because they can't yet run PostLoadAssetRegistrytags */
+	TMultiMap<FName, TUniquePtr<FAssetData>> DeferredAssets;
+	TMultiMap<FName, TUniquePtr<FAssetData>> DeferredAssetsForGameThread;
+	TMultiMap<FName, FPackageDependencyData> DeferredDependencies;
+	TMultiMap<FName, FPackageDependencyData> DeferredDependenciesForGameThread;
+
 #if !NO_LOGGING
 	/** Memory profiling information: How much memory is being used by the tags for each class. */
 	TMap<FTopLevelAssetPath, int64> TagSizeByClass;
@@ -456,6 +581,8 @@ private:
 
 	bool bVerboseLogging;
 
+	bool bForceCompletionEvenIfPostLoadsFail = false;
+
 	/** List of all class names derived from Blueprint (including Blueprint itself) */
 	TSet<FTopLevelAssetPath> ClassGeneratorNames;
 
@@ -472,23 +599,25 @@ private:
 	TArray<FAssetRegistryPackageRedirect> PackageRedirects;
 
 #if WITH_EDITOR
-	/** List of objects that need to be processed because they were loaded or saved */
-	TRingBuffer<TWeakObjectPtr<const UObject>> LoadedAssetsToProcess;
-
-	/** The set of object paths that have had their disk cache updated from the in memory version */
-	TSet<FSoftObjectPath> AssetDataObjectPathsUpdatedOnLoad;
-
 	/**
 	 * The set of object paths that have had their dependencies gathered since the last idle,
 	 * and that need to check for calculated dependencies at the next idle.
 	 */
 	TSet<FName> PackagesNeedingDependencyCalculation;
+	TSet<FName> PackagesNeedingDependencyCalculationOnGameThread;
+
+	/** List of objects that need to be processed because they were loaded or saved */
+	TRingBuffer<TWeakObjectPtr<const UObject>> LoadedAssetsToProcess;
+
+	/** The set of object paths that have had their disk cache updated from the in memory version */
+	TSet<FSoftObjectPath> AssetDataObjectPathsUpdatedOnLoad;
 
 	/** A map from directoryname to packagename of Packages that have CalculatedDependencies on packages in the directory. */
 	TMultiMap<FString, FName> DirectoryReferencers;
 
 	/** A map of per asset class dependency gatherer called in LoadCalculatedDependencies */
 	TMultiMap<FTopLevelAssetPath, UE::AssetDependencyGatherer::Private::FRegisteredAssetDependencyGatherer*> RegisteredDependencyGathererClasses;
+	mutable FRWLock RegisteredDependencyGathererClassesLock;
 	bool bRegisteredDependencyGathererClassesDirty;
 #endif
 #if WITH_ENGINE && WITH_EDITOR
@@ -579,7 +708,7 @@ struct FScanPathContext
 	int32 NumFoundAssets = 0;
 	bool bForceRescan = false;
 	bool bIgnoreDenyListScanFilters = false;
-	EGatherStatus Status = EGatherStatus::Active;
+	EGatherStatus Status = EGatherStatus::TickActiveGatherActive;
 };
 
 }
@@ -658,7 +787,7 @@ void EnumerateMemoryAssetsHelper(const FARCompiledFilter& InFilter, TSet<FName>&
  * Fills in OutPackageNamesWithAssets with names of all packages tested.
 */
 void EnumerateMemoryAssets(const FARCompiledFilter& InFilter, TSet<FName>& OutPackageNamesWithAssets,
-	bool& bOutStopIteration, FRWLock& InterfaceLock, const FAssetRegistryState& GuardedDataState,
+	bool& bOutStopIteration, UE::AssetRegistry::Private::FInterfaceRWLock& InterfaceLock, const FAssetRegistryState& GuardedDataState,
 	TFunctionRef<bool(FAssetData&&)> Callback, bool bSkipARFilteredAssets);
 
 }
diff --git a/Engine/Source/Runtime/AssetRegistry/Public/AssetRegistry/IAssetRegistry.h b/Engine/Source/Runtime/AssetRegistry/Public/AssetRegistry/IAssetRegistry.h
--- a/Engine/Source/Runtime/AssetRegistry/Public/AssetRegistry/IAssetRegistry.h
+++ b/Engine/Source/Runtime/AssetRegistry/Public/AssetRegistry/IAssetRegistry.h
@@ -778,6 +778,14 @@ public:
 		{
 		}
 
+		bool operator==(const FFileLoadProgressUpdateData& Other) const
+		{
+			return NumTotalAssets == Other.NumTotalAssets 
+				&& NumAssetsProcessedByAssetRegistry == Other.NumAssetsProcessedByAssetRegistry
+				&& NumAssetsPendingDataLoad == Other.NumAssetsPendingDataLoad
+				&& bIsDiscoveringAssetFiles == Other.bIsDiscoveringAssetFiles;
+		}
+
 		int32 NumTotalAssets;
 		int32 NumAssetsProcessedByAssetRegistry;
 		int32 NumAssetsPendingDataLoad;
@@ -888,7 +896,22 @@ public:
 	 */
 	virtual bool IsPathBeautificationNeeded(const FString& InAssetPath) const = 0;
 
+	/** Request to pause background processing of scan results while this object is in scope
+	 *  This can be used to allow a priority thread to perform along sequence of operations
+	 *  without having to contend with the background thread for data access
+	 */
+	struct FPauseBackgroundProcessingScope
+	{
+		FPauseBackgroundProcessingScope() { IAssetRegistry::GetChecked().RequestPauseBackgroundProcessing(); }
+		~FPauseBackgroundProcessingScope() { IAssetRegistry::GetChecked().RequestResumeBackgroundProcessing(); }
+	};
+
 protected:
+
+	// Implementation methods for FPauseBackgroundProcessingScope
+	virtual void RequestPauseBackgroundProcessing() {}
+	virtual void RequestResumeBackgroundProcessing() {}
+
 	// Functions specifically for calling from the asset manager
 	friend class UAssetManager;
 
diff --git a/Engine/Source/Runtime/CoreUObject/Private/Misc/AssetRegistryInterface.cpp b/Engine/Source/Runtime/CoreUObject/Private/Misc/AssetRegistryInterface.cpp
--- a/Engine/Source/Runtime/CoreUObject/Private/Misc/AssetRegistryInterface.cpp
+++ b/Engine/Source/Runtime/CoreUObject/Private/Misc/AssetRegistryInterface.cpp
@@ -52,8 +52,8 @@ namespace Utils
 		}
 
 		const bool bIsCooked = (PackageFlags & PKG_FilterEditorOnly);
-		if ((bIsCooked && SkipCookedClasses.Contains(AssetClass)) ||
-			(!bIsCooked && SkipUncookedClasses.Contains(AssetClass)))
+		if ((bIsCooked && InSkipCookedClasses.Contains(AssetClass)) ||
+			(!bIsCooked && InSkipUncookedClasses.Contains(AssetClass)))
 		{
 			return true;
 		}
diff --git a/Engine/Source/Runtime/CoreUObject/Private/UObject/Class.cpp b/Engine/Source/Runtime/CoreUObject/Private/UObject/Class.cpp
--- a/Engine/Source/Runtime/CoreUObject/Private/UObject/Class.cpp
+++ b/Engine/Source/Runtime/CoreUObject/Private/UObject/Class.cpp
@@ -4885,18 +4885,24 @@ void UClass::GetAssetRegistryTags(FAssetRegistryTagsContext Context) const
 }
 
 #if WITH_EDITOR
-void UClass::PostLoadAssetRegistryTags(const FAssetData& InAssetData, TArray<FAssetRegistryTag>& OutTagsAndValuesToUpdate) const
+void UClass::ThreadedPostLoadAssetRegistryTagsOverride(FPostLoadAssetRegistryTagsContext& Context) const
 {
-	Super::PostLoadAssetRegistryTags(InAssetData, OutTagsAndValuesToUpdate);
+	//@TODO this is not actually thread safe and should be guarded to avoid execution in parallel with BP compilation
+	// the ensure should help catch this via TSAN. UE-209240
+	// Unfortunately, we cannot simply limit the PostLoadAssetRegistryTags fixup to native classes because the base
+	// UObject::ThreadedPostLoadAssetRegistryTags iterates over the properties of non-native types and performs fixup.
+	ensure(!bLayoutChanging);
+
+	Super::ThreadedPostLoadAssetRegistryTagsOverride(Context);
 
 	static const FName ParentClassFName(TEXT("ParentClass"));
-	FString ParentClassTagValue = InAssetData.GetTagValueRef<FString>(ParentClassFName);
+	FString ParentClassTagValue = Context.GetAssetData().GetTagValueRef<FString>(ParentClassFName);
 	if (!ParentClassTagValue.IsEmpty() && FPackageName::IsShortPackageName(ParentClassTagValue))
 	{
-		FTopLevelAssetPath ParentClassPathName = UClass::TryConvertShortTypeNameToPathName<UStruct>(ParentClassTagValue, ELogVerbosity::Warning, TEXT("UClass::PostLoadAssetRegistryTags"));
+		FTopLevelAssetPath ParentClassPathName = UClass::TryConvertShortTypeNameToPathName<UStruct>(ParentClassTagValue, ELogVerbosity::Warning, TEXT("UClass::ThreadedPostLoadAssetRegistryTagsOverride"));
 		if (!ParentClassPathName.IsNull())
 		{
-			OutTagsAndValuesToUpdate.Add(FAssetRegistryTag(ParentClassFName, ParentClassPathName.ToString(), FAssetRegistryTag::TT_Alphabetical));
+			Context.AddTagToUpdate(FAssetRegistryTag(ParentClassFName, ParentClassPathName.ToString(), FAssetRegistryTag::TT_Alphabetical));
 		}
 	}
 }
diff --git a/Engine/Source/Runtime/CoreUObject/Private/UObject/Obj.cpp b/Engine/Source/Runtime/CoreUObject/Private/UObject/Obj.cpp
--- a/Engine/Source/Runtime/CoreUObject/Private/UObject/Obj.cpp
+++ b/Engine/Source/Runtime/CoreUObject/Private/UObject/Obj.cpp
@@ -2422,7 +2422,7 @@ const FName& UObject::SourceFileTagName()
 
 #if WITH_EDITOR
 
-static void PostLoadAssetRegistryTagProperty(FProperty* Prop, const FAssetData& AssetData, TArray<UObject::FAssetRegistryTag>& OutTagsAndValuesToUpdate)
+static void PostLoadAssetRegistryTagProperty(FProperty* Prop, UObject::FPostLoadAssetRegistryTagsContext& Context)
 {
 	// This TagType is ignored by the asset registry
 	UObject::FAssetRegistryTag::ETagType TagType = UObject::FAssetRegistryTag::ETagType::TT_Alphabetical;
@@ -2433,37 +2433,42 @@ static void PostLoadAssetRegistryTagProperty(FProperty* Prop, const FAssetData&
 		{
 			// Old files may contain legacy format of FSofObjectPtr::ToString() which used to return 
 			// an export path (ClassName'/Package/Name.ObjectName') however it now returns just a pathname (/Package/Name.ObjectName)
-			FString ExportPath = AssetData.GetTagValueRef<FString>(Prop->GetFName());
+			FString ExportPath = Context.GetAssetData().GetTagValueRef<FString>(Prop->GetFName());
 			int32 ClassSeparatorIndex = -1;
 			if (!ExportPath.IsEmpty() && ExportPath[0] != '/' && ExportPath.FindChar('\'', ClassSeparatorIndex))
 			{
 				// Strip the class name and leave just the pathname of an object
 				FString ObjectPath = FPackageName::ExportTextPathToObjectPath(ExportPath);
-				OutTagsAndValuesToUpdate.Add(UObject::FAssetRegistryTag(Prop->GetFName(), ObjectPath, TagType));
+				Context.AddTagToUpdate(UObject::FAssetRegistryTag(Prop->GetFName(), ObjectPath, TagType));
 			}
 		}
 		else if (Prop->IsA<FObjectPropertyBase>())
 		{
 			// Update the export path for short class names, but leave None alone to match save behavior
 			FObjectPropertyBase* PropertyObject = CastFieldChecked<FObjectPropertyBase>(Prop);
-			FString ExportPath = AssetData.GetTagValueRef<FString>(Prop->GetFName());
+			FString ExportPath = Context.GetAssetData().GetTagValueRef<FString>(Prop->GetFName());
 			if (!ExportPath.IsEmpty() && ExportPath[0] != '/' && ExportPath != TEXT("None"))
 			{
 				FString ObjectPath = FPackageName::ExportTextPathToObjectPath(ExportPath);
 				ExportPath = FObjectPropertyBase::GetExportPath(PropertyObject->PropertyClass->GetClassPathName(), ObjectPath);
-				OutTagsAndValuesToUpdate.Add(UObject::FAssetRegistryTag(Prop->GetFName(), ExportPath, TagType));
+				Context.AddTagToUpdate(UObject::FAssetRegistryTag(Prop->GetFName(), ExportPath, TagType));
 			}
 		}
 	}
 }
 
-void UObject::PostLoadAssetRegistryTags(const FAssetData& InAssetData, TArray<FAssetRegistryTag>& OutTagsAndValuesToUpdate) const
+void UObject::ThreadedPostLoadAssetRegistryTags(FPostLoadAssetRegistryTagsContext& Context) const
 {
+	ensureMsgf(GetClass()->HasAnyClassFlags(CLASS_Native), TEXT("ThreadedPostLoadAssetRegistryTags should not be called on non-native types. Detected a call on type '%s'"),
+		*GetClass()->GetName());
+
+	ThreadedPostLoadAssetRegistryTagsOverride(Context);
+
 	if (GetClass()->HasAssetRegistrySearchableProperties())
 	{
 		for (TFieldIterator<FProperty> PropertyIt(GetClass()); PropertyIt; ++PropertyIt)
 		{
-			PostLoadAssetRegistryTagProperty(*PropertyIt, InAssetData, OutTagsAndValuesToUpdate);
+			PostLoadAssetRegistryTagProperty(*PropertyIt, Context);
 		}
 
 		UScriptStruct* SparseClassDataStruct = GetClass()->GetSparseClassDataStruct();
@@ -2472,7 +2477,7 @@ void UObject::PostLoadAssetRegistryTags(const FAssetData& InAssetData, TArray<FA
 			const void* SparseClassData = GetClass()->GetSparseClassData(EGetSparseClassDataMethod::ArchetypeIfNull);
 			for (TFieldIterator<FProperty> PropertyIt(SparseClassDataStruct); PropertyIt; ++PropertyIt)
 			{
-				PostLoadAssetRegistryTagProperty(*PropertyIt, InAssetData, OutTagsAndValuesToUpdate);
+				PostLoadAssetRegistryTagProperty(*PropertyIt, Context);
 			}
 		}
 	}
diff --git a/Engine/Source/Runtime/CoreUObject/Public/UObject/Class.h b/Engine/Source/Runtime/CoreUObject/Public/UObject/Class.h
--- a/Engine/Source/Runtime/CoreUObject/Public/UObject/Class.h
+++ b/Engine/Source/Runtime/CoreUObject/Public/UObject/Class.h
@@ -3063,9 +3063,13 @@ public:
 	COREUOBJECT_API virtual void GetAssetRegistryTags(FAssetRegistryTagsContext Context) const override;
 	UE_DEPRECATED(5.4, "Implement the version that takes FAssetRegistryTagsContext instead.")
 	COREUOBJECT_API virtual void GetAssetRegistryTags(TArray<FAssetRegistryTag>& OutTags) const override;
+
 #if WITH_EDITOR
-	COREUOBJECT_API virtual void PostLoadAssetRegistryTags(const FAssetData& InAssetData, TArray<FAssetRegistryTag>& OutTagsAndValuesToUpdate) const;
+protected:
+	COREUOBJECT_API virtual void ThreadedPostLoadAssetRegistryTagsOverride(FPostLoadAssetRegistryTagsContext& Context) const;
+public:
 #endif // WITH_EDITOR
+
 	virtual bool IsAsset() const override { return false; }	
 	virtual bool IsNameStableForNetworking() const override { return true; } // For now, assume all classes have stable net names
 	COREUOBJECT_API virtual void GetPreloadDependencies(TArray<UObject*>& OutDeps) override;
diff --git a/Engine/Source/Runtime/CoreUObject/Public/UObject/Object.h b/Engine/Source/Runtime/CoreUObject/Public/UObject/Object.h
--- a/Engine/Source/Runtime/CoreUObject/Public/UObject/Object.h
+++ b/Engine/Source/Runtime/CoreUObject/Public/UObject/Object.h
@@ -16,6 +16,7 @@
 
 struct FAssetData;
 class FAssetRegistryTagsContext;
+namespace UE::AssetRegistry { class FAssetRegistryImpl; }
 class FConfigCacheIni;
 class FCustomPropertyConditionState;
 class FEditPropertyChain;
@@ -33,6 +34,7 @@ struct FObjectInstancingGraph;
 struct FPropertyChangedChainEvent;
 struct FTopLevelAssetPath;
 class UClass;
+
 #if UE_WITH_IRIS
 namespace UE::Net
 {
@@ -907,15 +909,44 @@ public:
 
 #if WITH_EDITOR
 
+	UE_DEPRECATED(5.5, "This function is no longer called. If you require this functionality, please implement ThreadedPostLoadAssetRegistryTagsOverride. \
+Note that this function MUST be safe to call on a background thread. If you have functionality that doesn't meet that requirement, \
+consider using the AssetRegistry's AssetAddedEvent which fires on the game thread.")
+	COREUOBJECT_API virtual void PostLoadAssetRegistryTags(const FAssetData& InAssetData, TArray<FAssetRegistryTag>& OutTagsAndValuesToUpdate) const final {}
+
+	/** Contains an AssetData referencing the asset to be processed and an array of tags and values to be updated on that asset */
+	struct FPostLoadAssetRegistryTagsContext
+	{
+		FPostLoadAssetRegistryTagsContext(const FAssetData& InAssetData, TArray<FAssetRegistryTag>& InTagsAndValuesToUpdate)
+		: AssetData(InAssetData), TagsAndValuesToUpdate(InTagsAndValuesToUpdate) {}
+
+		const FAssetData& GetAssetData() const { return AssetData; }
+		void AddTagToUpdate(const FAssetRegistryTag& Tag) { TagsAndValuesToUpdate.Add(Tag); }
+
+	private:
+		/** The AssetData being considered */
+		const FAssetData& AssetData;
+		/** The asset registry tags to be updated for this asset data */
+		TArray<FAssetRegistryTag>& TagsAndValuesToUpdate;
+	};
+
+protected:
 	/**
-	 * Performs fixup on loaded asset registry data. 
-	 * This function is called from inside the AssetRegistry CriticalSection. DO NOT CALL ASSETREGISTRY FUNCTIONS FROM THIS FUNCTION, IT WILL DEADLOCK.
+	 * Performs fixup on loaded asset registry data.
+	 *
+	 * DO NOT CALL ASSETREGISTRY FUNCTIONS FROM THIS FUNCTION, IT WILL DEADLOCK because this function is called from inside the AssetRegistry CriticalSection. 
+	 * IT MUST BE SAFE TO CALL IN MULTITHREADED BACKGROUND THREADS because the AssetRegistry can call this function from a background thread.
+	 * 
 	 * Note that this function is only called on Class Default Objects where the actual object instance data used to generate 
 	 * the asset data is not available.
-	 * @param InAssetData Asset data loaded from the AssetRegistry
-	 * @return Pointer to new asset data after fixup or nullptr if no fixup was required
+	 * @param Context Contains input and receives output from the function. @see FPostloadAssetRegistryTagsContext
 	 */
-	COREUOBJECT_API virtual void PostLoadAssetRegistryTags(const FAssetData& InAssetData, TArray<FAssetRegistryTag>& OutTagsAndValuesToUpdate) const;
+	COREUOBJECT_API virtual void ThreadedPostLoadAssetRegistryTagsOverride(FPostLoadAssetRegistryTagsContext& Context) const {}
+
+public:
+
+	/** Internal use function for the AssetRegistry. Performs general fixup required by UObject and invokes ThreadedPostLoadAssetRegistryTagsOverride */
+	COREUOBJECT_API void ThreadedPostLoadAssetRegistryTags(FPostLoadAssetRegistryTagsContext& Context) const;
 
 	/**
 	 * Additional data pertaining to asset registry tags used by the editor
diff --git a/Engine/Source/Runtime/Engine/Classes/Engine/Blueprint.h b/Engine/Source/Runtime/Engine/Classes/Engine/Blueprint.h
--- a/Engine/Source/Runtime/Engine/Classes/Engine/Blueprint.h
+++ b/Engine/Source/Runtime/Engine/Classes/Engine/Blueprint.h
@@ -930,8 +930,13 @@ public:
 	ENGINE_API virtual void GetAssetRegistryTags(FAssetRegistryTagsContext Context) const override;
 	UE_DEPRECATED(5.4, "Implement the version that takes FAssetRegistryTagsContext instead.")
 	ENGINE_API virtual void GetAssetRegistryTags(TArray<FAssetRegistryTag>& OutTags) const override;
-	ENGINE_API virtual void PostLoadAssetRegistryTags(const FAssetData& InAssetData, TArray<FAssetRegistryTag>& OutTagsAndValuesToUpdate) const;
+protected:
+	ENGINE_API virtual void ThreadedPostLoadAssetRegistryTagsOverride(FPostLoadAssetRegistryTagsContext& Context) const override;
+public:
+	UE_DEPRECATED(5.5, "This function is no longer called. See UObject::PostLoadAssetRegistryTags for more information.")
 	static ENGINE_API void PostLoadBlueprintAssetRegistryTags(const FAssetData& InAssetData, TArray<FAssetRegistryTag>& OutTagsAndValuesToUpdate);
+	
+	static ENGINE_API void PostLoadBlueprintAssetRegistryTags(FPostLoadAssetRegistryTagsContext& Context);
 	ENGINE_API virtual FPrimaryAssetId GetPrimaryAssetId() const override;
 	ENGINE_API virtual void BeginCacheForCookedPlatformData(const ITargetPlatform *TargetPlatform) override;
 	ENGINE_API virtual bool IsCachedCookedPlatformDataLoaded(const ITargetPlatform* TargetPlatform) override;
diff --git a/Engine/Source/Runtime/Engine/Classes/Engine/BlueprintGeneratedClass.h b/Engine/Source/Runtime/Engine/Classes/Engine/BlueprintGeneratedClass.h
--- a/Engine/Source/Runtime/Engine/Classes/Engine/BlueprintGeneratedClass.h
+++ b/Engine/Source/Runtime/Engine/Classes/Engine/BlueprintGeneratedClass.h
@@ -778,7 +778,9 @@ public:
 #if WITH_EDITOR
 	ENGINE_API virtual void GetAdditionalAssetDataObjectsForCook(FArchiveCookContext& CookContext, 
 		TArray<UObject*>& OutObjects) const override;
-	ENGINE_API virtual void PostLoadAssetRegistryTags(const FAssetData& InAssetData, TArray<FAssetRegistryTag>& OutTagsAndValuesToUpdate) const;
+protected:
+	ENGINE_API virtual void ThreadedPostLoadAssetRegistryTagsOverride(FPostLoadAssetRegistryTagsContext& Context) const;
+public:
 #endif //~ WITH_EDITOR
 	ENGINE_API virtual FPrimaryAssetId GetPrimaryAssetId() const override;
 	ENGINE_API virtual bool NeedsLoadForServer() const override;
diff --git a/Engine/Source/Runtime/Engine/Classes/Engine/DataTable.h b/Engine/Source/Runtime/Engine/Classes/Engine/DataTable.h
--- a/Engine/Source/Runtime/Engine/Classes/Engine/DataTable.h
+++ b/Engine/Source/Runtime/Engine/Classes/Engine/DataTable.h
@@ -124,7 +124,9 @@ public:
 	
 #if WITH_EDITOR
 	ENGINE_API virtual void PostEditChangeProperty(FPropertyChangedEvent& PropertyChangedEvent) override;
-	ENGINE_API virtual void PostLoadAssetRegistryTags(const FAssetData& InAssetData, TArray<FAssetRegistryTag>& OutTagsAndValuesToUpdate) const override;
+protected:
+	ENGINE_API virtual void ThreadedPostLoadAssetRegistryTagsOverride(FPostLoadAssetRegistryTagsContext& Context) const override;
+public:
 #endif // WITH_EDITOR
 
 	//~ Begin UObject Interface.
diff --git a/Engine/Source/Runtime/Engine/Classes/Engine/World.h b/Engine/Source/Runtime/Engine/Classes/Engine/World.h
--- a/Engine/Source/Runtime/Engine/Classes/Engine/World.h
+++ b/Engine/Source/Runtime/Engine/Classes/Engine/World.h
@@ -2811,7 +2811,9 @@ public:
 	virtual void GetAssetRegistryTags(FAssetRegistryTagsContext Context) const override;
 	UE_DEPRECATED(5.4, "Implement the version that takes FAssetRegistryTagsContext instead.")
 	virtual void GetAssetRegistryTags(TArray<FAssetRegistryTag>& OutTags) const override;
-	virtual void PostLoadAssetRegistryTags(const FAssetData& InAssetData, TArray<FAssetRegistryTag>& OutTagsAndValuesToUpdate) const;
+protected:
+	virtual void ThreadedPostLoadAssetRegistryTagsOverride(FPostLoadAssetRegistryTagsContext& Context) const;
+public:
 	virtual bool IsNameStableForNetworking() const override;
 #endif
 	virtual bool ResolveSubobject(const TCHAR* SubObjectPath, UObject*& OutObject, bool bLoadIfExists) override;
diff --git a/Engine/Source/Runtime/Engine/Private/Blueprint.cpp b/Engine/Source/Runtime/Engine/Private/Blueprint.cpp
--- a/Engine/Source/Runtime/Engine/Private/Blueprint.cpp
+++ b/Engine/Source/Runtime/Engine/Private/Blueprint.cpp
@@ -1103,23 +1103,29 @@ void UBlueprint::GetAssetRegistryTags(FAssetRegistryTagsContext Context) const
 #endif
 }
 
-void UBlueprint::PostLoadAssetRegistryTags(const FAssetData& InAssetData, TArray<FAssetRegistryTag>& OutTagsAndValuesToUpdate) const
+void UBlueprint::ThreadedPostLoadAssetRegistryTagsOverride(FPostLoadAssetRegistryTagsContext& Context) const
 {
-	Super::PostLoadAssetRegistryTags(InAssetData, OutTagsAndValuesToUpdate);
-	PostLoadBlueprintAssetRegistryTags(InAssetData, OutTagsAndValuesToUpdate);
+	Super::ThreadedPostLoadAssetRegistryTagsOverride(Context);
+	PostLoadBlueprintAssetRegistryTags(Context);
 }
 
 void UBlueprint::PostLoadBlueprintAssetRegistryTags(const FAssetData& InAssetData, TArray<FAssetRegistryTag>& OutTagsAndValuesToUpdate)
 {
-	auto FixTagValueShortClassName = [&InAssetData, &OutTagsAndValuesToUpdate](FName TagName, FAssetRegistryTag::ETagType TagType)
+	FPostLoadAssetRegistryTagsContext Context(InAssetData, OutTagsAndValuesToUpdate);
+	PostLoadBlueprintAssetRegistryTags(Context);
+}
+
+void UBlueprint::PostLoadBlueprintAssetRegistryTags(FPostLoadAssetRegistryTagsContext& Context)
+{
+	auto FixTagValueShortClassName = [&Context](FName TagName, FAssetRegistryTag::ETagType TagType)
 	{
-		FString TagValue = InAssetData.GetTagValueRef<FString>(TagName);
+		FString TagValue = Context.GetAssetData().GetTagValueRef<FString>(TagName);
 		if (!TagValue.IsEmpty() && TagValue != TEXT("None"))
 		{
 			if (UClass::TryFixShortClassNameExportPath(TagValue, ELogVerbosity::Warning,
-				TEXT("UBlueprint::PostLoadAssetRegistryTags"), true /* bClearOnError */))
+				TEXT("UBlueprint::PostLoadBlueprintAssetRegistryTags"), true /* bClearOnError */))
 			{
-				OutTagsAndValuesToUpdate.Add(FAssetRegistryTag(TagName, TagValue, TagType));
+				Context.AddTagToUpdate(FAssetRegistryTag(TagName, TagValue, TagType));
 			}
 		}
 	};
diff --git a/Engine/Source/Runtime/Engine/Private/BlueprintGeneratedClass.cpp b/Engine/Source/Runtime/Engine/Private/BlueprintGeneratedClass.cpp
--- a/Engine/Source/Runtime/Engine/Private/BlueprintGeneratedClass.cpp
+++ b/Engine/Source/Runtime/Engine/Private/BlueprintGeneratedClass.cpp
@@ -622,19 +622,19 @@ void UBlueprintGeneratedClass::GetAdditionalAssetDataObjectsForCook(FArchiveCook
 	}
 }
 
-void UBlueprintGeneratedClass::PostLoadAssetRegistryTags(const FAssetData& InAssetData, TArray<FAssetRegistryTag>& OutTagsAndValuesToUpdate) const
+void UBlueprintGeneratedClass::ThreadedPostLoadAssetRegistryTagsOverride(FPostLoadAssetRegistryTagsContext& Context) const
 {
-	Super::PostLoadAssetRegistryTags(InAssetData, OutTagsAndValuesToUpdate);
+	Super::ThreadedPostLoadAssetRegistryTagsOverride(Context);
 
-	auto FixTagValueShortClassName = [&InAssetData, &OutTagsAndValuesToUpdate](FName TagName, FAssetRegistryTag::ETagType TagType)
+	auto FixTagValueShortClassName = [&Context](FName TagName, FAssetRegistryTag::ETagType TagType)
 	{
-		FString TagValue = InAssetData.GetTagValueRef<FString>(TagName);
+		FString TagValue = Context.GetAssetData().GetTagValueRef<FString>(TagName);
 		if (!TagValue.IsEmpty() && TagValue != TEXT("None"))
 		{
 			if (UClass::TryFixShortClassNameExportPath(TagValue, ELogVerbosity::Warning,
-				TEXT("UBlueprintGeneratedClass::PostLoadAssetRegistryTags"), true /* bClearOnError */))
+				TEXT("UBlueprintGeneratedClass::ThreadedPostLoadAssetRegistryTagsOverride"), true /* bClearOnError */))
 			{
-				OutTagsAndValuesToUpdate.Add(FAssetRegistryTag(TagName, TagValue, TagType));
+				Context.AddTagToUpdate(FAssetRegistryTag(TagName, TagValue, TagType));
 			}
 		}
 	};
diff --git a/Engine/Source/Runtime/Engine/Private/DataTable.cpp b/Engine/Source/Runtime/Engine/Private/DataTable.cpp
--- a/Engine/Source/Runtime/Engine/Private/DataTable.cpp
+++ b/Engine/Source/Runtime/Engine/Private/DataTable.cpp
@@ -380,18 +380,18 @@ void UDataTable::PostLoad()
 #endif // WITH_EDITORONLY_DATA
 
 #if WITH_EDITOR
-void UDataTable::PostLoadAssetRegistryTags(const FAssetData& InAssetData, TArray<FAssetRegistryTag>& OutTagsAndValuesToUpdate) const
+void UDataTable::ThreadedPostLoadAssetRegistryTagsOverride(FPostLoadAssetRegistryTagsContext& Context) const
 {
-	Super::PostLoadAssetRegistryTags(InAssetData, OutTagsAndValuesToUpdate);
+	Super::ThreadedPostLoadAssetRegistryTagsOverride(Context);
 
 	static const FName RowStructureTag(TEXT("RowStructure"));
-	FString TagValue = InAssetData.GetTagValueRef<FString>(RowStructureTag);
+	FString TagValue = Context.GetAssetData().GetTagValueRef<FString>(RowStructureTag);
 	if (!TagValue.IsEmpty() && FPackageName::IsShortPackageName(TagValue))
 	{
-		FTopLevelAssetPath PathName = UClass::TryConvertShortTypeNameToPathName<UField>(TagValue, ELogVerbosity::Warning, TEXT("UDataTable::PostLoadAssetRegistryTags"));
+		FTopLevelAssetPath PathName = UClass::TryConvertShortTypeNameToPathName<UField>(TagValue, ELogVerbosity::Warning, TEXT("UDataTable::ThreadedPostLoadAssetRegistryTagsOverride"));
 		if (!PathName.IsNull())
 		{
-			OutTagsAndValuesToUpdate.Add(FAssetRegistryTag(RowStructureTag, PathName.ToString(), FAssetRegistryTag::TT_Alphabetical));
+			Context.AddTagToUpdate(FAssetRegistryTag(RowStructureTag, PathName.ToString(), FAssetRegistryTag::TT_Alphabetical));
 		}
 	}
 }
diff --git a/Engine/Source/Runtime/Engine/Private/World.cpp b/Engine/Source/Runtime/Engine/Private/World.cpp
--- a/Engine/Source/Runtime/Engine/Private/World.cpp
+++ b/Engine/Source/Runtime/Engine/Private/World.cpp
@@ -9267,12 +9267,12 @@ void UWorld::GetAssetRegistryTags(FAssetRegistryTagsContext Context) const
 	FWorldDelegates::GetAssetTagsWithContext.Broadcast(this, Context);
 }
 
-void UWorld::PostLoadAssetRegistryTags(const FAssetData& InAssetData, TArray<FAssetRegistryTag>& OutTagsAndValuesToUpdate) const
+void UWorld::ThreadedPostLoadAssetRegistryTagsOverride(FPostLoadAssetRegistryTagsContext& Context) const
 {
-	Super::PostLoadAssetRegistryTags(InAssetData, OutTagsAndValuesToUpdate);
+	Super::ThreadedPostLoadAssetRegistryTagsOverride(Context);
 
-	// GetAssetRegistryTags appends the LevelBlueprint tags to the World's tags, so we also have to run the Blueprint PostLoadAssetRegistryTags
-	UBlueprint::PostLoadBlueprintAssetRegistryTags(InAssetData, OutTagsAndValuesToUpdate);
+	// GetAssetRegistryTags appends the LevelBlueprint tags to the World's tags, so we also have to run the Blueprint ThreadedPostLoadAssetRegistryTagsOverride
+	UBlueprint::PostLoadBlueprintAssetRegistryTags(Context);
 }
 
 bool UWorld::IsNameStableForNetworking() const
diff --git a/Engine/Source/Runtime/LevelSequence/Private/LevelSequence.cpp b/Engine/Source/Runtime/LevelSequence/Private/LevelSequence.cpp
--- a/Engine/Source/Runtime/LevelSequence/Private/LevelSequence.cpp
+++ b/Engine/Source/Runtime/LevelSequence/Private/LevelSequence.cpp
@@ -210,12 +210,12 @@ void ULevelSequence::GetAssetRegistryTagMetadata(TMap<FName, FAssetRegistryTagMe
 	Super::GetAssetRegistryTagMetadata(OutMetadata);
 }
 
-void ULevelSequence::PostLoadAssetRegistryTags(const FAssetData& InAssetData, TArray<FAssetRegistryTag>& OutTagsAndValuesToUpdate) const
+void ULevelSequence::ThreadedPostLoadAssetRegistryTagsOverride(FPostLoadAssetRegistryTagsContext& Context) const
 {
-	Super::PostLoadAssetRegistryTags(InAssetData, OutTagsAndValuesToUpdate);
+	Super::ThreadedPostLoadAssetRegistryTagsOverride(Context);
 
-	// GetAssetRegistryTags appends the DirectorBlueprint tags to the World's tags, so we also have to run the Blueprint PostLoadAssetRegistryTags
-	UBlueprint::PostLoadBlueprintAssetRegistryTags(InAssetData, OutTagsAndValuesToUpdate);
+	// GetAssetRegistryTags appends the DirectorBlueprint tags to the World's tags, so we also have to run the Blueprint ThreadedPostLoadAssetRegistryTagsOverride
+	UBlueprint::PostLoadBlueprintAssetRegistryTags(Context);
 }
 
 void PurgeLegacyBlueprints(UObject* InObject, UPackage* Package)
diff --git a/Engine/Source/Runtime/LevelSequence/Public/LevelSequence.h b/Engine/Source/Runtime/LevelSequence/Public/LevelSequence.h
--- a/Engine/Source/Runtime/LevelSequence/Public/LevelSequence.h
+++ b/Engine/Source/Runtime/LevelSequence/Public/LevelSequence.h
@@ -79,8 +79,10 @@ public:
 	LEVELSEQUENCE_API virtual void GetAssetRegistryTags(FAssetRegistryTagsContext Context) const override;
 	UE_DEPRECATED(5.4, "Implement the version that takes FAssetRegistryTagsContext instead.")
 	LEVELSEQUENCE_API virtual void GetAssetRegistryTags(TArray<FAssetRegistryTag>& OutTags) const override;
-	LEVELSEQUENCE_API virtual void PostLoadAssetRegistryTags(const FAssetData& InAssetData, TArray<FAssetRegistryTag>& OutTagsAndValuesToUpdate) const override;
-	
+protected:
+	LEVELSEQUENCE_API virtual void ThreadedPostLoadAssetRegistryTagsOverride(FPostLoadAssetRegistryTagsContext& Context) const override;
+public:
+
 	DECLARE_DELEGATE_RetVal_OneParam(void, FPostDuplicateEvent, ULevelSequence*);
 	static LEVELSEQUENCE_API FPostDuplicateEvent PostDuplicateEvent;
 #endif
