Ensure that a call to EnumerateMemoryAssets doesn't contend with a background thread ticking the asset gatherer by pausing background work during the enumeration. Otherwise we can preempt background thread thousands of times as we take the interface lock for each object we examine. Additionally, ensure that if a CDO is not yet available for an AssetClass we will defer asset's processing similar to how we defer it if the AssetClass itself is unavailable.

diff --git a/Engine/Source/Runtime/AssetRegistry/Private/AssetRegistry.cpp b/Engine/Source/Runtime/AssetRegistry/Private/AssetRegistry.cpp
--- a/Engine/Source/Runtime/AssetRegistry/Private/AssetRegistry.cpp
+++ b/Engine/Source/Runtime/AssetRegistry/Private/AssetRegistry.cpp
@@ -2538,6 +2538,10 @@ void EnumerateMemoryAssets(const FARCompiledFilter& InFilter, TSet<FName>& OutPa
 	TFunctionRef<bool(FAssetData&&)> Callback, bool bSkipARFilteredAssets)
 {
 	check(!InFilter.IsEmpty() && Utils::IsFilterValid(InFilter));
+
+	// Avoid contending with the background thread every time we take the interface lock below.
+	IAssetRegistry::FPauseBackgroundProcessingScope PauseProcessingScopeGuard;
+
 	EnumerateMemoryAssetsHelper(InFilter, OutPackageNamesWithAssets, bOutStopIteration,
 		[&InFilter, &Callback, &InterfaceLock, &GuardedDataState](const UObject* Object, FAssetData&& PartialAssetData)
 		{
@@ -5894,6 +5898,15 @@ bool FAssetRegistryImpl::TryPostLoadAssetRegistryTags(FAssetData* AssetData)
 					AssetData->TagsAndValues = FAssetDataTagMapSharedView(MoveTemp(TagsAndValues));
 				}
 			}
+			else if (!bForceCompletionEvenIfPostLoadsFail)
+			{
+				CouldPostLoadAssetRegistryTags = false;
+			}
+			else 
+			{
+				ensureMsgf(!MakeFinalChecks, TEXT("Unable to PostLoadAssetRegistryTags for '%s' because the CDO for ancestor class '%s' could not be found."),
+					*AssetData->GetObjectPathString(), *AssetClassPath.ToString());
+			}
 		}
 	}
 	return CouldPostLoadAssetRegistryTags;
