Enabling multithreaded data processing for Asset Registry. In case of problems, this can be disabled with
-dpcvars=AssetRegistry.TickGatherOnGTOnly=1

#tests Multiple devs ran with this flag for a week, various preflights

diff --git a/Engine/Source/Runtime/AssetRegistry/Private/AssetDataGatherer.cpp b/Engine/Source/Runtime/AssetRegistry/Private/AssetDataGatherer.cpp
--- a/Engine/Source/Runtime/AssetRegistry/Private/AssetDataGatherer.cpp
+++ b/Engine/Source/Runtime/AssetRegistry/Private/AssetDataGatherer.cpp
@@ -85,9 +85,7 @@ static FAutoConsoleVariableRef CVarIgnoreEmptyDirectories(
 	bIgnoreEmptyDirectories,
 	TEXT("If true, completely empty leaf directories are ignored by the asset registry while scanning"));
 
-//@TODO This should be set to false (enabling multithreaded processing) only after resolving UE-209921 
-// (relating to the GC lock).
-bool bTickGatherOnGTOnly = true;
+bool bTickGatherOnGTOnly = false;
 static FAutoConsoleVariableRef CVarTickGatherOnGTOnly(
 	TEXT("AssetRegistry.TickGatherOnGTOnly"),
 	bTickGatherOnGTOnly,
