World Partition - Optimized World Partition builders startup time
* Added WorldPartitionBuilder to the CommandletsUsingAR config
* FAssetRegistryImpl::GetRedirectedObjectPath() - Calls done with bNeedsScanning should first test if asset was scanned rather than always performing a scan
* UAssetManager & UObjectLibrary are now using AssetRegistry.IsSearchAllAssets() in order to test whether there's an async AR scan in progress
* FActorDescArchive: Skip calls to FixupRedirectedAssetPath() for native classes

On test project:
* Reaching the end of the World Partition init on a Threadripper: from 8m01s to 4m40s
* Running the HLOD setup job: from 15m13s to 8m50s

diff --git a/Engine/Config/BaseEngine.ini b/Engine/Config/BaseEngine.ini
--- a/Engine/Config/BaseEngine.ini
+++ b/Engine/Config/BaseEngine.ini
@@ -2757,6 +2757,7 @@ bFilterSearchableNames=true
 +CommandletsUsingAR=DataValidation
 +CommandletsUsingAR=GatherText
 +CommandletsUsingAR=GenerateNaniteDisplacedMesh
++CommandletsUsingAR=WorldPartitionBuilder
 
 ; AssetDataGathering can be cached in the editor; there are two different caches, one for the directory scan to find all
 ; packages and one for tags and other data read from the packages. The directory scan cache by default is only read
diff --git a/Engine/Source/Runtime/AssetRegistry/Private/AssetRegistry.cpp b/Engine/Source/Runtime/AssetRegistry/Private/AssetRegistry.cpp
--- a/Engine/Source/Runtime/AssetRegistry/Private/AssetRegistry.cpp
+++ b/Engine/Source/Runtime/AssetRegistry/Private/AssetRegistry.cpp
@@ -3149,20 +3149,29 @@ FSoftObjectPath FAssetRegistryImpl::GetRedirectedObjectPath(const FSoftObjectPat
 		RedirectedPath = FSoftObjectPath(Redirector->DestinationObject);
 	}
 
-	if (bNeedsScanning)
+	FString SubPathString;
+
+	auto RetrieveAssetData = [&]()
+	{
+		const FAssetData* AssetData = State.GetAssetByObjectPath(RedirectedPath);
+		if (!AssetData && RedirectedPath.IsSubobject())
+		{
+			// If we found no Asset because it is a subobject, then look for its toplevelobject's Asset
+			SubPathString = RedirectedPath.GetSubPathString();
+			RedirectedPath = FSoftObjectPath(RedirectedPath.GetAssetPath(), FString());
+			AssetData = State.GetAssetByObjectPath(RedirectedPath);
+		}
+		return AssetData;
+	};
+
+	const FAssetData* AssetData = RetrieveAssetData();
+
+	if (!AssetData && bNeedsScanning)
 	{
 		UE::AssetRegistry::Impl::FScanPathContext Context(*EventContext, *InheritanceContext, {}, { RedirectedPath.ToString() });
 		ScanPathsSynchronous(Context);
-	}
 
-	FString SubPathString;
-	const FAssetData* AssetData = State.GetAssetByObjectPath(RedirectedPath);
-	if (!AssetData && RedirectedPath.IsSubobject())
-	{
-		// If we found no Asset because it is a subobject, then look for its toplevelobject's Asset
-		SubPathString = RedirectedPath.GetSubPathString();
-		RedirectedPath = FSoftObjectPath(RedirectedPath.GetAssetPath(), FString());
-		AssetData = State.GetAssetByObjectPath(RedirectedPath);
+		AssetData = RetrieveAssetData();
 	}
 
 	// Most of the time this will either not be a redirector or only have one redirect, so optimize for that case
@@ -3187,14 +3196,16 @@ FSoftObjectPath FAssetRegistryImpl::GetRedirectedObjectPath(const FSoftObjectPat
 			break;
 		}
 
-		if (bNeedsScanning)
+		AssetData = State.GetAssetByObjectPath(RedirectedPath);
+		if (!AssetData && bNeedsScanning)
 		{
 			UE::AssetRegistry::Impl::FScanPathContext Context(*EventContext, *InheritanceContext, {}, { RedirectedPath.ToString() });
 			ScanPathsSynchronous(Context);
+
+			AssetData = State.GetAssetByObjectPath(RedirectedPath);
 		}
 
 		SeenPaths.Add(RedirectedPath);
-		AssetData = State.GetAssetByObjectPath(RedirectedPath);
 	}
 
 	if (!SubPathString.IsEmpty())
diff --git a/Engine/Source/Runtime/AssetRegistry/Private/AssetRegistryHelpers.cpp b/Engine/Source/Runtime/AssetRegistry/Private/AssetRegistryHelpers.cpp
--- a/Engine/Source/Runtime/AssetRegistry/Private/AssetRegistryHelpers.cpp
+++ b/Engine/Source/Runtime/AssetRegistry/Private/AssetRegistryHelpers.cpp
@@ -234,6 +234,11 @@ UAssetRegistryHelpers::FTemporaryCachingModeScope::~FTemporaryCachingModeScope()
 
 void UAssetRegistryHelpers::FixupRedirectedAssetPath(FSoftObjectPath& InOutSoftObjectPath)
 {
+	if (InOutSoftObjectPath.IsNull())
+	{
+		return;
+	}
+
 	FSoftObjectPath FoundRedirection;
 	InOutSoftObjectPath.FixupCoreRedirects();
 
@@ -253,6 +258,11 @@ void UAssetRegistryHelpers::FixupRedirectedAssetPath(FSoftObjectPath& InOutSoftO
 
 void UAssetRegistryHelpers::FixupRedirectedAssetPath(FName& InOutAssetPath)
 {
+	if (InOutAssetPath.IsNone())
+	{
+		return;
+	}
+
 	FSoftObjectPath SoftObjectPath(InOutAssetPath.ToString());
 	FixupRedirectedAssetPath(SoftObjectPath);
 	InOutAssetPath = FName(*SoftObjectPath.ToString());
diff --git a/Engine/Source/Runtime/Engine/Private/AssetManager.cpp b/Engine/Source/Runtime/Engine/Private/AssetManager.cpp
--- a/Engine/Source/Runtime/Engine/Private/AssetManager.cpp
+++ b/Engine/Source/Runtime/Engine/Private/AssetManager.cpp
@@ -401,13 +401,13 @@ void UAssetManager::PostInitProperties()
 	{
 		const UAssetManagerSettings& Settings = GetSettings();
 #if WITH_EDITOR
-		bIsGlobalAsyncScanEnvironment = GIsEditor && !IsRunningCommandlet();
+		IAssetRegistry& AssetRegistry = GetAssetRegistry();
+
+		bIsGlobalAsyncScanEnvironment = AssetRegistry.IsSearchAllAssets();
 
 		if (bIsGlobalAsyncScanEnvironment)
 		{
 			// Listen for when the asset registry has finished discovering files
-			IAssetRegistry& AssetRegistry = GetAssetRegistry();
-
 			AssetRegistry.OnFilesLoaded().AddUObject(this, &UAssetManager::OnAssetRegistryFilesLoaded);
 			AssetRegistry.OnInMemoryAssetCreated().AddUObject(this, &UAssetManager::OnInMemoryAssetCreated);
 			AssetRegistry.OnInMemoryAssetDeleted().AddUObject(this, &UAssetManager::OnInMemoryAssetDeleted);
diff --git a/Engine/Source/Runtime/Engine/Private/ObjectLibrary.cpp b/Engine/Source/Runtime/Engine/Private/ObjectLibrary.cpp
--- a/Engine/Source/Runtime/Engine/Private/ObjectLibrary.cpp
+++ b/Engine/Source/Runtime/Engine/Private/ObjectLibrary.cpp
@@ -26,14 +26,14 @@ UObjectLibrary::UObjectLibrary(const FObjectInitializer& ObjectInitializer)
 #if WITH_EDITOR
 	if ( !HasAnyFlags(RF_ClassDefaultObject) )
 	{
-		bIsGlobalAsyncScanEnvironment = GIsEditor && !IsRunningCommandlet();
+		FAssetRegistryModule& AssetRegistryModule = FModuleManager::LoadModuleChecked<FAssetRegistryModule>(TEXT("AssetRegistry"));
+		IAssetRegistry& AssetRegistry = AssetRegistryModule.Get();
+
+		bIsGlobalAsyncScanEnvironment = AssetRegistry.IsSearchAllAssets();
 
 		if ( bIsGlobalAsyncScanEnvironment )
 		{
 			// Listen for when the asset registry has finished discovering files
-			FAssetRegistryModule& AssetRegistryModule = FModuleManager::LoadModuleChecked<FAssetRegistryModule>(TEXT("AssetRegistry"));
-			IAssetRegistry& AssetRegistry = AssetRegistryModule.Get();
-
 			AssetRegistry.OnFilesLoaded().AddUObject(this, &UObjectLibrary::OnAssetRegistryFilesLoaded);
 		}
 	}
diff --git a/Engine/Source/Runtime/Engine/Private/WorldPartition/WorldPartitionActorDescArchive.cpp b/Engine/Source/Runtime/Engine/Private/WorldPartition/WorldPartitionActorDescArchive.cpp
--- a/Engine/Source/Runtime/Engine/Private/WorldPartition/WorldPartitionActorDescArchive.cpp
+++ b/Engine/Source/Runtime/Engine/Private/WorldPartition/WorldPartitionActorDescArchive.cpp
@@ -68,7 +68,7 @@ PRAGMA_ENABLE_DEPRECATION_WARNINGS
 
 	if (IsLoading())
 	{
-		auto TryRedirectClass = [](FTopLevelAssetPath& InOutClassPath)
+		auto TryRedirectClass = [](FTopLevelAssetPath& InOutClassPath, bool bNativeClass)
 		{
 			if (InOutClassPath.IsValid())
 			{
@@ -80,14 +80,17 @@ PRAGMA_ENABLE_DEPRECATION_WARNINGS
 					InOutClassPath = FTopLevelAssetPath(RedirectedClassRedirect.ToString());
 				}
 
-				FSoftObjectPath RedirectedClassPath(InOutClassPath.ToString());
-				UAssetRegistryHelpers::FixupRedirectedAssetPath(RedirectedClassPath);
-				InOutClassPath = RedirectedClassPath.GetAssetPath();
+				if (!bNativeClass)
+				{
+					FSoftObjectPath RedirectedClassPath(InOutClassPath.ToString());
+					UAssetRegistryHelpers::FixupRedirectedAssetPath(RedirectedClassPath);
+					InOutClassPath = RedirectedClassPath.GetAssetPath();
+				}
 			}
 		};
 
-		TryRedirectClass(ActorDesc->NativeClass);
-		TryRedirectClass(ActorDesc->BaseClass);
+		TryRedirectClass(ActorDesc->NativeClass, true);
+		TryRedirectClass(ActorDesc->BaseClass, false);
 	}
 
 	// Get the class descriptor to do delta serialization if no base desc was provided
