AssetRegistry: Improve feedback for the programmer-targetted warning messages in ScanPathsSynchronous by adding callstack to the warning message.
#rnx

diff --git a/Engine/Source/Runtime/AssetRegistry/Private/AssetRegistry.cpp b/Engine/Source/Runtime/AssetRegistry/Private/AssetRegistry.cpp
--- a/Engine/Source/Runtime/AssetRegistry/Private/AssetRegistry.cpp
+++ b/Engine/Source/Runtime/AssetRegistry/Private/AssetRegistry.cpp
@@ -5487,6 +5487,14 @@ FScanPathContext::FScanPathContext(FEventContext& InEventContext, FClassInherita
 		OutFoundAssets->Empty();
 	}
 
+	bool bLogCallstack = false;
+	ON_SCOPE_EXIT
+	{
+		if (bLogCallstack)
+		{
+			FDebug::DumpStackTraceToLog(ELogVerbosity::Warning);
+		}
+	};
 	if (bIgnoreDenyListScanFilters && !bForceRescan)
 	{
 		// This restriction is necessary because we have not yet implemented some of the required behavior to handle bIgnoreDenyListScanFilters without bForceRescan;
@@ -5494,6 +5502,7 @@ FScanPathContext::FScanPathContext(FEventContext& InEventContext, FClassInherita
 		// rather than just checking whether it has been set to be monitored at all
 		UE_LOG(LogAssetRegistry, Warning, TEXT("ScanPathsSynchronous: bIgnoreDenyListScanFilters==true is only valid when bForceRescan==true. Setting bForceRescan=true."));
 		bForceRescan = true;
+		bLogCallstack = true;
 	}
 
 	FString LocalPath;
@@ -5511,11 +5520,13 @@ FScanPathContext::FScanPathContext(FEventContext& InEventContext, FClassInherita
 		else if (!FPackageName::TryConvertToMountedPath(InFile, &LocalPath, &PackageName, nullptr, nullptr, &Extension, &FlexNameType))
 		{
 			UE_LOG(LogAssetRegistry, Warning, TEXT("ScanPathsSynchronous: %s is not in a mounted path, will not scan."), *InFile);
+			bLogCallstack = true;
 			continue;
 		}
 		if (FPackageName::IsTempPackage(PackageName))
 		{
 			UE_LOG(LogAssetRegistry, Warning, TEXT("ScanPathsSynchronous: %s is in the /Temp path, will not scan."), *InFile);
+			bLogCallstack = true;
 			continue;
 		}
 		if (Extension.IsEmpty())
@@ -5530,6 +5541,7 @@ FScanPathContext::FScanPathContext(FEventContext& InEventContext, FClassInherita
 					false /* bMatchCaseOnDisk */, &PackagePath) == FPackageName::EPackageLocationFilter::None)
 				{
 					UE_LOG(LogAssetRegistry, Warning, TEXT("ScanPathsSynchronous: Package %s does not exist, will not scan."), *InFile);
+					bLogCallstack = true;
 					continue;
 				}
 				Extension = PackagePath.GetExtensionString(EPackageSegment::Header);
@@ -5549,11 +5561,13 @@ FScanPathContext::FScanPathContext(FEventContext& InEventContext, FClassInherita
 		else if (!FPackageName::TryConvertToMountedPath(InDir, &LocalPath, &PackageName, nullptr, nullptr, &Extension, &FlexNameType))
 		{
 			UE_LOG(LogAssetRegistry, Warning, TEXT("ScanPathsSynchronous: %s is not in a mounted path, will not scan."), *InDir);
+			bLogCallstack = true;
 			continue;
 		}
 		if (FPackageName::IsTempPackage(PackageName))
 		{
 			UE_LOG(LogAssetRegistry, Warning, TEXT("ScanPathsSynchronous: %s is in the /Temp path, will not scan."), *InDir);
+			bLogCallstack = true;
 			continue;
 		}
 		LocalDirs.Add(LocalPath + Extension);
