Temporarily disable multithreaded data processing for Asset Registry.

diff --git a/Engine/Source/Runtime/AssetRegistry/Private/AssetDataGatherer.cpp b/Engine/Source/Runtime/AssetRegistry/Private/AssetDataGatherer.cpp
--- a/Engine/Source/Runtime/AssetRegistry/Private/AssetDataGatherer.cpp
+++ b/Engine/Source/Runtime/AssetRegistry/Private/AssetDataGatherer.cpp
@@ -85,7 +85,7 @@ static FAutoConsoleVariableRef CVarIgnoreEmptyDirectories(
 	bIgnoreEmptyDirectories,
 	TEXT("If true, completely empty leaf directories are ignored by the asset registry while scanning"));
 
-bool bTickGatherOnGTOnly = false;
+bool bTickGatherOnGTOnly = true;
 static FAutoConsoleVariableRef CVarTickGatherOnGTOnly(
 	TEXT("AssetRegistry.TickGatherOnGTOnly"),
 	bTickGatherOnGTOnly,
