AssetRegistry: Centralize all of the hard-coded classpath strings used by the AssetRegistry. Use hard-coded strings wherever we were previously using StaticClass()->GetClassPathName() in case our functions are called from early asyncthread scanning code before the modules are initialized.
#rnx

diff --git a/Engine/Source/Runtime/AssetRegistry/Private/AssetDataGatherer.cpp b/Engine/Source/Runtime/AssetRegistry/Private/AssetDataGatherer.cpp
--- a/Engine/Source/Runtime/AssetRegistry/Private/AssetDataGatherer.cpp
+++ b/Engine/Source/Runtime/AssetRegistry/Private/AssetDataGatherer.cpp
@@ -4463,7 +4463,7 @@ bool FAssetDataGatherer::ReadAssetFile(FPackageReader& PackageReader, TArray<FAs
 			// to save them as in-game imports by adding HasNonEditorOnlyReferences; the next version bump after that
 			// fix was VER_UE4_CORRECT_LICENSEE_FLAG. Mark all dependencies in the affected version as used in game
 			// if the package has a UObjectRedirector object.
-			FTopLevelAssetPath RedirectorClassPathName = UObjectRedirector::StaticClass()->GetClassPathName();
+			FTopLevelAssetPath RedirectorClassPathName = UE::AssetRegistry::GetClassPathObjectRedirector();
 			if (Algo::AnyOf(AssetDataList, [RedirectorClassPathName](FAssetData* AssetData) { return AssetData->AssetClassPath == RedirectorClassPathName; }))
 			{
 				for (FPackageDependencyData::FPackageDependency& Dependency : DependencyData.PackageDependencies)
diff --git a/Engine/Source/Runtime/AssetRegistry/Private/AssetRegistry.cpp b/Engine/Source/Runtime/AssetRegistry/Private/AssetRegistry.cpp
--- a/Engine/Source/Runtime/AssetRegistry/Private/AssetRegistry.cpp
+++ b/Engine/Source/Runtime/AssetRegistry/Private/AssetRegistry.cpp
@@ -1555,7 +1555,7 @@ void FAssetRegistryImpl::CollectCodeGeneratorClasses()
 	SavedGeneratorClassesVersionNumber = CurrentGeneratorClassesVersionNumber;
 
 	TArray<UClass*> BlueprintCoreDerivedClasses;
-	FTopLevelAssetPath BlueprintCorePathName(TEXT("/Script/Engine"), TEXT("BlueprintCore"));
+	FTopLevelAssetPath BlueprintCorePathName(GetClassPathBlueprintCore());
 	UClass* BlueprintCoreClass = nullptr;
 
 	{
@@ -2976,7 +2976,7 @@ bool UAssetRegistryImpl::DoesPackageExistOnDisk(FName PackageName, FString* OutC
 {
 	auto CalculateExtension = [](const FString& PackageNameStr, TConstArrayView<FAssetData> Assets) -> FString
 	{
-		FTopLevelAssetPath ClassRedirector = UObjectRedirector::StaticClass()->GetClassPathName();
+		FTopLevelAssetPath ClassRedirector = UE::AssetRegistry::GetClassPathObjectRedirector();
 		bool bContainsMap = false;
 		bool bContainsRedirector = false;
 		for (const FAssetData& Asset : Assets)
@@ -3630,7 +3630,7 @@ PRAGMA_ENABLE_DEPRECATION_WARNINGS
 		if (InFilter.RecursiveClassPathsExclusionSet.Num() > 0 && InFilter.ClassPaths.Num() == 0)
 		{
 			TArray<FTopLevelAssetPath> ClassNamesObject;
-			ClassNamesObject.Add(UObject::StaticClass()->GetClassPathName());
+			ClassNamesObject.Add(UE::AssetRegistry::GetClassPathObject());
 
 			GetSubClasses(InheritanceContext, ClassNamesObject, InFilter.RecursiveClassPathsExclusionSet, OutCompiledFilter.ClassPaths);
 		}
@@ -6296,7 +6296,12 @@ void FAssetRegistryImpl::DependencyDataGathered(TMultiMap<FName, FPackageDepende
 			Node->SetIsReferencersSorted(ShouldSortReferencers());
 
 			// Don't bother registering dependencies on these packages, every package in the game will depend on them
-			static TArray<FName> ScriptPackagesToSkip = TArray<FName>{ TEXT("/Script/CoreUObject"), TEXT("/Script/Engine"), TEXT("/Script/BlueprintGraph"), TEXT("/Script/UnrealEd") };
+			static TArray<FName> ScriptPackagesToSkip = TArray<FName>{
+				UE::AssetRegistry::GetScriptPackageNameCoreUObject(),
+				UE::AssetRegistry::GetScriptPackageNameEngine(),
+				UE::AssetRegistry::GetScriptPackageNameBlueprintGraph(),
+				UE::AssetRegistry::GetScriptPackageNameUnrealEd(),
+			};
 
 			// Conditionally add package dependencies
 			TMap<FName, FDependsNode::FPackageFlagSet> PackageDependencies;
@@ -7294,7 +7299,7 @@ void FAssetRegistryImpl::PushProcessLoadedAssetsBatch(Impl::FEventContext& Event
 void FAssetRegistryImpl::UpdateRedirectCollector()
 {
 	// Look for all redirectors in list
-	const TArray<const FAssetData*>& RedirectorAssets = State.GetAssetsByClassPathName(UObjectRedirector::StaticClass()->GetClassPathName());
+	const TArray<const FAssetData*>& RedirectorAssets = State.GetAssetsByClassPathName(UE::AssetRegistry::GetClassPathObjectRedirector());
 
 	for (const FAssetData* AssetData : RedirectorAssets)
 	{
@@ -7716,7 +7721,7 @@ void FAssetRegistryImpl::UpdateInheritanceBuffer(Impl::FClassInheritanceBuffer&
 	}
 	OutBuffer.InheritanceMap.Reserve(NumNativeClasses + CachedBPInheritanceMap.Num());
 	OutBuffer.InheritanceMap = CachedBPInheritanceMap;
-	OutBuffer.InheritanceMap.Add(FTopLevelAssetPath(TEXT("/Script.CoreUObject"), TEXT("Object")), FTopLevelAssetPath());
+	OutBuffer.InheritanceMap.Add(UE::AssetRegistry::GetClassPathObject(), FTopLevelAssetPath());
 
 	for (TPair<FTopLevelAssetPath, TArray<FTopLevelAssetPath>>& Pair : OutBuffer.ReverseInheritanceMap)
 	{
diff --git a/Engine/Source/Runtime/CoreUObject/Private/AssetRegistry/AssetData.cpp b/Engine/Source/Runtime/CoreUObject/Private/AssetRegistry/AssetData.cpp
--- a/Engine/Source/Runtime/CoreUObject/Private/AssetRegistry/AssetData.cpp
+++ b/Engine/Source/Runtime/CoreUObject/Private/AssetRegistry/AssetData.cpp
@@ -11,6 +11,7 @@
 #include "HAL/CriticalSection.h"
 #include "HAL/PlatformMath.h"
 #include "Misc/AsciiSet.h"
+#include "Misc/AssetRegistryInterface.h"
 #include "Misc/PathViews.h"
 #include "Misc/ScopeRWLock.h"
 #include "Serialization/CompactBinary.h"
@@ -818,8 +819,7 @@ void FAssetData::SerializeForCacheOldVersionWithTagsAndBundles(FArchive& Ar, FAs
 
 bool FAssetData::IsRedirectorClassName(FTopLevelAssetPath ClassPathName)
 {
-	static const FTopLevelAssetPath ObjectRedirectorClassPathName = UObjectRedirector::StaticClass()->GetClassPathName();
-	return ClassPathName == ObjectRedirectorClassPathName;
+	return ClassPathName == UE::AssetRegistry::GetClassPathObjectRedirector();
 }
 
 FTopLevelAssetPath FAssetData::TryConvertShortClassNameToPathName(FName InClassName, ELogVerbosity::Type FailureMessageVerbosity /*= ELogVerbosity::Warning*/)
diff --git a/Engine/Source/Runtime/CoreUObject/Private/Misc/AssetRegistryInterface.cpp b/Engine/Source/Runtime/CoreUObject/Private/Misc/AssetRegistryInterface.cpp
--- a/Engine/Source/Runtime/CoreUObject/Private/Misc/AssetRegistryInterface.cpp
+++ b/Engine/Source/Runtime/CoreUObject/Private/Misc/AssetRegistryInterface.cpp
@@ -76,12 +76,12 @@ namespace Utils
 
 	void PopulateSkipClasses(TSet<FTopLevelAssetPath>& OutSkipUncookedClasses, TSet<FTopLevelAssetPath>& OutSkipCookedClasses)
 	{
-		static const FName NAME_EnginePackage("/Script/Engine");
+		const FName NAME_EnginePackage(GetScriptPackageNameEngine());
 		UPackage* EnginePackage = Cast<UPackage>(StaticFindObjectFast(UPackage::StaticClass(), nullptr, NAME_EnginePackage));
 		{
 			OutSkipUncookedClasses.Reset();
 
-			static const FName NAME_BlueprintGeneratedClass("BlueprintGeneratedClass");
+			const FName NAME_BlueprintGeneratedClass(GetClassNameBlueprintGeneratedClass());
 			UClass* BlueprintGeneratedClass = nullptr;
 			if (EnginePackage)
 			{
@@ -106,7 +106,7 @@ namespace Utils
 		{
 			OutSkipCookedClasses.Reset();
 
-			static const FName NAME_Blueprint("Blueprint");
+			const FName NAME_Blueprint(GetClassNameBlueprint());
 			UClass* BlueprintClass = nullptr;
 			if (EnginePackage)
 			{
@@ -178,4 +178,72 @@ namespace Utils
 		bInitializedSkipClasses = false;
 #endif
 	}
-}
+
+	static FName ScriptPackageNameCoreUObject(TEXT("/Script/CoreUObject"));
+	static FName ScriptPackageNameEngine(TEXT("/Script/Engine"));
+	static FName ScriptPackageNameBlueprintGraph(TEXT("/Script/BlueprintGraph"));
+	static FName ScriptPackageNameUnrealEd(TEXT("/Script/UnrealEd"));
+	static FName ClassNameObject(TEXT("Object"));
+	static FName ClassNameObjectRedirector(TEXT("ObjectRedirector"));
+	static FName ClassNameBlueprintCore(TEXT("BlueprintCore"));
+	static FName ClassNameBlueprint(TEXT("Blueprint"));
+	static FName ClassNameBlueprintGeneratedClass(TEXT("BlueprintGeneratedClass"));
+
+	FName GetScriptPackageNameCoreUObject()
+	{
+		return ScriptPackageNameCoreUObject;
+	}
+	FName GetScriptPackageNameEngine()
+	{
+		return ScriptPackageNameEngine;
+	}
+	FName GetScriptPackageNameBlueprintGraph()
+	{
+		return ScriptPackageNameBlueprintGraph;
+	}
+	FName GetScriptPackageNameUnrealEd()
+	{
+		return ScriptPackageNameUnrealEd;
+	}
+	FName GetClassNameObject()
+	{
+		return ClassNameObject;
+	}
+	FName GetClassNameObjectRedirector()
+	{
+		return ClassNameObjectRedirector;
+	}
+	FName GetClassNameBlueprintCore()
+	{
+		return ClassNameBlueprintCore;
+	}
+	FName GetClassNameBlueprint()
+	{
+		return ClassNameBlueprint;
+	}
+	FName GetClassNameBlueprintGeneratedClass()
+	{
+		return ClassNameBlueprintGeneratedClass;
+	}
+	FTopLevelAssetPath GetClassPathObject()
+	{
+		return FTopLevelAssetPath(GetScriptPackageNameCoreUObject(), GetClassNameObject());
+	}
+	FTopLevelAssetPath GetClassPathObjectRedirector()
+	{
+		return FTopLevelAssetPath(GetScriptPackageNameCoreUObject(), GetClassNameObjectRedirector());
+	}
+	FTopLevelAssetPath GetClassPathBlueprintCore()
+	{
+		return FTopLevelAssetPath(GetScriptPackageNameEngine(), GetClassNameBlueprintCore());
+	}
+	FTopLevelAssetPath GetClassPathBlueprint()
+	{
+		return FTopLevelAssetPath(GetScriptPackageNameEngine(), GetClassNameBlueprint());
+	}
+	FTopLevelAssetPath GetClassPathBlueprintGeneratedClass()
+	{
+		return FTopLevelAssetPath(GetScriptPackageNameEngine(), GetClassNameBlueprintGeneratedClass());
+	}
+
+} // namespace UE::AssetRegistry
diff --git a/Engine/Source/Runtime/CoreUObject/Public/Misc/AssetRegistryInterface.h b/Engine/Source/Runtime/CoreUObject/Public/Misc/AssetRegistryInterface.h
--- a/Engine/Source/Runtime/CoreUObject/Public/Misc/AssetRegistryInterface.h
+++ b/Engine/Source/Runtime/CoreUObject/Public/Misc/AssetRegistryInterface.h
@@ -253,6 +253,21 @@ namespace Utils
 }
 #endif
 
+	COREUOBJECT_API FName GetScriptPackageNameCoreUObject();
+	COREUOBJECT_API FName GetScriptPackageNameEngine();
+	COREUOBJECT_API FName GetScriptPackageNameBlueprintGraph();
+	COREUOBJECT_API FName GetScriptPackageNameUnrealEd();
+	COREUOBJECT_API FName GetClassNameObject();
+	COREUOBJECT_API FName GetClassNameObjectRedirector();
+	COREUOBJECT_API FName GetClassNameBlueprintCore();
+	COREUOBJECT_API FName GetClassNameBlueprint();
+	COREUOBJECT_API FName GetClassNameBlueprintGeneratedClass();
+	COREUOBJECT_API FTopLevelAssetPath GetClassPathObject();
+	COREUOBJECT_API FTopLevelAssetPath GetClassPathObjectRedirector();
+	COREUOBJECT_API FTopLevelAssetPath GetClassPathBlueprintCore();
+	COREUOBJECT_API FTopLevelAssetPath GetClassPathBlueprint();
+	COREUOBJECT_API FTopLevelAssetPath GetClassPathBlueprintGeneratedClass();
+
 }
 
 // Enums used in public Engine headers
