Extra debugging info for corrupt name maps in assets.
Use localized structured logs for these messages.

diff --git a/Engine/Source/Runtime/AssetRegistry/Private/PackageReader.cpp b/Engine/Source/Runtime/AssetRegistry/Private/PackageReader.cpp
--- a/Engine/Source/Runtime/AssetRegistry/Private/PackageReader.cpp
+++ b/Engine/Source/Runtime/AssetRegistry/Private/PackageReader.cpp
@@ -9,6 +9,7 @@
 #include "Internationalization/Internationalization.h"
 #include "Internationalization/GatherableTextData.h"
 #include "Logging/MessageLog.h"
+#include "Logging/StructuredLog.h"
 #include "Misc/Guid.h"
 #include "Misc/PackageName.h"
 #include "Misc/PackagePath.h"
@@ -19,6 +20,8 @@
 #include "UObject/PackageRelocation.h"
 #include "UObject/PackageTrailer.h"
 
+#define LOCTEXT_NAMESPACE "AssetRegistry"
+
 const TCHAR* LexToString(FPackageReader::EOpenPackageResult Result)
 {
 	switch(Result)
@@ -739,14 +742,18 @@ bool FPackageReader::SerializeNameMap()
 	{
 		if (!StartSerializeSection(PackageFileSummary.NameOffset))
 		{
-			UE_PACKAGEREADER_CORRUPTPACKAGE_WARNING("SerializeNameMapInvalidNameOffset", PackageFilename);
+			UE_LOGFMT_LOC(LogAssetRegistry, Warning, "SerializeNameMapInvalidNameOffset",
+				"Package is unloadable: {FileName}. Reason: Failed to seek to name table offset {NameOffset} in package of size {PackageSize}",
+				("FileName", PackageFilename), ("NameOffset", PackageFileSummary.NameOffset), ("PackageSize", PackageFileSize));
 			return false;
 		}
 
 		const int MinSizePerNameEntry = 1;
 		if (PackageFileSize < Tell() + PackageFileSummary.NameCount * MinSizePerNameEntry)
 		{
-			UE_PACKAGEREADER_CORRUPTPACKAGE_WARNING("SerializeNameMapInvalidNameCount", PackageFilename);
+			UE_LOGFMT_LOC(LogAssetRegistry, Warning, "SerializeNameMapInvalidNameCount",
+				"Package is unloadable: {FileName}. Reason: Name table count {NameCount} in package of size {PackageSize} at name offset {NameOffset}",
+				("FileName", PackageFilename), ("NameCount", PackageFileSummary.NameCount), ("NameOffset", PackageFileSummary.NameOffset), ("PackageSize", PackageFileSize));
 			return false;
 		}
 
@@ -757,7 +764,9 @@ bool FPackageReader::SerializeNameMap()
 			*this << NameEntry;
 			if (IsError())
 			{
-				UE_PACKAGEREADER_CORRUPTPACKAGE_WARNING("SerializeNameMapInvalidName", PackageFilename);
+				UE_LOGFMT_LOC(LogAssetRegistry, Warning, "SerializeNameMapInvalidName",
+					"Package is unloadable: {FileName}. Reason: Invalid name at index {NameIndex}",
+					("FileName", PackageFilename), ("NameIndex", NameMapIdx));
 				NameMap.Reset();
 				return false;
 			}
@@ -1760,3 +1769,5 @@ namespace UE::AssetRegistry
 		return !BinaryArchive.IsError();
 	}
 }
+
+#undef LOC_NAMESPACE
\ No newline at end of file
