Fix localization namepsace undef

diff --git a/Engine/Source/Runtime/AssetRegistry/Private/PackageReader.cpp b/Engine/Source/Runtime/AssetRegistry/Private/PackageReader.cpp
--- a/Engine/Source/Runtime/AssetRegistry/Private/PackageReader.cpp
+++ b/Engine/Source/Runtime/AssetRegistry/Private/PackageReader.cpp
@@ -1770,4 +1770,4 @@ namespace UE::AssetRegistry
 	}
 }
 
-#undef LOC_NAMESPACE
\ No newline at end of file
+#undef LOCTEXT_NAMESPACE
\ No newline at end of file
