AssetRegistry: AssetSearchDataGathered: Reorganize the function to remove duplicated code between the update and add paths.
#rnx

diff --git a/Engine/Source/Runtime/AssetRegistry/Private/AssetRegistry.cpp b/Engine/Source/Runtime/AssetRegistry/Private/AssetRegistry.cpp
--- a/Engine/Source/Runtime/AssetRegistry/Private/AssetRegistry.cpp
+++ b/Engine/Source/Runtime/AssetRegistry/Private/AssetRegistry.cpp
@@ -6100,96 +6100,82 @@ void FAssetRegistryImpl::AssetSearchDataGathered(Impl::FEventContext& EventConte
 		CA_ASSUME(BackgroundResult.Get() != nullptr);
 		Iter.RemoveCurrent();
 
-		// Try to update any asset data that may already exist
-		FCachedAssetKey Key(*BackgroundResult);
-
+		// Skip assets that are invalid because e.g. they are externalactors that were mistakenly not deleted
+		// when their map moved.
 		if (ShouldSkipGatheredAsset(*BackgroundResult))
 		{
 			continue;
 		}
 
-		FAssetData* const* FoundData = State.CachedAssets.Find(Key);
-		FAssetData* ExistingAssetData = FoundData ? *FoundData : nullptr;
-
+		// Skip stale gather results from unmounted roots caused by mount then unmount of a path within short period.
 		const FName PackagePath = BackgroundResult->PackagePath;
-
-		// Skip stale results caused by mount then unmount of a path within short period.
-		bool bPathIsMounted = true;
 		if (bVerifyMountPointAfterGather)
 		{
 			PackagePath.ToString(PackagePathString);
 			if (!Utils::IsPathMounted(PackagePathString, MountPoints, PackageRoot))
 			{
-				bPathIsMounted = false;
+				UE_LOG(LogAssetRegistry, Warning,
+					TEXT("AssetRegistry: An asset has been loaded with an invalid mount point: '%s', Mount Point: '%s'. Ignoring the asset."),
+					*BackgroundResult->GetObjectPathString(), *PackagePathString)
+				continue;
 			}
 		}
 
+#if WITH_EDITOR
+		// Postload assets based on their declared class. Queue them for for later retry if their class has not yet loaded.
+		bool CouldPostLoad = TryPostLoadAssetRegistryTags(BackgroundResult.Get());
+		if (!CouldPostLoad)
+		{
+			OutDeferredAssetResults.Add(BackgroundAssetPackageName, MoveTemp(BackgroundResult));
+			continue;
+		}
+#endif
+
+		// Look for an existing asset to check whether we need to add or update
+		FCachedAssetKey Key(*BackgroundResult);
+		FAssetData* const* FoundData = State.CachedAssets.Find(Key);
+		FAssetData* ExistingAssetData = FoundData ? *FoundData : nullptr;
+		// The background result should not already be registered; it should be impossible since it is in TUnqiuePtr
+		check(ExistingAssetData == nullptr || ExistingAssetData != BackgroundResult.Get());
+
 		if (ExistingAssetData)
 		{
-			// If this ensure fires then we've somehow processed the same result more than once, and that should never happen
-			if (ensure(ExistingAssetData != BackgroundResult.Get()))
-			{
 #if WITH_EDITOR
-				bool CouldPostLoad = TryPostLoadAssetRegistryTags(BackgroundResult.Get());
-				if (!CouldPostLoad)
-				{
-					OutDeferredAssetResults.Add(BackgroundAssetPackageName, MoveTemp(BackgroundResult));
-					continue;
-				}
-				if (AssetDataObjectPathsUpdatedOnLoad.Contains(BackgroundResult->GetSoftObjectPath()))
-				{
-					// If the current AssetData came from a loaded asset, don't overwrite it with the new one from disk
-					// The loaded asset is more authoritative because it has run the postload steps.
-					// However, the loaded asset is missing the extended tags. Our contract for extended tags is to keep any 
-					// that do not exist in the non-extended tags. So add on any tags from the BackgroundResult that
-					// are not already on the existing asset.
-					AddNonOverlappingTags(EventContext, *ExistingAssetData, *BackgroundResult);
-				}
-				else
+			if (AssetDataObjectPathsUpdatedOnLoad.Contains(BackgroundResult->GetSoftObjectPath()))
+			{
+				// If the current AssetData came from a loaded asset, don't overwrite it with the new one from disk
+				// The loaded asset is more authoritative because it has run the postload steps.
+				// However, the loaded asset is missing the extended tags. Our contract for extended tags is to keep any 
+				// that do not exist in the non-extended tags. So add on any tags from the BackgroundResult that
+				// are not already on the existing asset.
+				AddNonOverlappingTags(EventContext, *ExistingAssetData, *BackgroundResult);
+			}
+			else
 #endif
-				{
-					// The asset exists in the cache from disk and has not yet been loaded into memory, update it with the new background data
-					UpdateAssetData(EventContext, ExistingAssetData, MoveTemp(*BackgroundResult), false /* bKeepDeletedTags */);
-				}
+			{
+				// The asset exists in the cache from disk and has not yet been loaded into memory, update it with the new background data
+				UpdateAssetData(EventContext, ExistingAssetData, MoveTemp(*BackgroundResult), false /* bKeepDeletedTags */);
 			}
 		}
 		else
 		{
 			// The asset isn't in the cache yet, add it and notify subscribers
-			if (bPathIsMounted)
-			{
 #if !NO_LOGGING
-				if (bVerboseLogging)
-				{
-					int64& ClassTagSizes = TagSizeByClass.FindOrAdd(BackgroundResult->AssetClassPath);
-					BackgroundResult->TagsAndValues.ForEach([&ClassTagSizes](const TPair<FName, FAssetTagValueRef>& Pair)
-						{
-							ClassTagSizes += Pair.Value.GetResourceSize();
-						});
-				}
+			if (bVerboseLogging)
+			{
+				int64& ClassTagSizes = TagSizeByClass.FindOrAdd(BackgroundResult->AssetClassPath);
+				BackgroundResult->TagsAndValues.ForEach([&ClassTagSizes](const TPair<FName, FAssetTagValueRef>& Pair)
+					{
+						ClassTagSizes += Pair.Value.GetResourceSize();
+					});
+			}
 #endif
 
-#if WITH_EDITOR
-				bool CouldPostLoad = TryPostLoadAssetRegistryTags(BackgroundResult.Get());
-				if (!CouldPostLoad)
-				{
-					OutDeferredAssetResults.Add(BackgroundAssetPackageName, MoveTemp(BackgroundResult));
-					continue;
-				}
-#endif
-				AddAssetData(EventContext, BackgroundResult.Release());
-			}
+			AddAssetData(EventContext, BackgroundResult.Release());
 		}
 
-		if (bPathIsMounted)
-		{
-			// Populate the path tree
-			AddAssetPath(EventContext, PackagePath);
-		}
-		else
-		{
-			UE_LOG(LogAssetRegistry, Warning, TEXT("AssetRegistry: An asset has been loaded with an invalid mount point: '%s', Mount Point: '%s'"), *BackgroundResult->GetObjectPathString(), *PackagePathString)
-		}
+		// Populate the path tree
+		AddAssetPath(EventContext, PackagePath);
 	}
 }
 
