Fixing a race condition in which the AssetRegistry singleton isn't initialized before the asset gatherer thread begins to run and accesses it.

diff --git a/Engine/Source/Runtime/AssetRegistry/Private/AssetDataGatherer.cpp b/Engine/Source/Runtime/AssetRegistry/Private/AssetDataGatherer.cpp
--- a/Engine/Source/Runtime/AssetRegistry/Private/AssetDataGatherer.cpp
+++ b/Engine/Source/Runtime/AssetRegistry/Private/AssetDataGatherer.cpp
@@ -3729,7 +3729,8 @@ uint32 FAssetDataGatherer::Run()
 			if (bLocalIdle && !UE::AssetDataGather::Private::bTickGatherOnGTOnly 
 				&& (IsProcessingPaused.load(std::memory_order_relaxed) == 0))
 			{
-				Status = Cast<UAssetRegistryImpl>(IAssetRegistry::Get())->TickOnBackgroundThread();
+				IAssetRegistry& Registry = IAssetRegistry::GetChecked();
+				Status = Cast<UAssetRegistryImpl>(&Registry)->TickOnBackgroundThread();
 			}
 
 			// TODO: Need IsGatheringPaused to be a condition variable so we avoid sleeping while waiting for it and then taking a long time to wake after it is unset.
