AssetRegistry: ScanPathsSynchronous handling of deleted files:
1) Fix OnAssetRemoved event not being sent when ScanPathsSynchronous finds deleted packages.
2) Fix verse files not being removed from the AssetRegistry at all AND OnVerseRemoved not being sent when ScanPathsSynchronous finds deleted verse files.
#rnx

diff --git a/Engine/Source/Runtime/AssetRegistry/Private/AssetDataGatherer.cpp b/Engine/Source/Runtime/AssetRegistry/Private/AssetDataGatherer.cpp
--- a/Engine/Source/Runtime/AssetRegistry/Private/AssetDataGatherer.cpp
+++ b/Engine/Source/Runtime/AssetRegistry/Private/AssetDataGatherer.cpp
@@ -5393,9 +5393,22 @@ bool FAssetDataGatherer::IsMonitored(FStringView LocalPath) const
 	return Discovery->IsMonitored(NormalizeLocalPath(LocalPath));
 }
 
+static const TCHAR* VerseExtensions[] = { TEXT(".verse"), TEXT(".vmodule") };
 bool FAssetDataGatherer::IsVerseFile(FStringView FilePath)
 {
-	return FilePath.EndsWith(TEXT(".verse")) || FilePath.EndsWith(TEXT(".vmodule"));
+	for (const TCHAR* Extension : VerseExtensions)
+	{
+		if (FilePath.EndsWith(Extension))
+		{
+			return true;
+		}
+	}
+	return false;
+}
+
+TConstArrayView<const TCHAR*> FAssetDataGatherer::GetVerseFileExtensions()
+{
+	return TConstArrayView<const TCHAR*>(VerseExtensions, UE_ARRAY_COUNT(VerseExtensions));
 }
 
 void FAssetDataGatherer::SetIsIdle(bool bInIsIdle)
diff --git a/Engine/Source/Runtime/AssetRegistry/Private/AssetDataGatherer.h b/Engine/Source/Runtime/AssetRegistry/Private/AssetDataGatherer.h
--- a/Engine/Source/Runtime/AssetRegistry/Private/AssetDataGatherer.h
+++ b/Engine/Source/Runtime/AssetRegistry/Private/AssetDataGatherer.h
@@ -237,6 +237,8 @@ public:
 
 	/** Determine, based on the file extension, if the given file path is a Verse file */
 	static bool IsVerseFile(FStringView FilePath);
+	/** Return the list of extensions that indicate verse files. */
+	static TConstArrayView<const TCHAR*> GetVerseFileExtensions();
 
 	/**
 	 * Reads FAssetData information out of a previously initialized package reader
diff --git a/Engine/Source/Runtime/AssetRegistry/Private/AssetRegistry.cpp b/Engine/Source/Runtime/AssetRegistry/Private/AssetRegistry.cpp
--- a/Engine/Source/Runtime/AssetRegistry/Private/AssetRegistry.cpp
+++ b/Engine/Source/Runtime/AssetRegistry/Private/AssetRegistry.cpp
@@ -4519,7 +4519,7 @@ bool FAssetRegistryImpl::ClassRequiresGameThreadProcessing(const UClass* Class)
 
 Impl::EGatherStatus FAssetRegistryImpl::TickGatherer(Impl::FEventContext& EventContext,
 	Impl::FClassInheritanceContext& InheritanceContext, Impl::FInterruptionContext& InOutInterruptionContext,
-	TOptional<FAssetsFoundCallback> AssetsFoundCallback)
+	TOptional<FAssetsFoundCallback> AssetsFoundCallback, TOptional<FVerseFilesFoundCallback> VerseFilesFoundCallback)
 {
 	using namespace UE::AssetRegistry::Impl;
 	bool bLocalIsInGameThread = IsInGameThread();
@@ -4773,6 +4773,11 @@ Impl::EGatherStatus FAssetRegistryImpl::TickGatherer(Impl::FEventContext& EventC
 	if (BackgroundResults.VerseFiles.Num())
 	{
 		LazyStartTimer();
+		if (VerseFilesFoundCallback.IsSet())
+		{
+			VerseFilesFoundCallback.GetValue()(BackgroundResults.VerseFiles);
+		}
+
 		VerseFilesGathered(EventContext, BackgroundResults.VerseFiles, InOutInterruptionContext);
 		if (InOutInterruptionContext.ShouldExitEarly())
 		{
@@ -5618,6 +5623,7 @@ void FAssetRegistryImpl::ScanPathsSynchronous(Impl::FScanPathContext& Context)
 	// If we are forcing a rescan, then delete any old assets that no longer exist. If we are not forcing a rescan,
 	// then there should not be any old assets that no longer exist, so we skip the cost of searching for them.
 	TSet<FSoftObjectPath> OldAssetsToRemove;
+	TSet<FName> OldVerseFilesToRemove;
 	if (Context.bForceRescan)
 	{
 		// Initialize OldAssetsToRemove to the list of all assets in the given paths.
@@ -5639,6 +5645,14 @@ void FAssetRegistryImpl::ScanPathsSynchronous(Impl::FScanPathContext& Context)
 					OldAssetsToRemove.Add(AssetData.ToSoftObjectPath());
 					return true;
 				});
+			for (FName PackagePath : CompiledFilter.PackagePaths)
+			{
+				TArray<FName>* VerseFiles = CachedVerseFilesByPath.Find(PackagePath);
+				if (VerseFiles)
+				{
+					OldVerseFilesToRemove.Append(*VerseFiles);
+				}
+			}
 		}
 		for (const FString& PackageName : Context.PackageFiles)
 		{
@@ -5646,6 +5660,14 @@ void FAssetRegistryImpl::ScanPathsSynchronous(Impl::FScanPathContext& Context)
 			{
 				OldAssetsToRemove.Add(AssetData->ToSoftObjectPath());
 			}
+			for (const TCHAR* Extension : FAssetDataGatherer::GetVerseFileExtensions())
+			{
+				FName VerseName(*WriteToString<256>(PackageName, Extension), FNAME_Find);
+				if (!VerseName.IsNone() && CachedVerseFiles.Contains(VerseName))
+				{
+					OldVerseFilesToRemove.Add(VerseName);
+				}
+			}
 		}
 	}
 
@@ -5718,10 +5740,22 @@ void FAssetRegistryImpl::ScanPathsSynchronous(Impl::FScanPathContext& Context)
 			}
 		}
 	};
+	auto VerseFileFoundCallback =
+		[&OldVerseFilesToRemove]
+		(const TRingBuffer<FName>& InFoundVerseFiles)
+		{
+			if (!OldVerseFilesToRemove.IsEmpty())
+			{
+				for (const FName VerseFile : InFoundVerseFiles)
+				{
+					OldVerseFilesToRemove.Remove(VerseFile);
+				}
+			}
+		};
 
 	Impl::FInterruptionContext InterruptionContext(-1., -1.);
 	Context.Status = TickGatherer(Context.EventContext, Context.InheritanceContext,	InterruptionContext,
-		FAssetsFoundCallback(AssetsFoundCallback));
+		FAssetsFoundCallback(AssetsFoundCallback), FVerseFilesFoundCallback(VerseFileFoundCallback));
 
 	// Temporary hack/partial solution. The expectation is that this function will return cause all assets
 	// under the specified directories to be ingested into the registry. However, one of the early steps
@@ -5776,10 +5810,15 @@ void FAssetRegistryImpl::ScanPathsSynchronous(Impl::FScanPathContext& Context)
 #endif
 	for (FSoftObjectPath& OldAssetToRemove : OldAssetsToRemove)
 	{
-		bool bOutRemovedAssetData;
-		bool bOutRemovedPackageData;
-		State.RemoveAssetData(OldAssetToRemove, true /* bRemoveDependencyData */,
-			bOutRemovedAssetData, bOutRemovedPackageData);
+		FAssetData* AssetDataToRemove = const_cast<FAssetData*>(State.GetAssetByObjectPath(OldAssetToRemove));
+		if (AssetDataToRemove)
+		{
+			RemoveAssetData(Context.EventContext, AssetDataToRemove);
+		}
+	}
+	for (FName OldVerseFileToRemove : OldVerseFilesToRemove)
+	{
+		RemoveVerseFile(Context.EventContext, OldVerseFileToRemove);
 	}
 }
 
diff --git a/Engine/Source/Runtime/AssetRegistry/Private/AssetRegistryImpl.h b/Engine/Source/Runtime/AssetRegistry/Private/AssetRegistryImpl.h
--- a/Engine/Source/Runtime/AssetRegistry/Private/AssetRegistryImpl.h
+++ b/Engine/Source/Runtime/AssetRegistry/Private/AssetRegistryImpl.h
@@ -269,12 +269,14 @@ public:
 	void WaitForGathererIdleIfSynchronous();
 	/** Waits for the gatherer to be idle. */
 	void WaitForGathererIdle(float TimeoutSeconds);
-	/** Callback type for TickGatherer */
+	/** Callback types for TickGatherer */
 	typedef TFunctionRef<void(const TMultiMap<FName, FAssetData*>&)> FAssetsFoundCallback;
+	typedef TFunctionRef<void(const TRingBuffer<FName>&)> FVerseFilesFoundCallback;
 	/** Consume any results from the gatherer and return its status */
 	Impl::EGatherStatus TickGatherer(Impl::FEventContext& EventContext,
 		Impl::FClassInheritanceContext& InheritanceContext, Impl::FInterruptionContext& InOutInterruptionContext,
-		TOptional<FAssetsFoundCallback> AssetsFoundCallback = TOptional<FAssetsFoundCallback>());
+		TOptional<FAssetsFoundCallback> AssetsFoundCallback = TOptional<FAssetsFoundCallback>(),
+		TOptional<FVerseFilesFoundCallback> VerseFilesFoundCallback = TOptional<FVerseFilesFoundCallback>());
 
 	/** Send a log message with the search statistics. 
 	 *  StartTime is used to report wall clock search time in the case of background scan
