AssetRegistry: Gathering of ExternalActors:
Add a warning when two packages contain the same Actor from a map (the outerpath of the Actors is the same); this can be caused by moving the actor between content bundles and is an error state. Resolve the collision by picking the package with a later timestamp.
#rnx

diff --git a/Engine/Source/Runtime/AssetRegistry/Private/AssetRegistry.cpp b/Engine/Source/Runtime/AssetRegistry/Private/AssetRegistry.cpp
--- a/Engine/Source/Runtime/AssetRegistry/Private/AssetRegistry.cpp
+++ b/Engine/Source/Runtime/AssetRegistry/Private/AssetRegistry.cpp
@@ -5851,6 +5851,29 @@ bool IsPathMounted(const FString& Path, const TSet<FString>& MountPointsNoTraili
 }
 
 #if WITH_EDITOR
+
+FAssetData* FAssetRegistryImpl::ResolveAssetIdCollision(FAssetData& A, FAssetData& B)
+{
+	// We could use file age to try to guess which file is correct:
+	// FPackageName::InternalDoesPackageExistEx() to get the filename, and IFileManager::GetFileAgeSeconds
+	// But that would vary from machine to machine based on when the files were synced.
+	// So instead just pick one using an arbitrary deterministic process: alphabetical order
+	FAssetData* Keep = A.PackageName.LexicalLess(B.PackageName) ? &A : &B;
+	FAssetData* Discard = Keep == &A ? &B : &A;
+
+	FString PackageNameB = B.PackageName.ToString();
+	FString FileNameA;
+	FString FileNameB;
+	UE_LOG(LogAssetRegistry, Warning, TEXT("Invalid duplicate copies of ExternalActor %s. Resolve by deleting the package that is invalid. Chosing alphabetically for this process.")
+		TEXT("\n\tDiscarding: %s")
+		TEXT("\n\tKeeping:    %s"),
+		*Keep->GetObjectPathString(),
+		*Discard->PackageName.ToString(),
+		*Keep->PackageName.ToString());
+
+	return Keep;
+}
+
 bool FAssetRegistryImpl::TryPostLoadAssetRegistryTags(FAssetData* AssetData)
 {
 	check(AssetData);
@@ -6177,6 +6200,26 @@ void FAssetRegistryImpl::AssetSearchDataGathered(Impl::FEventContext& EventConte
 		// The background result should not already be registered; it should be impossible since it is in TUnqiuePtr
 		check(ExistingAssetData == nullptr || ExistingAssetData != BackgroundResult.Get());
 
+#if WITH_EDITOR
+		if (ExistingAssetData && ExistingAssetData->PackageName != BackgroundResult->PackageName)
+		{
+			// This can happen with ExternalActors, which have a Key based on their outermost map, but 
+			// are in a separate package. It's invalid to have more than one of them, but can happen when
+			// actors are moved between packages if the delete is not recorded.
+			FAssetData* PackageToKeep = ResolveAssetIdCollision(*ExistingAssetData, *BackgroundResult);
+			if (PackageToKeep == ExistingAssetData)
+			{
+				continue;
+			}
+			else
+			{
+				check(PackageToKeep == BackgroundResult.Get());
+				RemoveAssetData(EventContext, ExistingAssetData);
+				ExistingAssetData = nullptr;
+			}
+		}
+#endif
+
 		if (ExistingAssetData)
 		{
 #if WITH_EDITOR
diff --git a/Engine/Source/Runtime/AssetRegistry/Private/AssetRegistryImpl.h b/Engine/Source/Runtime/AssetRegistry/Private/AssetRegistryImpl.h
--- a/Engine/Source/Runtime/AssetRegistry/Private/AssetRegistryImpl.h
+++ b/Engine/Source/Runtime/AssetRegistry/Private/AssetRegistryImpl.h
@@ -447,7 +447,9 @@ private:
 	bool RemoveAssetData(Impl::FEventContext& EventContext, FAssetData* AssetData);
 
 #if WITH_EDITOR
-	/** 
+	/** Called when two assets have the same ID but different packages. Logs it and returns which one to keep. */
+	FAssetData* ResolveAssetIdCollision(FAssetData& A, FAssetData& B);
+	/**
 	 * Calls PostLoadAssetRegistryTags on the CDO of the asset class this data represents.
 	 * @param AssetData Existing asset data
 	 * @return Returns false if the required parent UClass is not yet available
