AssetRegistry - Optimize unmount and delete asset by converting a TArray to TSet

diff --git a/Engine/Plugins/Editor/AssetRegistryExport/Source/Private/AssetRegistryExportCommandlet.cpp b/Engine/Plugins/Editor/AssetRegistryExport/Source/Private/AssetRegistryExportCommandlet.cpp
--- a/Engine/Plugins/Editor/AssetRegistryExport/Source/Private/AssetRegistryExportCommandlet.cpp
+++ b/Engine/Plugins/Editor/AssetRegistryExport/Source/Private/AssetRegistryExportCommandlet.cpp
@@ -256,16 +256,18 @@ int32 UAssetRegistryExportCommandlet::Main(const FString& CmdLineParams)
 		// Find all the assets that are of the requested primary asset types. These we'll use as roots for the 
 		// dependency listing.
 		TArray<const FAssetData*> PrimaryAssets;
-		const TArray<const FAssetData*>& AssetsWithPrimaryAssetType = AssetRegistry.GetAssetsByTagName(FPrimaryAssetId::PrimaryAssetTypeTag);
-		for (const FAssetData* GFD : AssetsWithPrimaryAssetType)
-		{
-			FName PrimaryAssetType;
-			if (GFD->GetTagValue(FPrimaryAssetId::PrimaryAssetTypeTag, PrimaryAssetType) &&
-				DependencyListAssetTypeNames.Contains(PrimaryAssetType))
+		AssetRegistry.EnumerateAssetsByTagName(FPrimaryAssetId::PrimaryAssetTypeTag, 
+			[&DependencyListAssetTypeNames, &PrimaryAssets](const FAssetData* GFD)
 			{
-				PrimaryAssets.Add(GFD);
-			}
-		}
+				FName PrimaryAssetType;
+				if (GFD->GetTagValue(FPrimaryAssetId::PrimaryAssetTypeTag, PrimaryAssetType) &&
+					DependencyListAssetTypeNames.Contains(PrimaryAssetType))
+				{
+					PrimaryAssets.Add(GFD);
+				}
+
+				return true;
+			});
 
 		
 		UE_LOG(LogAssetRegistryExport, Display, TEXT("ListDependencies matching PrimaryAssets discovered %s"), *NumberString(PrimaryAssets.Num()));
@@ -274,16 +276,18 @@ int32 UAssetRegistryExportCommandlet::Main(const FString& CmdLineParams)
 		TSet<const FAssetData*> RemainingAssetsWithSize;
 		uint64 TotalCompressedBytes = 0;
 		{
-			const TArray<const FAssetData*>& AssetsWithSize = AssetRegistry.GetAssetsByTagName(UE::AssetRegistry::Stage_ChunkCompressedSizeFName);
-			for (const FAssetData* AD : AssetsWithSize)
-			{
-				uint64 CompressedSize = 0;
-				if (AD->GetTagValue(UE::AssetRegistry::Stage_ChunkCompressedSizeFName, CompressedSize))
+			AssetRegistry.EnumerateAssetsByTagName(UE::AssetRegistry::Stage_ChunkCompressedSizeFName,
+				[&RemainingAssetsWithSize, &TotalCompressedBytes](const FAssetData* AD)
 				{
-					TotalCompressedBytes += CompressedSize;
-					RemainingAssetsWithSize.Add(AD);
-				}
-			}
+					uint64 CompressedSize = 0;
+					if (AD->GetTagValue(UE::AssetRegistry::Stage_ChunkCompressedSizeFName, CompressedSize))
+					{
+						TotalCompressedBytes += CompressedSize;
+						RemainingAssetsWithSize.Add(AD);
+					}
+
+					return true;
+				});
 		}
 
 		const TMap<FName, const FAssetPackageData*>& AllPackages = AssetRegistry.GetAssetPackageDataMap();
diff --git a/Engine/Plugins/Runtime/AssetTags/Source/AssetTags/Private/AssetTagsSubsystem.cpp b/Engine/Plugins/Runtime/AssetTags/Source/AssetTags/Private/AssetTagsSubsystem.cpp
--- a/Engine/Plugins/Runtime/AssetTags/Source/AssetTags/Private/AssetTagsSubsystem.cpp
+++ b/Engine/Plugins/Runtime/AssetTags/Source/AssetTags/Private/AssetTagsSubsystem.cpp
@@ -387,8 +387,8 @@ TArray<FName> UAssetTagsSubsystem::GetCollections()
 		FStringView CollectionTagPrefix(FAssetData::GetCollectionTagPrefix());
 
 		FString TagNameStr;
-		AssetRegistry.ReadLockEnumerateTagToAssetDatas(
-			[&TagNameStr, &CollectionTagPrefix, &CollectionNames](FName TagName, const TArray<const FAssetData*>& Assets)
+		AssetRegistry.ReadLockEnumerateAllTagToAssetDatas(
+			[&TagNameStr, &CollectionTagPrefix, &CollectionNames](FName TagName, IAssetRegistry::FEnumerateAssetDatasFunc EnumerateAssets)
 			{
 				TagName.ToString(TagNameStr);
 				if (FStringView(TagNameStr).StartsWith(CollectionTagPrefix, ESearchCase::IgnoreCase))
@@ -396,6 +396,8 @@ TArray<FName> UAssetTagsSubsystem::GetCollections()
 					const FString TrimmedTagNameStr = TagNameStr.Mid(CollectionTagPrefix.Len());
 					CollectionNames.Add(*TrimmedTagNameStr);
 				}
+
+				return true;
 			});
 	}
 #endif	// WITH_EDITOR
diff --git a/Engine/Source/Editor/ContentBrowser/Private/SContentBrowser.cpp b/Engine/Source/Editor/ContentBrowser/Private/SContentBrowser.cpp
--- a/Engine/Source/Editor/ContentBrowser/Private/SContentBrowser.cpp
+++ b/Engine/Source/Editor/ContentBrowser/Private/SContentBrowser.cpp
@@ -4194,14 +4194,16 @@ void SContentBrowser::OnAssetSearchSuggestionFilter(const FText& SearchText, TAr
 
 		const FText MetaDataCategoryName = NSLOCTEXT("ContentBrowser", "MetaDataCategoryName", "Meta-Data");
 		FString TagNameStr;
-		AssetRegistry.ReadLockEnumerateTagToAssetDatas(
-			[&PassesValueFilter, &PossibleSuggestions, &MetaDataCategoryName, &TagNameStr](FName TagName, const TArray<const FAssetData*>& Assets)
+		AssetRegistry.ReadLockEnumerateAllTagToAssetDatas(
+			[&PassesValueFilter, &PossibleSuggestions, &MetaDataCategoryName, &TagNameStr](FName TagName, IAssetRegistry::FEnumerateAssetDatasFunc EnumerateAssets)
 			{
 				TagName.ToString(TagNameStr);
 				if (PassesValueFilter(TagNameStr))
 				{
 					PossibleSuggestions.Add(FAssetSearchBoxSuggestion{ TagNameStr, FText::FromString(TagNameStr), MetaDataCategoryName });
 				}
+
+				return true;
 			});
 	}
 
diff --git a/Engine/Source/Runtime/AssetRegistry/Private/AssetRegistry.cpp b/Engine/Source/Runtime/AssetRegistry/Private/AssetRegistry.cpp
--- a/Engine/Source/Runtime/AssetRegistry/Private/AssetRegistry.cpp
+++ b/Engine/Source/Runtime/AssetRegistry/Private/AssetRegistry.cpp
@@ -5469,7 +5469,7 @@ const TSet<FName>& UAssetRegistryImpl::GetCachedEmptyPackages() const
 bool UAssetRegistryImpl::ContainsTag(FName TagName) const
 {
 	UE::AssetRegistry::FInterfaceReadScopeLock InterfaceScopeLock(InterfaceLock);
-	return GuardedData.GetState().GetTagToAssetDatasMap().Contains(TagName);
+	return GuardedData.GetState().ContainsTag(TagName);
 }
 
 namespace UE::AssetRegistry
@@ -8573,11 +8573,21 @@ void FEventContext::Append(FEventContext&& Other)
 
 void UAssetRegistryImpl::ReadLockEnumerateTagToAssetDatas(TFunctionRef<void(FName TagName, const TArray<const FAssetData*>& Assets)> Callback) const
 {
+	UE_LOG(LogAssetRegistry, Error, TEXT("ReadLockEnumerateTagToAssetDatas has been deprecated. Use ReadLockEnumerateAllTagToAssetDatas instead."));
+
 	UE::AssetRegistry::FInterfaceReadScopeLock InterfaceScopeLock(InterfaceLock);
-	for (const TPair<FName, const TArray<const FAssetData*>>& Pair : GuardedData.GetState().GetTagToAssetDatasMap())
+	GuardedData.GetState().EnumerateTags([Callback](FName TagName)
 	{
-		Callback(Pair.Key, Pair.Value);
-	}
+		TArray<const FAssetData*> EmptyArray;
+		Callback(TagName, EmptyArray);
+		return true;
+	});
+}
+
+void UAssetRegistryImpl::ReadLockEnumerateAllTagToAssetDatas(TFunctionRef<bool(FName TagName, IAssetRegistry::FEnumerateAssetDatasFunc EnumerateAssets)> Callback) const
+{
+	UE::AssetRegistry::FInterfaceReadScopeLock InterfaceScopeLock(InterfaceLock);
+	GuardedData.GetState().EnumerateTagToAssetDatas(Callback);
 }
 
 void UAssetRegistryImpl::Broadcast(UE::AssetRegistry::Impl::FEventContext& EventContext)
diff --git a/Engine/Source/Runtime/AssetRegistry/Private/AssetRegistry.h b/Engine/Source/Runtime/AssetRegistry/Private/AssetRegistry.h
--- a/Engine/Source/Runtime/AssetRegistry/Private/AssetRegistry.h
+++ b/Engine/Source/Runtime/AssetRegistry/Private/AssetRegistry.h
@@ -202,7 +202,9 @@ public:
 
 	virtual void Tick (float DeltaTime) override;
 
+	UE_DEPRECATED(5.5, "ReadLockEnumerateTagToAssetDatas with TArray has been deprecated. Use ReadLockEnumerateAllTagToAssetDatas instead.")
 	virtual void ReadLockEnumerateTagToAssetDatas(TFunctionRef<void(FName TagName, const TArray<const FAssetData*>& Assets)> Callback) const override;
+	virtual void ReadLockEnumerateAllTagToAssetDatas(TFunctionRef<bool(FName TagName, IAssetRegistry::FEnumerateAssetDatasFunc EnumerateAssets)> Callback) const override;
 
 	virtual bool IsPathBeautificationNeeded(const FString& InAssetPath) const override;
 
diff --git a/Engine/Source/Runtime/AssetRegistry/Private/AssetRegistryState.cpp b/Engine/Source/Runtime/AssetRegistry/Private/AssetRegistryState.cpp
--- a/Engine/Source/Runtime/AssetRegistry/Private/AssetRegistryState.cpp
+++ b/Engine/Source/Runtime/AssetRegistry/Private/AssetRegistryState.cpp
@@ -841,8 +841,8 @@ bool AssetDataMatchesTag(const FAssetData* AssetData, const TPair<FName, TOption
 	}
 }
 
-template<typename CallbackType>
-void FilterAssets(TArray<const FAssetData*>& InOutResults, const TMap<FName, TArray<FAssetData*>>& AccelerationMap,
+template<typename CallbackType, typename AccelerationMapType>
+void FilterAssets(TArray<const FAssetData*>& InOutResults, const AccelerationMapType& AccelerationMap,
 	const TMultiMap<FName, TOptional<FString>>& TagsAndValues, CallbackType&& FunctionToKeepAsset, int32 FilterComplexity)
 {
 	TArray<TArray<const FAssetData*>, TInlineAllocator<10>> Matches;
@@ -852,7 +852,7 @@ void FilterAssets(TArray<const FAssetData*>& InOutResults, const TMap<FName, TAr
 	for (const TPair<FName, TOptional<FString>>& TagPair : TagsAndValues)
 	{
 		TArray<const FAssetData*>& Results = Matches.Emplace_GetRef();
-		if (const TArray<FAssetData*>* TagAssets = AccelerationMap.Find(TagPair.Key))
+		if (const auto TagAssets = AccelerationMap.Find(TagPair.Key))
 		{
 			Results.Reserve(TagAssets->Num());
 			for (FAssetData* AssetData : *TagAssets)
@@ -1981,11 +1981,16 @@ void FAssetRegistryState::SetAssetDatas(TArrayView<FAssetData> AssetDatas, const
 		{
 			for (const TPair<FName, FAssetTagValueRef>& Pair : AssetData.TagsAndValues)
 			{
-				TArray<FAssetData*>& TagAssets = CachedAssetsByTag.FindOrAdd(Pair.Key);
+				TSet<FAssetData*>& TagAssets = CachedAssetsByTag.FindOrAdd(Pair.Key);
 				TagAssets.Add(&AssetData);
 			}
 		}
-		ShrinkMultimap(CachedAssetsByTag);
+
+		CachedAssetsByTag.Shrink();
+		for (TPair<FName, TSet<FAssetData*>>& Pair : CachedAssetsByTag)
+		{
+			Pair.Value.Shrink();
+		}
 	};
 
 	if (Options.ParallelWorkers <= 1)
@@ -2030,7 +2035,7 @@ void FAssetRegistryState::AddAssetData(FAssetData* AssetData)
 	{
 		FName Key = TagIt.Key();
 
-		TArray<FAssetData*>& TagAssets = CachedAssetsByTag.FindOrAdd(Key);
+		TSet<FAssetData*>& TagAssets = CachedAssetsByTag.FindOrAdd(Key);
 		TagAssets.Add(AssetData);
 	}
 }
@@ -2057,11 +2062,10 @@ void FAssetRegistryState::AddTagsToAssetData(const FSoftObjectPath& InObjectPath
 
 void FAssetRegistryState::FilterTags(const FAssetRegistrySerializationOptions& Options)
 {
-	// Calling SetTagsOnExistingAsset for any changed tags is slow because the elements of 
-	// CachedAssetsByTag are unsorted TArrays and removal of the AssetData from its old
-	// CachedAssetsByTag is slow. For cases where many Assets change it is therefore faster
-	// just to recreate CachedAssetsByTag rather than trying to update them.
-	for (TPair<FName, TArray<FAssetData*>>& Pair : CachedAssetsByTag)
+	// Calling SetTagsOnExistingAsset for any changed tags might be slow.
+	// For cases where many Assets change it might be faster to recreate CachedAssetsByTag
+	// rather than trying to update its elements for each Asset change
+	for (TPair<FName, TSet<FAssetData*>>& Pair : CachedAssetsByTag)
 	{
 		Pair.Value.Reset();
 	}
@@ -2096,10 +2100,10 @@ void FAssetRegistryState::SetTagsOnExistingAsset(FAssetData* AssetData, FAssetDa
 
 		if (!NewTags.Contains(FNameKey))
 		{
-			TArray<FAssetData*>* OldTagAssets = CachedAssetsByTag.Find(FNameKey);
+			TSet<FAssetData*>* OldTagAssets = CachedAssetsByTag.Find(FNameKey);
 			if (OldTagAssets)
 			{
-				OldTagAssets->RemoveSingleSwap(AssetData);
+				OldTagAssets->Remove(AssetData);
 			}
 		}
 	}
@@ -2222,8 +2226,8 @@ void FAssetRegistryState::UpdateAssetData(FAssetData* AssetData, FAssetData&& Ne
 
 			if (!NewAssetData.TagsAndValues.Contains(FNameKey))
 			{
-				TArray<FAssetData*>* OldTagAssets = CachedAssetsByTag.Find(FNameKey);
-				OldTagAssets->RemoveSingleSwap(AssetData);
+				TSet<FAssetData*>* OldTagAssets = CachedAssetsByTag.Find(FNameKey);
+				OldTagAssets->Remove(AssetData);
 			}
 		}
 
@@ -2233,7 +2237,7 @@ void FAssetRegistryState::UpdateAssetData(FAssetData* AssetData, FAssetData&& Ne
 
 			if (!AssetData->TagsAndValues.Contains(FNameKey))
 			{
-				TArray<FAssetData*>& NewTagAssets = CachedAssetsByTag.FindOrAdd(FNameKey);
+				TSet<FAssetData*>& NewTagAssets = CachedAssetsByTag.FindOrAdd(FNameKey);
 				NewTagAssets.Add(AssetData);
 			}
 		}
@@ -2348,8 +2352,8 @@ void FAssetRegistryState::RemoveAssetData(FAssetData* AssetData, const FCachedAs
 
 	for (auto TagIt = AssetData->TagsAndValues.CreateConstIterator(); TagIt; ++TagIt)
 	{
-		TArray<FAssetData*>* OldTagAssets = CachedAssetsByTag.Find(TagIt.Key());
-		OldTagAssets->RemoveSingleSwap(AssetData);
+		TSet<FAssetData*>* OldTagAssets = CachedAssetsByTag.Find(TagIt.Key());
+		OldTagAssets->Remove(AssetData);
 	}
 
 	// Only remove dependencies and package data if there are no other known assets in the package
@@ -2665,7 +2669,11 @@ static void PrintAssetDataMap(FString Name, const MapType& AssetMap, TStringBuil
 		++ValidCount;
 
 		Items.Reset();
-		Items.Append(AssetArray);
+		Items.Reserve(AssetArray.Num());
+		for (const auto& It : AssetArray)
+		{
+			Items.Add(It);
+		}
 		Items.Sort([](const FAssetData& A, const FAssetData& B)
 			{ return A.GetSoftObjectPath().LexicalLess(B.GetSoftObjectPath()); }
 		);
diff --git a/Engine/Source/Runtime/AssetRegistry/Public/AssetRegistry/AssetRegistryState.h b/Engine/Source/Runtime/AssetRegistry/Public/AssetRegistry/AssetRegistryState.h
--- a/Engine/Source/Runtime/AssetRegistry/Public/AssetRegistry/AssetRegistryState.h
+++ b/Engine/Source/Runtime/AssetRegistry/Public/AssetRegistry/AssetRegistryState.h
@@ -392,22 +392,90 @@ public:
 		return InvalidArray;
 	}
 
+	/**
+	 * Enumerates all tags of any asset in the AssetRegistry
+	 *
+	 * @param Callback the function for each tag
+	 */
+	void EnumerateTags(TFunctionRef<bool(FName TagName)> Callback) const
+	{
+		for (const TPair<FName, TSet<FAssetData*>>& Pair : CachedAssetsByTag)
+		{
+			if (!Callback(Pair.Key))
+			{
+				break;
+			}
+		}
+	}
+
+	/** Return whether the given TagName occurs in the tags of any asset in the AssetRegistry */
+	bool ContainsTag(FName TagName) const
+	{
+		return CachedAssetsByTag.Contains(TagName);
+	}
+
 	/**
 	 * Gets the asset data for the specified asset tag
 	 *
 	 * @param TagName the tag name to search for
 	 * @return An array of AssetData*, empty if nothing found
 	 */
+	UE_DEPRECATED(5.5, "GetAssetsByTagName has been deprecated. Please use EnumerateAssetsByTagName")
 	const TArray<const FAssetData*>& GetAssetsByTagName(const FName TagName) const
 	{
+		ensureMsgf(false, TEXT("GetAssetsByTagName has been deprecated. Please use EnumerateAssetsByTagName"));
 		static TArray<const FAssetData*> InvalidArray;
-		const TArray<FAssetData*>* FoundAssetArray = CachedAssetsByTag.Find(TagName);
-		if (FoundAssetArray)
+		return InvalidArray;
+	}
+
+	/**
+	 * Enumerates the asset data for the specified asset tag
+	 *
+	 * @param TagName the tag name to search for
+	 * @param Callback the function called for each asset data
+	 */
+	void EnumerateAssetsByTagName(const FName TagName, TFunctionRef<bool(const FAssetData* AssetData)> Callback) const
+	{
+		const TSet<FAssetData*>* FoundAssets = CachedAssetsByTag.Find(TagName);
+		if (FoundAssets)
 		{
-			return reinterpret_cast<const TArray<const FAssetData*>&>(*FoundAssetArray);
+			for (const FAssetData* AssetData : *FoundAssets)
+			{
+				if (!Callback(AssetData))
+				{
+					break;
+				}
+			}
 		}
+	}
 
-		return InvalidArray;
+	/**
+	 * Enumerates all tags of any asset in the AssetRegistry including function that can be called to enumerate assets for each tag
+	 *
+	 * @param Callback the function for each tag pair
+	 */
+	void EnumerateTagToAssetDatas(TFunctionRef<bool(FName TagName, IAssetRegistry::FEnumerateAssetDatasFunc EnumerateAssets)> Callback) const
+	{
+		for (const TPair<FName, TSet<FAssetData*>>& Pair : CachedAssetsByTag)
+		{
+			const bool bKeepEnumerating = Callback(Pair.Key, [&Pair](IAssetRegistry::FAssetDataFunc AssetCallback)
+				{
+					for (const FAssetData* AssetData : Pair.Value)
+					{
+						if (!AssetCallback(AssetData))
+						{
+							return false;
+						}
+					}
+
+					return true;
+				});
+
+			if (!bKeepEnumerating)
+			{
+				break;
+			}
+		}
 	}
 
 	/** Returns const version of internal ObjectPath->AssetData map for fast iteration */
@@ -417,9 +485,12 @@ public:
 	}
 
 	/** Returns const version of internal Tag->AssetDatas map for fast iteration */
+	UE_DEPRECATED(5.5, "GetTagToAssetDatasMap has been deprecated. Please use EnumerateTags or EnumerateTagToAssetDatas")
 	const TMap<FName, const TArray<const FAssetData*>> GetTagToAssetDatasMap() const
 	{
-		return reinterpret_cast<const TMap<FName, const TArray<const FAssetData*>>&>(CachedAssetsByTag);
+		ensureMsgf(false, TEXT("GetTagToAssetDatasMap has been deprecated. Please use EnumerateTags or EnumerateTagToAssetDatas"));
+		static TMap<FName, const TArray<const FAssetData*>> InvalidMap;
+		return InvalidMap;
 	}
 
 	/** Returns const version of internal PackageName->PackageData map for fast iteration */
@@ -638,7 +709,7 @@ private:
 	TMap<FTopLevelAssetPath, TArray<FAssetData*> > CachedAssetsByClass;
 
 	/** The map of asset tag to asset data for assets saved to disk */
-	TMap<FName, TArray<FAssetData*> > CachedAssetsByTag;
+	TMap<FName, TSet<FAssetData*> > CachedAssetsByTag;
 
 	/** A map of object names to dependency data */
 	TMap<FAssetIdentifier, FDependsNode*> CachedDependsNodes;
diff --git a/Engine/Source/Runtime/AssetRegistry/Public/AssetRegistry/IAssetRegistry.h b/Engine/Source/Runtime/AssetRegistry/Public/AssetRegistry/IAssetRegistry.h
--- a/Engine/Source/Runtime/AssetRegistry/Public/AssetRegistry/IAssetRegistry.h
+++ b/Engine/Source/Runtime/AssetRegistry/Public/AssetRegistry/IAssetRegistry.h
@@ -881,8 +881,19 @@ public:
 	 * To avoid copies, the callback is called from within the ReadLock.
 	 * DO NOT CALL AssetRegistry functions from the callback; doing so will create a deadlock.
 	 */
+	UE_DEPRECATED(5.5, "ReadLockEnumerateTagToAssetDatas with TArray has been deprecated. Use ReadLockEnumerateTagToAssetDatas with EnumerateAssets callback instead.")
 	virtual void ReadLockEnumerateTagToAssetDatas(TFunctionRef<void(FName TagName, const TArray<const FAssetData*>& Assets)> Callback) const = 0;
 
+	typedef TFunctionRef<bool(const FAssetData* AssetData)> FAssetDataFunc;
+	typedef TFunctionRef<bool(FAssetDataFunc AssetCallback)> FEnumerateAssetDatasFunc;
+
+	/**
+	 * Enumerate all pairs in State->TagToAssetDataMapAssetRegistry and call a callback on each pair.
+	 * To avoid copies, the callback is called from within the ReadLock.
+	 * DO NOT CALL AssetRegistry functions from the callback; doing so will create a deadlock.
+	 */
+	virtual void ReadLockEnumerateAllTagToAssetDatas(TFunctionRef<bool(FName TagName, FEnumerateAssetDatasFunc EnumerateAssets)> Callback) const = 0;
+
 	/**
 	 * Predicate called to decide whether to recurse into a reference when setting manager references
 	 *
