Asset registry: add new API functions for getting assets.
GetAssets taking compiled AR filter.
GetInMemoryAssets to perform the reverse of bIncludeOnlyOnDiskAssets.

diff --git a/Engine/Source/Runtime/AssetRegistry/Private/AssetRegistry.cpp b/Engine/Source/Runtime/AssetRegistry/Private/AssetRegistry.cpp
--- a/Engine/Source/Runtime/AssetRegistry/Private/AssetRegistry.cpp
+++ b/Engine/Source/Runtime/AssetRegistry/Private/AssetRegistry.cpp
@@ -2228,9 +2228,16 @@ bool UAssetRegistryImpl::GetAssets(const FARFilter& InFilter, TArray<FAssetData>
 	{
 		return false;
 	}
+	return GetAssets(CompiledFilter, OutAssetData, bSkipARFilteredAssets);
+}
+
+bool UAssetRegistryImpl::GetAssets(const FARCompiledFilter& CompiledFilter, TArray<FAssetData>& OutAssetData,
+	bool bSkipARFilteredAssets) const
+{
+	using namespace UE::AssetRegistry::Utils;
 
 	TSet<FName> PackagesToSkip;
-	if (!InFilter.bIncludeOnlyOnDiskAssets)
+	if (!CompiledFilter.bIncludeOnlyOnDiskAssets)
 	{
 		bool bStopIterationUnused;
 		EnumerateMemoryAssets(CompiledFilter, PackagesToSkip, bStopIterationUnused,
@@ -2253,6 +2260,34 @@ bool UAssetRegistryImpl::GetAssets(const FARFilter& InFilter, TArray<FAssetData>
 	return true;
 }
 
+bool UAssetRegistryImpl::GetInMemoryAssets(const FARFilter& InFilter, TArray<FAssetData>& OutAssetData, bool bSkipARFilteredAssets) const 
+{
+	using namespace UE::AssetRegistry::Utils;
+
+	FARCompiledFilter CompiledFilter;
+	CompileFilter(InFilter, CompiledFilter);
+	if (CompiledFilter.IsEmpty() || !IsFilterValid(CompiledFilter))
+	{
+		return false;
+	}
+	return GetInMemoryAssets(CompiledFilter, OutAssetData, bSkipARFilteredAssets);
+}
+
+bool UAssetRegistryImpl::GetInMemoryAssets(const FARCompiledFilter& CompiledFilter, TArray<FAssetData>& OutAssetData, bool bSkipARFilteredAssets) const 
+{
+	using namespace UE::AssetRegistry::Utils;
+	TSet<FName> PackagesToSkipUnused;
+	bool bStopIterationUnused;
+	EnumerateMemoryAssets(CompiledFilter, PackagesToSkipUnused, bStopIterationUnused,
+		InterfaceLock, GuardedData.GetState(),
+		[&OutAssetData](FAssetData&& AssetData)
+		{
+			OutAssetData.Add(MoveTemp(AssetData));
+			return true;
+		}, bSkipARFilteredAssets);
+	return true;
+}
+
 bool UAssetRegistryImpl::EnumerateAssets(const FARFilter& InFilter, TFunctionRef<bool(const FAssetData&)> Callback,
 	bool bSkipARFilteredAssets) const
 {
diff --git a/Engine/Source/Runtime/AssetRegistry/Private/AssetRegistry.h b/Engine/Source/Runtime/AssetRegistry/Private/AssetRegistry.h
--- a/Engine/Source/Runtime/AssetRegistry/Private/AssetRegistry.h
+++ b/Engine/Source/Runtime/AssetRegistry/Private/AssetRegistry.h
@@ -37,6 +37,9 @@ public:
 	virtual bool GetAssetsByTags(const TArray<FName>& AssetTags, TArray<FAssetData>& OutAssetData) const override;
 	virtual bool GetAssetsByTagValues(const TMultiMap<FName, FString>& AssetTagsAndValues, TArray<FAssetData>& OutAssetData) const override;
 	virtual bool GetAssets(const FARFilter& Filter, TArray<FAssetData>& OutAssetData, bool bSkipARFilteredAssets = true) const override;
+	virtual bool GetAssets(const FARCompiledFilter& Filter, TArray<FAssetData>& OutAssetData, bool bSkipARFilteredAssets = true) const override;
+	virtual bool GetInMemoryAssets(const FARFilter& Filter, TArray<FAssetData>& OutAssetData, bool bSkipARFilteredAssets=true) const override;
+	virtual bool GetInMemoryAssets(const FARCompiledFilter& Filter, TArray<FAssetData>& OutAssetData, bool bSkipARFilteredAssets=true) const override;
 	virtual bool EnumerateAssets(const FARFilter& Filter, TFunctionRef<bool(const FAssetData&)> Callback, bool bSkipARFilteredAssets=true) const override;
 	virtual bool EnumerateAssets(const FARCompiledFilter& Filter, TFunctionRef<bool(const FAssetData&)> Callback, bool bSkipARFilteredAssets = true) const override;
 	UE_DEPRECATED(5.1, "Asset path FNames have been deprecated, use FSoftObjectPath instead.")
diff --git a/Engine/Source/Runtime/AssetRegistry/Public/AssetRegistry/IAssetRegistry.h b/Engine/Source/Runtime/AssetRegistry/Public/AssetRegistry/IAssetRegistry.h
--- a/Engine/Source/Runtime/AssetRegistry/Public/AssetRegistry/IAssetRegistry.h
+++ b/Engine/Source/Runtime/AssetRegistry/Public/AssetRegistry/IAssetRegistry.h
@@ -274,6 +274,22 @@ public:
 	 */
 	UFUNCTION(BlueprintCallable, BlueprintPure=false, Category="AssetRegistry")
 	virtual bool GetAssets(const FARFilter& Filter, TArray<FAssetData>& OutAssetData, bool bSkipARFilteredAssets=true) const = 0;
+	virtual bool GetAssets(const FARCompiledFilter& Filter, TArray<FAssetData>& OutAssetData, bool bSkipARFilteredAssets=true) const = 0;
+
+	/**
+	 * Gets asset data for in-memory assets only, that match the filter.
+	 * Returns assets that would be excluded by calling GetAssets with bIncludeOnlyOnDiskAssets set to true.
+	 * @note This method IGNORES the value of bIncludeOnlyOnDiskAssets on the provided filter. 
+	 * Assets returned must satisfy every filter component if there is at least one element in the component's array.
+	 * Assets will satisfy a component if they match any of the elements in it.
+	 *
+	 * @param Filter filter to apply to the assets in the AssetRegistry
+	 * @param OutAssetData the list of assets in this path
+	 * @param bSkipARFilteredAssets If true, skips Objects that return true for IsAsset but are not assets in the current platform.
+	 */
+	UFUNCTION(BlueprintCallable, BlueprintPure=false, Category="AssetRegistry")
+	virtual bool GetInMemoryAssets(const FARFilter& Filter, TArray<FAssetData>& OutAssetData, bool bSkipARFilteredAssets=true) const = 0;
+	virtual bool GetInMemoryAssets(const FARCompiledFilter& Filter, TArray<FAssetData>& OutAssetData, bool bSkipARFilteredAssets=true) const = 0;
 
 	/**
 	 * Enumerate asset data for all assets that match the filter.
