Adding a check in AssetRegistry to verify that a ClassDefaultObject is actually fully constructed before calling ThreadedPostLoadAssetRegistryTags.

diff --git a/Engine/Source/Runtime/AssetRegistry/Private/AssetRegistry.cpp b/Engine/Source/Runtime/AssetRegistry/Private/AssetRegistry.cpp
--- a/Engine/Source/Runtime/AssetRegistry/Private/AssetRegistry.cpp
+++ b/Engine/Source/Runtime/AssetRegistry/Private/AssetRegistry.cpp
@@ -6000,11 +6000,17 @@ bool FAssetRegistryImpl::TryPostLoadAssetRegistryTags(FAssetData* AssetData)
 
 		if (AssetClass)
 		{
-			if (UObject* ClassCDO = AssetClass->GetDefaultObject(false))
-			{
+			UObject* ClassDefaultObject = AssetClass->GetDefaultObject(false);
+			if (ClassDefaultObject && !ClassDefaultObject->HasAnyFlags(RF_NeedInitialization))
+			{
+				// We are using RF_NeedInitialization to guarantee that ClassDefaultObject is fully initialized
+				// potentially on another thread. For weakly ordered memory platforms, we need to 
+				// ensure that our read of the vtable ptr isn't performed prior to the read of the class flags
+				// otherwise we might see a stale vtable despite seeing RF_NeedInit clear.
+				std::atomic_thread_fence(std::memory_order_acquire);
 				TArray<UObject::FAssetRegistryTag> TagsToModify;
 				UObject::FPostLoadAssetRegistryTagsContext Context(*AssetData, TagsToModify);
-				ClassCDO->ThreadedPostLoadAssetRegistryTags(Context);
+				ClassDefaultObject->ThreadedPostLoadAssetRegistryTags(Context);
 				if (TagsToModify.Num())
 				{
 					FAssetDataTagMap TagsAndValues = AssetData->TagsAndValues.CopyMap();
@@ -6029,7 +6035,7 @@ bool FAssetRegistryImpl::TryPostLoadAssetRegistryTags(FAssetData* AssetData)
 			}
 			else 
 			{
-				ensureMsgf(!MakeFinalChecks, TEXT("Unable to PostLoadAssetRegistryTags for '%s' because the CDO for ancestor class '%s' could not be found."),
+				ensureMsgf(!MakeFinalChecks, TEXT("Unable to PostLoadAssetRegistryTags for '%s' because the CDO for ancestor class '%s' could not be found or was not ready."),
 					*AssetData->GetObjectPathString(), *AssetClassPath.ToString());
 			}
 		}
diff --git a/Engine/Source/Runtime/CoreUObject/Private/UObject/UObjectGlobals.cpp b/Engine/Source/Runtime/CoreUObject/Private/UObject/UObjectGlobals.cpp
--- a/Engine/Source/Runtime/CoreUObject/Private/UObject/UObjectGlobals.cpp
+++ b/Engine/Source/Runtime/CoreUObject/Private/UObject/UObjectGlobals.cpp
@@ -4077,6 +4077,11 @@ void FObjectInitializer::PostConstructInit()
 		Obj->CheckDefaultSubobjects();
 	}
 
+	// If we want to be able to use RF_NeedInitialization from another thread
+	// to know that Obj is fully constructed, then on weakly order platforms we 
+	// need a fence to guarantee that the cleared flag is only visible to other threads
+	// after the other initialization-related writes
+	std::atomic_thread_fence(std::memory_order_release);
 	UE_AUTORTFM_OPEN({ Obj->ClearFlags(RF_NeedInitialization); });
 
 	// clear the object pointer so we can guard against running this function again
