Re-enabling background thread processing of asset registry data.

diff --git a/Engine/Source/Runtime/AssetRegistry/Private/AssetDataGatherer.cpp b/Engine/Source/Runtime/AssetRegistry/Private/AssetDataGatherer.cpp
--- a/Engine/Source/Runtime/AssetRegistry/Private/AssetDataGatherer.cpp
+++ b/Engine/Source/Runtime/AssetRegistry/Private/AssetDataGatherer.cpp
@@ -85,7 +85,7 @@ static FAutoConsoleVariableRef CVarIgnoreEmptyDirectories(
 	bIgnoreEmptyDirectories,
 	TEXT("If true, completely empty leaf directories are ignored by the asset registry while scanning"));
 
-bool bTickGatherOnGTOnly = true;
+bool bTickGatherOnGTOnly = false;
 static FAutoConsoleVariableRef CVarTickGatherOnGTOnly(
 	TEXT("AssetRegistry.TickGatherOnGTOnly"),
 	bTickGatherOnGTOnly,
