Fixed assert when a path ending with a slash is added or removed from the asset registry. Remove trailing slashes from within the asset registry public add & remove functions.

diff --git a/Engine/Source/Runtime/AssetRegistry/Private/AssetRegistry.cpp b/Engine/Source/Runtime/AssetRegistry/Private/AssetRegistry.cpp
--- a/Engine/Source/Runtime/AssetRegistry/Private/AssetRegistry.cpp
+++ b/Engine/Source/Runtime/AssetRegistry/Private/AssetRegistry.cpp
@@ -3896,7 +3896,7 @@ bool UAssetRegistryImpl::AddPath(const FString& PathToAdd)
 	{
 		LLM_SCOPE(ELLMTag::AssetRegistry);
 		UE::AssetRegistry::FInterfaceWriteScopeLock InterfaceScopeLock(InterfaceLock);
-		bResult = GuardedData.AddPath(EventContext, PathToAdd);
+		bResult = GuardedData.AddPath(EventContext, UE::String::RemoveFromEnd(FStringView(PathToAdd), TEXTVIEW("/")));
 	}
 	Broadcast(EventContext);
 	return bResult;
@@ -3933,7 +3933,7 @@ bool UAssetRegistryImpl::RemovePath(const FString& PathToRemove)
 	{
 		LLM_SCOPE(ELLMTag::AssetRegistry);
 		UE::AssetRegistry::FInterfaceWriteScopeLock InterfaceScopeLock(InterfaceLock);
-		bResult = GuardedData.RemoveAssetPath(EventContext, FName(*PathToRemove));
+		bResult = GuardedData.RemoveAssetPath(EventContext, FName(UE::String::RemoveFromEnd(FStringView(PathToRemove), TEXTVIEW("/"))));
 	}
 	Broadcast(EventContext);
 	return bResult;
