AssetRegistry: Disable DirectoryWatchRoots in UE_IS_COOKED_EDITOR configuration because it causes performance problems in some workflows.
#rnx

diff --git a/Engine/Source/Runtime/AssetRegistry/Private/AssetRegistry.cpp b/Engine/Source/Runtime/AssetRegistry/Private/AssetRegistry.cpp
--- a/Engine/Source/Runtime/AssetRegistry/Private/AssetRegistry.cpp
+++ b/Engine/Source/Runtime/AssetRegistry/Private/AssetRegistry.cpp
@@ -58,6 +58,8 @@
 #include "IDirectoryWatcher.h"
 #endif // WITH_EDITOR
 
+#define UE_ENABLE_DIRECTORYWATCH_ROOTS !UE_IS_COOKED_EDITOR
+
 /**
  * ********** AssetRegistry threading model **********
  * *** Functions and InterfaceLock ***
@@ -1101,6 +1103,12 @@ void UAssetRegistryImpl::InitializeEvents(UE::AssetRegistry::Impl::FInitializeCo
 			// The vast majority of directories we are watching are below the Plugin directories. The memory cost per watch
 			// is sufficiently high to want to avoid setting up many granular watches when we can also setup two coarse ones.
 
+			// Don't add any roots in configurations where the feature is disabled; their existence can cause performance
+			// problems when there are too many disk changes in a short amount of time and the directory watcher's buffer
+			// overflows and it issues a FCA_RescanRequired; in that case with one large root we rescan many unrelated
+			// directories.
+
+#if UE_ENABLE_DIRECTORYWATCH_ROOTS
 			const FString ProjectPluginDir = FPaths::CreateStandardFilename(FPaths::ProjectPluginsDir());
 			if (IPlatformFile::GetPlatformPhysical().DirectoryExists(*ProjectPluginDir))
 			{
@@ -1123,6 +1131,7 @@ void UAssetRegistryImpl::InitializeEvents(UE::AssetRegistry::Impl::FInitializeCo
 
 				OnDirectoryChangedDelegateHandles.Add(WatchRoot, NewHandle);
 			}
+#endif //UE_ENABLE_DIRECTORYWATCH_ROOTS
 
 			FString ContentFolder;
 			for (TArray<FString>::TConstIterator RootPathIt(Context.RootContentPaths); RootPathIt; ++RootPathIt)
