Fixed "LogConfig: Warning: GConfig::Find attempting to access config with non-normalized path..." fired when running -game on a project stored on the same drive as the engine but outside the engine source tree.

diff --git a/Engine/Source/Runtime/AssetRegistry/Private/AssetRegistry.cpp b/Engine/Source/Runtime/AssetRegistry/Private/AssetRegistry.cpp
--- a/Engine/Source/Runtime/AssetRegistry/Private/AssetRegistry.cpp
+++ b/Engine/Source/Runtime/AssetRegistry/Private/AssetRegistry.cpp
@@ -1247,7 +1247,7 @@ void FAssetRegistryImpl::InitRedirectors(Impl::FEventContext& EventContext,
 	TArray<TSharedRef<IPlugin>> EnabledPlugins = IPluginManager::Get().GetEnabledPlugins();
 	for (const TSharedRef<IPlugin>& Plugin : EnabledPlugins)
 	{
-		FString PluginConfigFilename = FString::Printf(TEXT("%s%s/%s.ini"), *FPaths::GeneratedConfigDir(), ANSI_TO_TCHAR(FPlatformProperties::PlatformName()), *Plugin->GetName() );
+		FString PluginConfigFilename = FConfigCacheIni::NormalizeConfigIniPath(FString::Printf(TEXT("%s%s/%s.ini"), *FPaths::GeneratedConfigDir(), ANSI_TO_TCHAR(FPlatformProperties::PlatformName()), *Plugin->GetName() ));
 		
 		bool bShouldRemap = false;
 		
