Removed memory allocations

diff --git a/Engine/Source/Runtime/AssetRegistry/Private/AssetDataGatherer.cpp b/Engine/Source/Runtime/AssetRegistry/Private/AssetDataGatherer.cpp
--- a/Engine/Source/Runtime/AssetRegistry/Private/AssetDataGatherer.cpp
+++ b/Engine/Source/Runtime/AssetRegistry/Private/AssetDataGatherer.cpp
@@ -4164,7 +4164,6 @@ FAssetDataGatherer::ETickResult FAssetDataGatherer::TickInternal(double& TickSta
 		}
 
 		const FName PackageName = *AssetFileData.LongPackageName;
-		const FName Extension = FName(*FPaths::GetExtension(AssetFileData.LocalAbsPath));
 
 		FDiskCachedAssetData** DiskCachedAssetDataPtr = DiskCachedAssetDataMap.Find(PackageName);
 		FDiskCachedAssetData* DiskCachedAssetData = DiskCachedAssetDataPtr ? *DiskCachedAssetDataPtr : nullptr;
@@ -4177,7 +4176,7 @@ FAssetDataGatherer::ETickResult FAssetDataGatherer::TickInternal(double& TickSta
 				DiskCachedAssetData = nullptr;
 			}
 			else if ((!DiskCachedAssetData->DependencyData.PackageName.IsEqual(PackageName, ENameCase::CaseSensitive) && DiskCachedAssetData->DependencyData.PackageName != NAME_None) ||
-				DiskCachedAssetData->Extension != Extension)
+				DiskCachedAssetData->Extension != FName(*FPaths::GetExtension(AssetFileData.LocalAbsPath)))
 			{
 				UE_LOG(LogAssetRegistry, Display, TEXT("Cached dependency data for package '%s' is invalid. Discarding cached data."), *PackageName.ToString());
 				DiskCachedAssetData = nullptr;
@@ -4229,7 +4228,7 @@ FAssetDataGatherer::ETickResult FAssetDataGatherer::TickInternal(double& TickSta
 		else
 		{
 			// Not found in cache (or stale) - schedule to be read from disk
-			ReadContexts.Emplace(PackageName, Extension, AssetFileData);
+			ReadContexts.Emplace(PackageName, FName(*FPaths::GetExtension(AssetFileData.LocalAbsPath)), AssetFileData);
 		}
 	}
 
diff --git a/Engine/Source/Runtime/CoreUObject/Private/UObject/Obj.cpp b/Engine/Source/Runtime/CoreUObject/Private/UObject/Obj.cpp
--- a/Engine/Source/Runtime/CoreUObject/Private/UObject/Obj.cpp
+++ b/Engine/Source/Runtime/CoreUObject/Private/UObject/Obj.cpp
@@ -2992,7 +2992,14 @@ void UObject::LoadConfig( UClass* ConfigClass/*=NULL*/, const TCHAR* InFilename/
 
 		// globalconfig properties should always use the owning class's config file
 		// specifying a value for InFilename will override this behavior (as it does with normal properties)
-		const FString& PropFileName = (bGlobalConfig && InFilename == NULL) ? OwnerClass->GetConfigName() : Filename;
+		const FString* PropFileNamePtr = &Filename;
+		FString Temp;
+		if (bGlobalConfig && InFilename == NULL)
+		{
+			Temp = OwnerClass->GetConfigName();
+			PropFileNamePtr = &Temp;
+		}
+		const FString& PropFileName = *PropFileNamePtr;
 
 		FString Key = Property->GetName();
 		int32 PortFlags = 0;
