AssetRegistry: Add IgnoreInvalidPathWarning flag for ScanPathsSynchronous, for callsites such as level save that want to update the path if it exists but don't care if it doesn't.
#rnx

diff --git a/Engine/Source/Editor/UnrealEd/Private/FileHelpers.cpp b/Engine/Source/Editor/UnrealEd/Private/FileHelpers.cpp
--- a/Engine/Source/Editor/UnrealEd/Private/FileHelpers.cpp
+++ b/Engine/Source/Editor/UnrealEd/Private/FileHelpers.cpp
@@ -958,13 +958,15 @@ static bool SaveWorld(UWorld* World,
 			FAssetRegistryModule& AssetRegistryModule = FModuleManager::LoadModuleChecked<FAssetRegistryModule>("AssetRegistry");
 			IAssetRegistry& AssetRegistry = AssetRegistryModule.Get();
 			
-			// Make sure when we exit SaveWorld AssetRegistry is up to date with saved map
-			AssetRegistry.ScanModifiedAssetFiles({ FinalFilename });
+			// Make sure when we exit SaveWorld AssetRegistry is up to date with saved map. Ignore warnings if the map is
+			// stored in /Temp.
+			AssetRegistry.ScanModifiedAssetFiles({ FinalFilename }, UE::AssetRegistry::EScanFlags::IgnoreInvalidPathWarning);
 			
 			if (bPackageNeedsRename || bNewlyCreated || !bNewPackageExists)
 			{
 				// Force rescan to make sure assets are found on map open or world partition initialize`
-				AssetRegistry.ScanPathsSynchronous( ULevel::GetExternalObjectsPaths(NewPackageName) , true);
+				AssetRegistry.ScanSynchronous(ULevel::GetExternalObjectsPaths(NewPackageName), {} /* FilePaths */,
+					UE::AssetRegistry::EScanFlags::IgnoreInvalidPathWarning | UE::AssetRegistry::EScanFlags::ForceRescan);
 			}
 
 			if (RenamedWorldPartition && RenamedWorldPartition->IsStreamingEnabled())
diff --git a/Engine/Source/Runtime/AssetRegistry/Private/AssetRegistry.cpp b/Engine/Source/Runtime/AssetRegistry/Private/AssetRegistry.cpp
--- a/Engine/Source/Runtime/AssetRegistry/Private/AssetRegistry.cpp
+++ b/Engine/Source/Runtime/AssetRegistry/Private/AssetRegistry.cpp
@@ -128,7 +128,8 @@ namespace UE::AssetRegistry
 		const TCHAR* Names[] = {
 			TEXT("ForceRescan"),
 			TEXT("IgnoreDenyListScanFilters"),
-			TEXT("WaitForInMemoryObjects")
+			TEXT("WaitForInMemoryObjects"),
+			TEXT("IgnoreInvalidPathWarning")
 		};
 		
 		if (Flags == EScanFlags::None)
@@ -4005,15 +4006,13 @@ void UAssetRegistryImpl::ScanPathsSynchronousInternal(const TArray<FString>& InD
 	UE_TRACK_REFERENCING_OPNAME_SCOPED(PackageAccessTrackingOps::NAME_ResetContext);
 	const double SearchStartTime = FPlatformTime::Seconds();
 
-	const bool bForceRescan = !!(InScanFlags & UE::AssetRegistry::EScanFlags::ForceRescan);
-	const bool bIgnoreDenyListScanFilters = !!(InScanFlags & UE::AssetRegistry::EScanFlags::IgnoreDenyListScanFilters);
 	const bool bWaitForInMemoryObjects = !!(InScanFlags & UE::AssetRegistry::EScanFlags::WaitForInMemoryObjects);
 
 	UE::AssetRegistry::Impl::FEventContext EventContext;
 	UE::AssetRegistry::Impl::FClassInheritanceContext InheritanceContext;
 	UE::AssetRegistry::Impl::FClassInheritanceBuffer InheritanceBuffer;
 	UE::AssetRegistry::Impl::FScanPathContext Context(EventContext, InheritanceContext, InDirs, InFiles,
-		bForceRescan, bIgnoreDenyListScanFilters, nullptr /* OutFindAssets */);
+		InScanFlags, nullptr /* OutFindAssets */);
 
 	bool bInitialSearchStarted;
 	bool bInitialSearchCompleted;
@@ -5523,13 +5522,14 @@ namespace Impl
 {
 
 FScanPathContext::FScanPathContext(FEventContext& InEventContext, FClassInheritanceContext& InInheritanceContext, 
-	const TArray<FString>& InDirs, const TArray<FString>& InFiles, bool bInForceRescan,
-	bool bInIgnoreDenyListScanFilters, TArray<FSoftObjectPath>* FoundAssets)
+	const TArray<FString>& InDirs, const TArray<FString>& InFiles, UE::AssetRegistry::EScanFlags InScanFlags,
+	TArray<FSoftObjectPath>* FoundAssets)
 	: EventContext(InEventContext)
 	, InheritanceContext(InInheritanceContext)
 	, OutFoundAssets(FoundAssets)
-	, bForceRescan(bInForceRescan)
-	, bIgnoreDenyListScanFilters(bInIgnoreDenyListScanFilters)
+	, bForceRescan(!!(InScanFlags & UE::AssetRegistry::EScanFlags::ForceRescan))
+	, bIgnoreDenyListScanFilters(!!(InScanFlags & UE::AssetRegistry::EScanFlags::IgnoreDenyListScanFilters))
+	, bIgnoreInvalidPathWarning(!!(InScanFlags & UE::AssetRegistry::EScanFlags::IgnoreInvalidPathWarning))
 {
 	if (OutFoundAssets)
 	{
@@ -5568,14 +5568,20 @@ FScanPathContext::FScanPathContext(FEventContext& InEventContext, FClassInherita
 		}
 		else if (!FPackageName::TryConvertToMountedPath(InFile, &LocalPath, &PackageName, nullptr, nullptr, &Extension, &FlexNameType))
 		{
-			UE_LOG(LogAssetRegistry, Warning, TEXT("ScanPathsSynchronous: %s is not in a mounted path, will not scan."), *InFile);
-			bLogCallstack = true;
+			if (!bIgnoreInvalidPathWarning)
+			{
+				UE_LOG(LogAssetRegistry, Warning, TEXT("ScanPathsSynchronous: %s is not in a mounted path, will not scan."), *InFile);
+				bLogCallstack = true;
+			}
 			continue;
 		}
 		if (FPackageName::IsTempPackage(PackageName))
 		{
-			UE_LOG(LogAssetRegistry, Warning, TEXT("ScanPathsSynchronous: %s is in the /Temp path, will not scan."), *InFile);
-			bLogCallstack = true;
+			if (!bIgnoreInvalidPathWarning)
+			{
+				UE_LOG(LogAssetRegistry, Warning, TEXT("ScanPathsSynchronous: %s is in the /Temp path, will not scan."), *InFile);
+				bLogCallstack = true;
+			}
 			continue;
 		}
 		if (Extension.IsEmpty())
@@ -5589,8 +5595,11 @@ FScanPathContext::FScanPathContext(FEventContext& InEventContext, FClassInherita
 				if (FPackageName::InternalDoesPackageExistEx(PackagePath, FPackageName::EPackageLocationFilter::Any,
 					false /* bMatchCaseOnDisk */, &PackagePath) == FPackageName::EPackageLocationFilter::None)
 				{
-					UE_LOG(LogAssetRegistry, Warning, TEXT("ScanPathsSynchronous: Package %s does not exist, will not scan."), *InFile);
-					bLogCallstack = true;
+					if (!bIgnoreInvalidPathWarning)
+					{
+						UE_LOG(LogAssetRegistry, Warning, TEXT("ScanPathsSynchronous: Package %s does not exist, will not scan."), *InFile);
+						bLogCallstack = true;
+					}
 					continue;
 				}
 				Extension = PackagePath.GetExtensionString(EPackageSegment::Header);
@@ -5609,14 +5618,20 @@ FScanPathContext::FScanPathContext(FEventContext& InEventContext, FClassInherita
 		}
 		else if (!FPackageName::TryConvertToMountedPath(InDir, &LocalPath, &PackageName, nullptr, nullptr, &Extension, &FlexNameType))
 		{
-			UE_LOG(LogAssetRegistry, Warning, TEXT("ScanPathsSynchronous: %s is not in a mounted path, will not scan."), *InDir);
-			bLogCallstack = true;
+			if (!bIgnoreInvalidPathWarning)
+			{
+				UE_LOG(LogAssetRegistry, Warning, TEXT("ScanPathsSynchronous: %s is not in a mounted path, will not scan."), *InDir);
+				bLogCallstack = true;
+			}
 			continue;
 		}
 		if (FPackageName::IsTempPackage(PackageName))
 		{
-			UE_LOG(LogAssetRegistry, Warning, TEXT("ScanPathsSynchronous: %s is in the /Temp path, will not scan."), *InDir);
-			bLogCallstack = true;
+			if (!bIgnoreInvalidPathWarning)
+			{
+				UE_LOG(LogAssetRegistry, Warning, TEXT("ScanPathsSynchronous: %s is in the /Temp path, will not scan."), *InDir);
+				bLogCallstack = true;
+			}
 			continue;
 		}
 		LocalDirs.Add(LocalPath + Extension);
@@ -7094,12 +7109,12 @@ void FAssetRegistryImpl::OnDirectoryChanged(Impl::FEventContext& EventContext,
 			if (GlobalGatherer->IsSynchronous())
 			{
 				Impl::FScanPathContext Context(EventContext, InheritanceContext, NewDirs, NewFiles,
-					false /* bForceRescan */, false /* bIgnoreDenyListScanFilters */, nullptr /* OutFoundAssets */);
+					UE::AssetRegistry::EScanFlags::None, nullptr /* OutFoundAssets */);
 				ScanPathsSynchronous(Context);
 			}
 		}
 	}
-	ScanModifiedAssetFiles(EventContext, InheritanceContext, ModifiedFiles);
+	ScanModifiedAssetFiles(EventContext, InheritanceContext, ModifiedFiles, UE::AssetRegistry::EScanFlags::None);
 }
 
 void FAssetRegistryImpl::OnDirectoryRescanRequired(Impl::FEventContext& EventContext,
@@ -7239,12 +7254,12 @@ void FAssetRegistryImpl::OnDirectoryRescanRequired(Impl::FEventContext& EventCon
 			{
 				TArray<FString> UnusedNewDirs;
 				Impl::FScanPathContext Context(EventContext, InheritanceContext, UnusedNewDirs, FinalResult.NewFiles,
-					false /* bForceRescan */, false /* bIgnoreDenyListScanFilters */, nullptr /* OutFoundAssets */);
+					UE::AssetRegistry::EScanFlags::None, nullptr /* OutFoundAssets */);
 				ScanPathsSynchronous(Context);
 			}
 		}
 	}
-	ScanModifiedAssetFiles(EventContext, InheritanceContext, FinalResult.ModifiedFiles);
+	ScanModifiedAssetFiles(EventContext, InheritanceContext, FinalResult.ModifiedFiles, UE::AssetRegistry::EScanFlags::None);
 }
 
 }
@@ -7436,6 +7451,11 @@ void FAssetRegistryImpl::UpdateRedirectCollector()
 #endif // WITH_EDITOR
 
 void UAssetRegistryImpl::ScanModifiedAssetFiles(const TArray<FString>& InFilePaths)
+{
+	ScanModifiedAssetFiles(InFilePaths, UE::AssetRegistry::EScanFlags::None);
+}
+
+void UAssetRegistryImpl::ScanModifiedAssetFiles(const TArray<FString>& InFilePaths, UE::AssetRegistry::EScanFlags ScanFlags)
 {
 	UE::AssetRegistry::Impl::FEventContext EventContext;
 	{
@@ -7444,7 +7464,7 @@ void UAssetRegistryImpl::ScanModifiedAssetFiles(const TArray<FString>& InFilePat
 		UE::AssetRegistry::Impl::FClassInheritanceContext InheritanceContext;
 		UE::AssetRegistry::Impl::FClassInheritanceBuffer InheritanceBuffer;
 		GetInheritanceContextWithRequiredLock(InterfaceScopeLock, InheritanceContext, InheritanceBuffer);
-		GuardedData.ScanModifiedAssetFiles(EventContext, InheritanceContext, InFilePaths);
+		GuardedData.ScanModifiedAssetFiles(EventContext, InheritanceContext, InFilePaths, ScanFlags);
 	}
 
 #if WITH_EDITOR
@@ -7463,7 +7483,8 @@ namespace UE::AssetRegistry
 {
 
 void FAssetRegistryImpl::ScanModifiedAssetFiles(Impl::FEventContext& EventContext,
-	Impl::FClassInheritanceContext& InheritanceContext, const TArray<FString>& InFilePaths)
+	Impl::FClassInheritanceContext& InheritanceContext, const TArray<FString>& InFilePaths,
+	UE::AssetRegistry::EScanFlags InScanFlags)
 {
 	if (InFilePaths.Num() > 0)
 	{
@@ -7491,10 +7512,13 @@ void FAssetRegistryImpl::ScanModifiedAssetFiles(Impl::FEventContext& EventContex
 			}
 		}
 
+		// ScanModifiedAssetFiles always does a force rescan of the given files
+		InScanFlags |= UE::AssetRegistry::EScanFlags::ForceRescan;
+
 		// Re-scan and update the asset registry with the new asset data
 		TArray<FSoftObjectPath> FoundAssets;
 		Impl::FScanPathContext Context(EventContext, InheritanceContext, TArray<FString>(), InFilePaths,
-			true /* bForceRescan */, false /* bIgnoreDenyListScanFilters */, &FoundAssets);
+			InScanFlags, &FoundAssets);
 		ScanPathsSynchronous(Context);
 
 		// Remove any assets that are no longer present in the package
diff --git a/Engine/Source/Runtime/AssetRegistry/Private/AssetRegistry.h b/Engine/Source/Runtime/AssetRegistry/Private/AssetRegistry.h
--- a/Engine/Source/Runtime/AssetRegistry/Private/AssetRegistry.h
+++ b/Engine/Source/Runtime/AssetRegistry/Private/AssetRegistry.h
@@ -102,6 +102,7 @@ public:
 	virtual void ScanFilesSynchronous(const TArray<FString>& InFilePaths, bool bForceRescan = false) override;
 	virtual void PrioritizeSearchPath(const FString& PathToPrioritize) override;
 	virtual void ScanModifiedAssetFiles(const TArray<FString>& InFilePaths) override;
+	virtual void ScanModifiedAssetFiles(const TArray<FString>& InFilePaths, UE::AssetRegistry::EScanFlags ScanFlags) override;
 	virtual void Serialize(FArchive& Ar) override;
 	virtual void AppendState(const FAssetRegistryState& InState) override;
 	virtual SIZE_T GetAllocatedSize(bool bLogDetailed = false) const override;
diff --git a/Engine/Source/Runtime/AssetRegistry/Private/AssetRegistryImpl.h b/Engine/Source/Runtime/AssetRegistry/Private/AssetRegistryImpl.h
--- a/Engine/Source/Runtime/AssetRegistry/Private/AssetRegistryImpl.h
+++ b/Engine/Source/Runtime/AssetRegistry/Private/AssetRegistryImpl.h
@@ -227,7 +227,7 @@ public:
 	void ScanPathsSynchronous(Impl::FScanPathContext& Context);
 	void PrioritizeSearchPath(const FString& PathToPrioritize);
 	void ScanModifiedAssetFiles(Impl::FEventContext& EventContext, Impl::FClassInheritanceContext& InheritanceContext,
-		const TArray<FString>& InFilePaths);
+		const TArray<FString>& InFilePaths, UE::AssetRegistry::EScanFlags ScanFlags);
 	void Serialize(FArchive& Ar, Impl::FEventContext& EventContext);
 	void AppendState(Impl::FEventContext& EventContext, const FAssetRegistryState& InState,
 		FAssetRegistryState::EInitializationMode Mode = FAssetRegistryState::EInitializationMode::Append);
@@ -700,8 +700,9 @@ struct FClassInheritanceContext
 struct FScanPathContext
 {
 	FScanPathContext(FEventContext& InEventContext, FClassInheritanceContext& InInheritanceContext,
-		const TArray<FString>& InDirs, const TArray<FString>& InFiles, bool bInForceRescan = false,
-		bool bInIgnoreDenyListScanFilters = false, TArray<FSoftObjectPath>* FoundAssets = nullptr);
+		const TArray<FString>& InDirs, const TArray<FString>& InFiles,
+		UE::AssetRegistry::EScanFlags InScanFlags = UE::AssetRegistry::EScanFlags::None,
+		TArray<FSoftObjectPath>* FoundAssets = nullptr);
 
 	TArray<FString> PackageDirs;
 	TArray<FString> LocalDirs;
@@ -714,6 +715,7 @@ struct FScanPathContext
 	int32 NumFoundAssets = 0;
 	bool bForceRescan = false;
 	bool bIgnoreDenyListScanFilters = false;
+	bool bIgnoreInvalidPathWarning = false;
 	EGatherStatus Status = EGatherStatus::TickActiveGatherActive;
 };
 
diff --git a/Engine/Source/Runtime/AssetRegistry/Public/AssetRegistry/IAssetRegistry.h b/Engine/Source/Runtime/AssetRegistry/Public/AssetRegistry/IAssetRegistry.h
--- a/Engine/Source/Runtime/AssetRegistry/Public/AssetRegistry/IAssetRegistry.h
+++ b/Engine/Source/Runtime/AssetRegistry/Public/AssetRegistry/IAssetRegistry.h
@@ -59,13 +59,14 @@ namespace EAssetAvailabilityProgressReportingType
 namespace UE::AssetRegistry
 {
 
-	// Please update LexToString when modifying this enum
+// Please update LexToString when modifying this enum
 enum class EScanFlags : uint32
 {
 	None = 0,
 	ForceRescan = 1 << 0,				// the paths will be scanned again, even if they were previously scanned
 	IgnoreDenyListScanFilters = 1 << 1,	// ignore deny list scan filters
 	WaitForInMemoryObjects = 1 << 2,	// update the tags of all assets that have loaded into memory before returning from the scan
+	IgnoreInvalidPathWarning = 1 << 3,  // By default, scans of invalid paths (unmounted, or /Temp) log a warning. Suppress that warning.
 };
 ENUM_CLASS_FLAGS(EScanFlags);
 
@@ -669,6 +670,7 @@ public:
 	/** Forces a rescan of specific filenames, call this when you need to refresh from disk */
 	UFUNCTION(BlueprintCallable, Category = "AssetRegistry")
 	virtual void ScanModifiedAssetFiles(const TArray<FString>& InFilePaths) = 0;
+	virtual void ScanModifiedAssetFiles(const TArray<FString>& InFilePaths, UE::AssetRegistry::EScanFlags ScanFlags) = 0;
 
 	/** Event for when one or more files have been blocked from the registry */
 	DECLARE_EVENT_OneParam( IAssetRegistry, FFilesBlockedEvent, const TArray<FString>& /*Files*/ );
