Strip trailing slash from paths before converting to FName

diff --git a/Engine/Source/Runtime/AssetRegistry/Private/AssetRegistry.cpp b/Engine/Source/Runtime/AssetRegistry/Private/AssetRegistry.cpp
--- a/Engine/Source/Runtime/AssetRegistry/Private/AssetRegistry.cpp
+++ b/Engine/Source/Runtime/AssetRegistry/Private/AssetRegistry.cpp
@@ -7063,7 +7063,8 @@ void FAssetRegistryImpl::OnDirectoryChanged(Impl::FEventContext& EventContext,
 				}
 				case FFileChangeData::FCA_Removed:
 				{
-					RemoveAssetPath(EventContext, *LongPackageName);
+					FName Path(UE::String::RemoveFromEnd(FStringView(LongPackageName), TEXTVIEW("/")));
+					RemoveAssetPath(EventContext, Path);
 					UE_LOG(LogAssetRegistry, Verbose, TEXT("Directory was removed from content directory: %s"), *File);
 					break;
 				}
