Fixing a bug where we could return Complete status from asset registry's TickGatherer even though we aren't actually able to finish processing the data. Adding additional diagnostics to asset registry WaitForCompletion to help identify the cause of a rare situation in which we seem to return from WaitForCompletion without actually having completed the scan.

diff --git a/Engine/Source/Runtime/AssetRegistry/Private/AssetRegistry.cpp b/Engine/Source/Runtime/AssetRegistry/Private/AssetRegistry.cpp
--- a/Engine/Source/Runtime/AssetRegistry/Private/AssetRegistry.cpp
+++ b/Engine/Source/Runtime/AssetRegistry/Private/AssetRegistry.cpp
@@ -1967,6 +1967,7 @@ void UAssetRegistryImpl::WaitForCompletion()
 
 	bool bInitialSearchStarted = false;
 	bool bInitialSearchCompleted = false;
+
 	// Try taking over the gather thread for a short time in case it is mostly done.
 	// But if it has more than a small amount of work to do, let the gather thread do that work
 	// while we consume the results in parallel.
@@ -1998,6 +1999,7 @@ void UAssetRegistryImpl::WaitForCompletion()
 	}
 #endif
 
+	bool bLocalHasSentFileLoadedEventBroadcast = bInitialSearchCompleted;
 	for (;;)
 	{
 		FEventContext EventContext;
@@ -2026,8 +2028,44 @@ void UAssetRegistryImpl::WaitForCompletion()
 		ProcessLoadedAssetsToUpdateCache(EventContext, Status, InterruptionContext);
 #endif
 		Broadcast(EventContext);
+		bLocalHasSentFileLoadedEventBroadcast |= EventContext.bHasSentFileLoadedEventBroadcast;
 		if (!IsTickActive(Status) && Status != EGatherStatus::WaitingForEvents)
 		{
+			if (Status == EGatherStatus::UnableToProgress)
+			{
+				UE_LOG(LogAssetRegistry, Warning, 
+					TEXT("UAssetRegistryImpl::WaitForCompletion exiting without completing because TickGatherer returned UnableToProgress. IsInGameThread() == %s"),
+					IsInGameThread() ? TEXT("TRUE") : TEXT("FALSE"));
+			}
+			else if (Status == EGatherStatus::Complete && bInitialSearchStarted)
+			{
+				// We only perform this validation if we are in a context where we expect the initial search to occur at all
+				// In some commandlets, e.g., we do not expect to run the initial search at all
+				UE::AssetRegistry::FInterfaceWriteScopeLock InterfaceScopeLock(InterfaceLock);
+				if (!GuardedData.IsInitialSearchCompleted())
+				{
+					UE_LOG(LogAssetRegistry, Error,
+						TEXT("Exiting from UAssetRegistryImpl::WaitForCompletion but IsInitialSearchCompleted is still false."
+							"EventContext.bHasSentFileLoadedEventBroadcast == %s; IsInGameThread() == %s"),
+						EventContext.bHasSentFileLoadedEventBroadcast ? TEXT("TRUE") : TEXT("FALSE"),
+						IsInGameThread() ? TEXT("TRUE") : TEXT("FALSE"));
+				}
+				else 
+				{
+					// If we are the main thread and we are exiting this function, one of two things should be true:
+					// a) The search was completed before we enter this function (i.e., bInitialSearchCompleted == true); or
+					// b) The search has completed during this function and, as the game thread, we have broadcast the FileLoadedEvent 
+					//    (i.e., EventContext.bHasSentFileLoadedEventBroadcast == true)
+					// Otherwise, something has gone wrong
+					ensureMsgf(bLocalHasSentFileLoadedEventBroadcast || bInitialSearchCompleted || !IsInGameThread(),
+						TEXT("Exiting from UAssetRegistryImpl::WaitForCompletion in an inconsistent state. "
+							 "bLocalHasSentFileLoadedEventBroadcast == %s; EventContext.bHasSentFileLoadedEventBroadcast == %s; bInitialSearchCompleted == %s; IsInGameThread() == %s"),
+						bLocalHasSentFileLoadedEventBroadcast ? TEXT("TRUE") : TEXT("FALSE"),
+						EventContext.bHasSentFileLoadedEventBroadcast ? TEXT("TRUE") : TEXT("FALSE"),
+						bInitialSearchCompleted ? TEXT("TRUE") : TEXT("FALSE"),
+						IsInGameThread() ? TEXT("TRUE") : TEXT("FALSE"));
+				}
+			}
 			break;
 		}
 
@@ -4905,6 +4943,10 @@ Impl::EGatherStatus FAssetRegistryImpl::TickGatherer(Impl::FEventContext& EventC
 				RecordTimer(); // OnInitialSearchComplete reads data set by RecordTimer
 				OnInitialSearchCompleted(EventContext);
 			}
+			else
+			{
+				OutStatus = EGatherStatus::UnableToProgress;
+			}
 		}
 	}
 
@@ -8608,6 +8650,7 @@ namespace Impl
 void FEventContext::Clear()
 {
 	bFileLoadedEventBroadcast = false;
+	bHasSentFileLoadedEventBroadcast = false;
 	ProgressUpdateData.Reset();
 	PathEvents.Empty();
 	AssetEvents.Empty();
@@ -8856,6 +8899,7 @@ void UAssetRegistryImpl::Broadcast(UE::AssetRegistry::Impl::FEventContext& Event
 
 		FileLoadedEvent.Broadcast();
 		EventContext.bFileLoadedEventBroadcast = false;
+		EventContext.bHasSentFileLoadedEventBroadcast = true;
 	}
 
 	if (EventContext.bScanEndedEventBroadcast)
diff --git a/Engine/Source/Runtime/AssetRegistry/Private/AssetRegistryImpl.h b/Engine/Source/Runtime/AssetRegistry/Private/AssetRegistryImpl.h
--- a/Engine/Source/Runtime/AssetRegistry/Private/AssetRegistryImpl.h
+++ b/Engine/Source/Runtime/AssetRegistry/Private/AssetRegistryImpl.h
@@ -663,6 +663,7 @@ struct FEventContext
 	bool bFileLoadedEventBroadcast = false;
 	bool bScanStartedEventBroadcast = false;
 	bool bScanEndedEventBroadcast = false;
+	bool bHasSentFileLoadedEventBroadcast = false;
 
 	/** Remove all stored events */
 	void Clear();
