Downgrading a warning to display to fix build. Adding additional logging to TickGatherer to try to understand how we are returning UnableToProgress during WaitForCompletion.

diff --git a/Engine/Source/Runtime/AssetRegistry/Private/AssetRegistry.cpp b/Engine/Source/Runtime/AssetRegistry/Private/AssetRegistry.cpp
--- a/Engine/Source/Runtime/AssetRegistry/Private/AssetRegistry.cpp
+++ b/Engine/Source/Runtime/AssetRegistry/Private/AssetRegistry.cpp
@@ -2033,7 +2033,7 @@ void UAssetRegistryImpl::WaitForCompletion()
 		{
 			if (Status == EGatherStatus::UnableToProgress)
 			{
-				UE_LOG(LogAssetRegistry, Warning, 
+				UE_LOG(LogAssetRegistry, Display,
 					TEXT("UAssetRegistryImpl::WaitForCompletion exiting without completing because TickGatherer returned UnableToProgress. IsInGameThread() == %s"),
 					IsInGameThread() ? TEXT("TRUE") : TEXT("FALSE"));
 			}
@@ -4945,6 +4945,12 @@ Impl::EGatherStatus FAssetRegistryImpl::TickGatherer(Impl::FEventContext& EventC
 			}
 			else
 			{
+				UE_LOG(LogAssetRegistry, Display, TEXT("TickGatherer returning UnableToProgress because bCanCompleteInitialSearch is false but our work is otherwise complete. "
+					"bPreloadingComplete == %s; IsEngineStartupModuleLoadingComplete() == %s; bLocalIsInGameThread == %s"),
+					bPreloadingComplete ? TEXT("TRUE") : TEXT("FALSE"),
+					IsEngineStartupModuleLoadingComplete() ? TEXT("TRUE") : TEXT("FALSE"),
+					bLocalIsInGameThread ? TEXT("TRUE") : TEXT("FALSE"));
+
 				OutStatus = EGatherStatus::UnableToProgress;
 			}
 		}
