AssetRegistry: Suppress TickGatherer returning UnableToProgress spam sent from the assetDatagatherer thread.
#rnx

diff --git a/Engine/Source/Runtime/AssetRegistry/Private/AssetRegistry.cpp b/Engine/Source/Runtime/AssetRegistry/Private/AssetRegistry.cpp
--- a/Engine/Source/Runtime/AssetRegistry/Private/AssetRegistry.cpp
+++ b/Engine/Source/Runtime/AssetRegistry/Private/AssetRegistry.cpp
@@ -4945,11 +4945,14 @@ Impl::EGatherStatus FAssetRegistryImpl::TickGatherer(Impl::FEventContext& EventC
 			}
 			else
 			{
-				UE_LOG(LogAssetRegistry, Display, TEXT("TickGatherer returning UnableToProgress because bCanCompleteInitialSearch is false but our work is otherwise complete. "
-					"bPreloadingComplete == %s; IsEngineStartupModuleLoadingComplete() == %s; bLocalIsInGameThread == %s"),
-					bPreloadingComplete ? TEXT("TRUE") : TEXT("FALSE"),
-					IsEngineStartupModuleLoadingComplete() ? TEXT("TRUE") : TEXT("FALSE"),
-					bLocalIsInGameThread ? TEXT("TRUE") : TEXT("FALSE"));
+				if (bLocalIsInGameThread)
+				{
+					UE_LOG(LogAssetRegistry, Display, TEXT("TickGatherer returning UnableToProgress because bCanCompleteInitialSearch is false but our work is otherwise complete. "
+						"bPreloadingComplete == %s; IsEngineStartupModuleLoadingComplete() == %s; bLocalIsInGameThread == %s"),
+						bPreloadingComplete ? TEXT("TRUE") : TEXT("FALSE"),
+						IsEngineStartupModuleLoadingComplete() ? TEXT("TRUE") : TEXT("FALSE"),
+						bLocalIsInGameThread ? TEXT("TRUE") : TEXT("FALSE"));
+				}
 
 				OutStatus = EGatherStatus::UnableToProgress;
 			}
