Restoring log notification by asset registry when WaitForCompletion fails to complete from display level to warning now that a common spurious case has been eliminated.

#rnx

diff --git a/Engine/Source/Runtime/AssetRegistry/Private/AssetRegistry.cpp b/Engine/Source/Runtime/AssetRegistry/Private/AssetRegistry.cpp
--- a/Engine/Source/Runtime/AssetRegistry/Private/AssetRegistry.cpp
+++ b/Engine/Source/Runtime/AssetRegistry/Private/AssetRegistry.cpp
@@ -2033,7 +2033,7 @@ void UAssetRegistryImpl::WaitForCompletion()
 		{
 			if (Status == EGatherStatus::UnableToProgress)
 			{
-				UE_LOG(LogAssetRegistry, Display,
+				UE_LOG(LogAssetRegistry, Warning,
 					TEXT("UAssetRegistryImpl::WaitForCompletion exiting without completing because TickGatherer returned UnableToProgress. IsInGameThread() == %s"),
 					IsInGameThread() ? TEXT("TRUE") : TEXT("FALSE"));
 			}
