Converting asset registry to use FCoreDelegates::OnAllModuleLoadingPhasesComplete (instead of FPluginManager::OnLoadingPhaseComplete) to know when engine startup modules have finished loading. Adding additional diagnostic data to the case where AssetRegistry WaitForCompletion can return due to TickGatherer returning UnableToProgress.

diff --git a/Engine/Source/Runtime/AssetRegistry/Private/AssetDataGatherer.cpp b/Engine/Source/Runtime/AssetRegistry/Private/AssetDataGatherer.cpp
--- a/Engine/Source/Runtime/AssetRegistry/Private/AssetDataGatherer.cpp
+++ b/Engine/Source/Runtime/AssetRegistry/Private/AssetDataGatherer.cpp
@@ -3587,7 +3587,6 @@ FAssetDataGatherer::FAssetDataGatherer(const TArray<FString>& InLongPackageNames
 	, IsStopped(0)
 	, IsGatheringPaused(0)
 	, IsProcessingPaused(0)
-	, bInitialPluginsLoaded(false)
 	, bSaveAsyncCacheTriggered(false)
 	, CurrentSearchTime(0.)
 	, LastCacheWriteTime(0.0)
@@ -3641,12 +3640,10 @@ FAssetDataGatherer::FAssetDataGatherer(const TArray<FString>& InLongPackageNames
 	Discovery = MakeUnique<UE::AssetDataGather::Private::FAssetDataDiscovery>(InLongPackageNamesDenyList,
 		InMountRelativePathsDenyList, bAsyncEnabled);
 	FilesToSearch = MakeUnique<UE::AssetDataGather::Private::FFilesToSearch>();
-	FCoreDelegates::OnAllModuleLoadingPhasesComplete.AddRaw(this, &FAssetDataGatherer::OnAllModuleLoadingPhasesComplete);
 }
 
 FAssetDataGatherer::~FAssetDataGatherer()
 {
-	FCoreDelegates::OnAllModuleLoadingPhasesComplete.RemoveAll(this);
 	EnsureCompletion();
 	NewCachedAssetDataMap.Empty();
 	DiskCachedAssetDataMap.Empty();
@@ -4012,7 +4009,7 @@ FAssetDataGatherer::ETickResult FAssetDataGatherer::TickInternal(double& TickSta
 		{
 			IngestDiscoveryResults();
 		}
-		if (bInitialPluginsLoaded && !bFlushedRetryFiles)
+		if (IsEngineStartupModuleLoadingComplete() && !bFlushedRetryFiles)
 		{
 			bFlushedRetryFiles = true;
 			FilesToSearch->RetryLaterRetryFiles();
@@ -4412,7 +4409,7 @@ bool FAssetDataGatherer::ReadAssetFile(const FString& AssetLongPackageName, cons
 		// If we're missing a custom version, we might be able to load this package later once the module containing that version is loaded...
 		//   -	We can only attempt a retry in editors (not commandlets) that haven't yet finished initializing (!GIsRunning), as we 
 		//		have no guarantee that a commandlet or an initialized editor is going to load any more modules/plugins
-		const bool bAllowRetry = GIsEditor && !bInitialPluginsLoaded;
+		const bool bAllowRetry = GIsEditor && !IsEngineStartupModuleLoadingComplete();
 		if (OpenPackageResult == FPackageReader::EOpenPackageResult::CustomVersionMissing)
 		{
 			OutCanRetry = bAllowRetry;
@@ -4779,7 +4776,6 @@ bool FAssetDataGatherer::IsComplete() const
 
 void FAssetDataGatherer::SetInitialPluginsLoaded()
 {
-	bInitialPluginsLoaded = true;
 	FGathererScopeLock ResultsScopeLock(&ResultsLock);
 	SetIsIdle(false);
 }
@@ -5476,13 +5472,6 @@ FStringView FAssetDataGatherer::NormalizeLongPackageName(FStringView LongPackage
 	return LongPackageName;
 }
 
-void FAssetDataGatherer::OnAllModuleLoadingPhasesComplete()
-{
-	CHECK_IS_NOT_LOCKED_CURRENT_THREAD(ResultsLock);
-	FGathererScopeLock ResultsScopeLock(&ResultsLock);
-	SetIsIdle(false);
-}
-
 namespace UE::AssetDataGather::Private
 {
 
diff --git a/Engine/Source/Runtime/AssetRegistry/Private/AssetDataGatherer.h b/Engine/Source/Runtime/AssetRegistry/Private/AssetDataGatherer.h
--- a/Engine/Source/Runtime/AssetRegistry/Private/AssetDataGatherer.h
+++ b/Engine/Source/Runtime/AssetRegistry/Private/AssetDataGatherer.h
@@ -374,7 +374,6 @@ private:
 	/** Convert the LongPackageName into our normalized version. */
 	static FStringView NormalizeLongPackageName(FStringView LongPackageName);
 
-	void OnAllModuleLoadingPhasesComplete();
 private:
 
 	/**
@@ -427,8 +426,6 @@ private:
 	 * Pointer is constant during threading. Object pointed to internally provides threadsafety.
 	 */
 	TUniquePtr<UE::AssetDataGather::Private::FAssetDataDiscovery> Discovery;
-	/** Async only. Set to true once initial plugins have been loaded. */
-	std::atomic<bool> bInitialPluginsLoaded;
 	/** True when TickInternal requests periodic or final save of the async cache. */
 	std::atomic<bool> bSaveAsyncCacheTriggered;
 	/** True if the current process allows reading of AssetDataGatherer cache files. */
diff --git a/Engine/Source/Runtime/AssetRegistry/Private/AssetRegistry.cpp b/Engine/Source/Runtime/AssetRegistry/Private/AssetRegistry.cpp
--- a/Engine/Source/Runtime/AssetRegistry/Private/AssetRegistry.cpp
+++ b/Engine/Source/Runtime/AssetRegistry/Private/AssetRegistry.cpp
@@ -1189,10 +1189,9 @@ void UAssetRegistryImpl::InitializeEvents(UE::AssetRegistry::Impl::FInitializeCo
 	FCoreDelegates::OnPostEngineInit.AddUObject(this, &UAssetRegistryImpl::OnRefreshNativeClasses);
 
 	IPluginManager& PluginManager = IPluginManager::Get();
-	ELoadingPhase::Type LoadingPhase = PluginManager.GetLastCompletedLoadingPhase();
-	if (LoadingPhase == ELoadingPhase::None || LoadingPhase < ELoadingPhase::PostEngineInit)
+	if (!IsEngineStartupModuleLoadingComplete())
 	{
-		PluginManager.OnLoadingPhaseComplete().AddUObject(this, &UAssetRegistryImpl::OnPluginLoadingPhaseComplete);
+		FCoreDelegates::OnAllModuleLoadingPhasesComplete.AddUObject(this, &UAssetRegistryImpl::OnInitialPluginLoadingComplete);
 	}
 }
 
@@ -1292,19 +1291,15 @@ void FAssetRegistryImpl::InitRedirectors(Impl::FEventContext& EventContext,
 
 }
 
-void UAssetRegistryImpl::OnPluginLoadingPhaseComplete(ELoadingPhase::Type LoadingPhase, bool bPhaseSuccessful)
+void UAssetRegistryImpl::OnInitialPluginLoadingComplete()
 {
-	if (LoadingPhase != ELoadingPhase::PostEngineInit)
-	{
-		return;
-	}
 	{
 		LLM_SCOPE(ELLMTag::AssetRegistry);
 		UE::AssetRegistry::FInterfaceWriteScopeLock InterfaceScopeLock(InterfaceLock);
-		GuardedData.OnPostEngineInit(bPhaseSuccessful);
+		GuardedData.OnPostEngineInit(true);
 	}
 
-	IPluginManager::Get().OnLoadingPhaseComplete().RemoveAll(this);
+	FCoreDelegates::OnAllModuleLoadingPhasesComplete.RemoveAll(this);
 }
 
 namespace UE::AssetRegistry
@@ -2034,8 +2029,10 @@ void UAssetRegistryImpl::WaitForCompletion()
 			if (Status == EGatherStatus::UnableToProgress)
 			{
 				UE_LOG(LogAssetRegistry, Display,
-					TEXT("UAssetRegistryImpl::WaitForCompletion exiting without completing because TickGatherer returned UnableToProgress. IsInGameThread() == %s"),
-					IsInGameThread() ? TEXT("TRUE") : TEXT("FALSE"));
+					TEXT("UAssetRegistryImpl::WaitForCompletion exiting without completing because TickGatherer returned UnableToProgress. "
+					"IsInGameThread() == %s; IsEngineStartupModuleLoadingComplete() == %s"),
+					IsInGameThread() ? TEXT("TRUE") : TEXT("FALSE"),
+					IsEngineStartupModuleLoadingComplete() ? TEXT("TRUE") : TEXT("FALSE"));
 			}
 			else if (Status == EGatherStatus::Complete && bInitialSearchStarted)
 			{
diff --git a/Engine/Source/Runtime/AssetRegistry/Private/AssetRegistry.h b/Engine/Source/Runtime/AssetRegistry/Private/AssetRegistry.h
--- a/Engine/Source/Runtime/AssetRegistry/Private/AssetRegistry.h
+++ b/Engine/Source/Runtime/AssetRegistry/Private/AssetRegistry.h
@@ -259,8 +259,8 @@ private:
 	/** Called to refresh the native classes list, called at end of engine initialization. */
 	void OnRefreshNativeClasses();
 
-	/** Called from the PluginManager's loading phase, used to scan classes that were loaded by plugins. */
-	void OnPluginLoadingPhaseComplete(ELoadingPhase::Type LoadingPhase, bool bPhaseSuccessful);
+	/** Called from LaunchEngineLoop via SetEngineStartupModuleLoadingComplete after plugins are loaded, used to scan classes that were loaded by plugins. */
+	void OnInitialPluginLoadingComplete();
 
 	/** Shared helper for Scan*Synchronous function */
 	void ScanPathsSynchronousInternal(const TArray<FString>& InDirs, const TArray<FString>& InFiles,
