AssetRegistry: PackageReader: Downgrade ensure for invalid objectname to a warning. The fallback for the warning being present is that we cannot do AssetRegistry queries against the bad object. This is a good enough fallback that an ensure or even an error is higher severity than necessary.
#rnx

diff --git a/Engine/Source/Runtime/AssetRegistry/Private/PackageReader.cpp b/Engine/Source/Runtime/AssetRegistry/Private/PackageReader.cpp
--- a/Engine/Source/Runtime/AssetRegistry/Private/PackageReader.cpp
+++ b/Engine/Source/Runtime/AssetRegistry/Private/PackageReader.cpp
@@ -1726,10 +1726,10 @@ namespace UE::AssetRegistry
 			// if we do not have a full object path already, build it
 			if (!bFullObjectPath)
 			{
-				// if we do not have a full object path, ensure that we have a top level object for the package and not a sub object
-				if (!ensureMsgf(!ObjectPackageData.ObjectPath.Contains(TEXT("."), ESearchCase::CaseSensitive),
-					TEXT("Package is loadable but its AssetRegistry data is corrupt: %s. Reason: Cannot make FAssetData for sub object %s."),
-					*PackageName , *ObjectPackageData.ObjectPath))
+				// if we do not have a full object path, ensure that we have a top level object for the package and not
+				// a subobject. This warning can also fire if a top level object was created with the invalid character
+				// '.' in its objectname. Savepackage is supposed to prevent that, but we do not enforce it yet.
+				if (ObjectPackageData.ObjectPath.Contains(TEXT("."), ESearchCase::CaseSensitive))
 				{
 					UE_LOG(LogAssetRegistry, Warning,
 						TEXT("Package is loadable but its AssetRegistry data is corrupt: %s. Reason: Cannot make FAssetData for sub object %s."),
