AssetRegistry: Add more feedback for the error when the filename looked up for a file on disk differs unexpectedly by more than capitalization from the filename we queried.
#rnx

diff --git a/Engine/Source/Runtime/AssetRegistry/Private/AssetDataGatherer.cpp b/Engine/Source/Runtime/AssetRegistry/Private/AssetDataGatherer.cpp
--- a/Engine/Source/Runtime/AssetRegistry/Private/AssetDataGatherer.cpp
+++ b/Engine/Source/Runtime/AssetRegistry/Private/AssetDataGatherer.cpp
@@ -2359,8 +2359,8 @@ void FPathExistence::LoadExistenceData()
 		else
 		{
 			UE_LOG(LogAssetRegistry, Error,
-				TEXT("FPathExistence failed to gather correct capitalization from disk for %s, because GetFilenameOnDisk returned a non-matching filename"),
-				*LocalAbsPath);
+				TEXT("FPathExistence failed to gather correct capitalization from disk for %s, because GetFilenameOnDisk returned non-matching filename '%s'."),
+				*LocalAbsPath, *CorrectedCapitalization);
 		}
 
 		ModificationTime = StatData.ModificationTime;
@@ -2383,8 +2383,8 @@ void FPathExistence::LoadExistenceData()
 			else
 			{
 				UE_LOG(LogAssetRegistry, Error,
-					TEXT("FPathExistence failed to gather correct capitalization from disk for %s, because GetFilenameOnDisk returned a non-matching filename"),
-					*LocalAbsPath);
+					TEXT("FPathExistence failed to gather correct capitalization from disk for %s, because GetFilenameOnDisk returned non-matching filename '%s'."),
+					*LocalAbsPath, *CorrectedCapitalization);
 			}
 			PathType = EType::MissingButDirExists;
 		}
