Fixing an issue where we could cause a deadlock between the GC system and the AssetRegistry by not calling FCoreDelegates::OnEnginePreExit before exiting on engine initialization failure. We now rely on OnPreExit instead for AssetRegistry shutdown.

diff --git a/Engine/Source/Runtime/AssetRegistry/Private/AssetRegistry.cpp b/Engine/Source/Runtime/AssetRegistry/Private/AssetRegistry.cpp
--- a/Engine/Source/Runtime/AssetRegistry/Private/AssetRegistry.cpp
+++ b/Engine/Source/Runtime/AssetRegistry/Private/AssetRegistry.cpp
@@ -1178,7 +1178,9 @@ void UAssetRegistryImpl::InitializeEvents(UE::AssetRegistry::Impl::FInitializeCo
 	UE::AssetDependencyGatherer::Private::FRegisteredAssetDependencyGatherer::OnAssetDependencyGathererRegistered.AddUObject(this, &UAssetRegistryImpl::OnAssetDependencyGathererRegistered);
 #endif // WITH_EDITOR
 
-	FCoreDelegates::OnEnginePreExit.AddUObject(this, &UAssetRegistryImpl::OnEnginePreExit);
+	// We use OnPreExit and not OnEnginePreExit because OnPreExit will be called if there's an error in engine init and 
+	// we never get through OnPostEngineInit.
+	FCoreDelegates::OnPreExit.AddUObject(this, &UAssetRegistryImpl::OnPreExit);
 
 	// Listen for new content paths being added or removed at runtime.  These are usually plugin-specific asset paths that
 	// will be loaded a bit later on.
@@ -1656,7 +1658,7 @@ void UAssetRegistryImpl::OnAssetDependencyGathererRegistered()
 }
 #endif
 
-void UAssetRegistryImpl::OnEnginePreExit()
+void UAssetRegistryImpl::OnPreExit()
 {
 	LLM_SCOPE(ELLMTag::AssetRegistry);
 
@@ -1685,7 +1687,7 @@ void UAssetRegistryImpl::FinishDestroy()
 		FPackageName::OnContentPathMounted().RemoveAll(this);
 		FPackageName::OnContentPathDismounted().RemoveAll(this);
 		FCoreDelegates::OnPostEngineInit.RemoveAll(this);
-		FCoreDelegates::OnEnginePreExit.RemoveAll(this);
+		FCoreDelegates::OnPreExit.RemoveAll(this);
 		IPluginManager::Get().OnLoadingPhaseComplete().RemoveAll(this);
 
 #if WITH_EDITOR
diff --git a/Engine/Source/Runtime/AssetRegistry/Private/AssetRegistry.h b/Engine/Source/Runtime/AssetRegistry/Private/AssetRegistry.h
--- a/Engine/Source/Runtime/AssetRegistry/Private/AssetRegistry.h
+++ b/Engine/Source/Runtime/AssetRegistry/Private/AssetRegistry.h
@@ -219,7 +219,7 @@ protected:
 	virtual bool SetPrimaryAssetIdForObjectPath(const FSoftObjectPath& ObjectPath, FPrimaryAssetId PrimaryAssetId) override;
 
 private:
-	void OnEnginePreExit();
+	void OnPreExit();
 #if WITH_EDITOR
 	void OnFEngineLoopInitCompleteSearchAllAssets();
 	/** Called when new gatherer is registered. Requires subsequent call to RebuildAssetDependencyGathererMapIfNeeded */
