AssetRegistry: Downgrade error about corrected capitalization not working to a warning. It fires when disk junctions are used. Restore this to an error when AssetRegistry is hardened to support junctions.
#rnx

diff --git a/Engine/Source/Runtime/AssetRegistry/Private/AssetDataGatherer.cpp b/Engine/Source/Runtime/AssetRegistry/Private/AssetDataGatherer.cpp
--- a/Engine/Source/Runtime/AssetRegistry/Private/AssetDataGatherer.cpp
+++ b/Engine/Source/Runtime/AssetRegistry/Private/AssetDataGatherer.cpp
@@ -2358,7 +2358,7 @@ void FPathExistence::LoadExistenceData()
 		}
 		else
 		{
-			UE_LOG(LogAssetRegistry, Error,
+			UE_LOG(LogAssetRegistry, Warning,
 				TEXT("FPathExistence failed to gather correct capitalization from disk for %s, because GetFilenameOnDisk returned non-matching filename '%s'."),
 				*LocalAbsPath, *CorrectedCapitalization);
 		}
@@ -2382,7 +2382,7 @@ void FPathExistence::LoadExistenceData()
 			}
 			else
 			{
-				UE_LOG(LogAssetRegistry, Error,
+				UE_LOG(LogAssetRegistry, Warning,
 					TEXT("FPathExistence failed to gather correct capitalization from disk for %s, because GetFilenameOnDisk returned non-matching filename '%s'."),
 					*LocalAbsPath, *CorrectedCapitalization);
 			}
