AssetRegistry: Turn on async scan and use of cache in editor -game and editor -server modes; the AssetRegistry is used in all game modes during startup and needs to be cached in all game modes. Gather depends data in those modes as well so that they share the cache with the editor IDE. Delete all code for the non-monolithic cache; the monolithic cache now handles being used for various request differenting sets without any loss of data used by the various sets so we no longer need a cache per set of requested directories. Turning off the non-monolithic cache also fixes an issue we would otherwise have to fix with -game; it was not using the monolithic cache and was running uncached. Fix the initial SearchAllAssets scan to still be run (and log its performance) when it is triggered during by SearchAllAssetsInitialAsync but we are running in game and so cannot use a separate thread and must still run the gather synchronous.
[RN] Minor,Cooking

diff --git a/Engine/Config/BaseEngine.ini b/Engine/Config/BaseEngine.ini
--- a/Engine/Config/BaseEngine.ini
+++ b/Engine/Config/BaseEngine.ini
@@ -2797,6 +2797,14 @@ bFilterAssetDataWithNoTags=false
 bFilterDependenciesWithNoTags=false
 bFilterSearchableNames=true
 
+; The editor can run the game or dedicated server on uncooked content, via UnrealEditor.exe <ProjectName> -game
+; This is bad for performance but is supported for use in PlayInEditor
+; When running on uncooked content, there is no cooked asset registry, so the game needs to scan the AssetRegistry
+; the same way the editor does to have the same AssetRegistry available that the cooked game has. This adds cputime and
+; memory cost to the editor game startup, so some projects may want to turn it off for their game when running editor -game.
+; Set EditorGameScansAR=false to turn it off.
+EditorGameScansAR=true
+
 ; Commandlets by default do not SearchAllAssets because many do not need it. A commandlet can be opted in
 ; using CommandletsUsingAR. Format is the shortname of the commandlet, optionally without the "Commandlet" suffix.
 +CommandletsUsingAR=DataValidation
diff --git a/Engine/Source/Runtime/AssetRegistry/Private/AssetDataGatherer.cpp b/Engine/Source/Runtime/AssetRegistry/Private/AssetDataGatherer.cpp
--- a/Engine/Source/Runtime/AssetRegistry/Private/AssetDataGatherer.cpp
+++ b/Engine/Source/Runtime/AssetRegistry/Private/AssetDataGatherer.cpp
@@ -183,7 +183,8 @@ void FPreloadSettings::Initialize()
 		AssetRegistryCacheRootFolder = FPaths::ProjectIntermediateDir();
 	}
 	bForceDependsGathering = FParse::Param(FCommandLine::Get(), TEXT("ForceDependsGathering"));
-	bGatherDependsData = (GIsEditor && !FParse::Param(FCommandLine::Get(), TEXT("NoDependsGathering"))) || bForceDependsGathering;
+	constexpr bool bEditorExecutable = WITH_EDITOR;
+	bGatherDependsData = (bEditorExecutable && !FParse::Param(FCommandLine::Get(), TEXT("NoDependsGathering"))) || bForceDependsGathering;
 	bool bNoAssetRegistryCache = FParse::Param(FCommandLine::Get(), TEXT("NoAssetRegistryCache"));
 	bool bNoAssetRegistryDiscoveryCache = bNoAssetRegistryCache || FParse::Param(FCommandLine::Get(), TEXT("NoAssetRegistryDiscoveryCache"));
 	bool bNoAssetRegistryCacheRead = FParse::Param(FCommandLine::Get(), TEXT("NoAssetRegistryCacheRead"));
@@ -268,11 +269,11 @@ void FPreloadSettings::Initialize()
 
 	bool bAsyncEnabled = FPlatformProcess::SupportsMultithreading() && FTaskGraphInterface::IsRunning();
 
-	MonolithicCacheBaseFilename = AssetRegistryCacheRootFolder / (bGatherDependsData ? TEXT("CachedAssetRegistry") : TEXT("CachedAssetRegistryNoDeps"));
+	CacheBaseFilename = AssetRegistryCacheRootFolder / (bGatherDependsData ? TEXT("CachedAssetRegistry") : TEXT("CachedAssetRegistryNoDeps"));
 #if UE_EDITOR // See note on FPreloader for why we only allow preloading if UE_EDITOR
-	bMonolithicCacheActivatedDuringPreload = bAsyncEnabled && UE::AssetRegistry::ShouldSearchAllAssetsAtStart();
+	bPreloadCache = bAsyncEnabled && UE::AssetRegistry::ShouldSearchAllAssetsAtStart();
 #else
-	bMonolithicCacheActivatedDuringPreload = false;
+	bPreloadCache = false;
 #endif
 }
 bool FPreloadSettings::IsGatherCacheReadEnabled() const
@@ -295,13 +296,9 @@ bool FPreloadSettings::IsDiscoveryCacheInvalidateEnabled() const
 {
 	return bDiscoveryCacheInvalidateEnabled;
 }
-bool FPreloadSettings::IsMonolithicCacheActivatedDuringPreload() const
+bool FPreloadSettings::IsPreloadCache() const
 {
-	return bMonolithicCacheActivatedDuringPreload;
-}
-bool FPreloadSettings::IsPreloadMonolithicCache() const
-{
-	return bGatherCacheReadEnabled && bMonolithicCacheActivatedDuringPreload;
+	return bGatherCacheReadEnabled && bPreloadCache;
 }
 bool FPreloadSettings::IsGatherDependsData() const
 {
@@ -311,26 +308,30 @@ bool FPreloadSettings::IsForceDependsGathering() const
 {
 	return bForceDependsGathering;
 }
-FString FPreloadSettings::GetLegacyMonolithicCacheFilename() const
+FString FPreloadSettings::GetLegacyCacheFilename() const
+{
+	return CacheBaseFilename + TEXT(".bin");
+}
+FString FPreloadSettings::GetLegacyNonMonolithicCacheDirectory() const
 {
-	return MonolithicCacheBaseFilename + TEXT(".bin");
+	return UE::AssetDataGather::Private::GPreloadSettings.GetAssetRegistryCacheRootFolder() / TEXT("AssetRegistryCache");
 }
-const FString& FPreloadSettings::GetMonolithicCacheBaseFilename() const
+const FString& FPreloadSettings::GetCacheBaseFilename() const
 {
-	return MonolithicCacheBaseFilename;
+	return CacheBaseFilename;
 }
 const FString& FPreloadSettings::GetAssetRegistryCacheRootFolder() const
 {
 	return AssetRegistryCacheRootFolder;
 }
 
-TArray<FString> FPreloadSettings::FindShardedMonolithicCacheFiles() const
+TArray<FString> FPreloadSettings::FindShardedCacheFiles() const
 {
 	TArray<FString> CachePaths;
-	IFileManager::Get().FindFiles(CachePaths, *(GetMonolithicCacheBaseFilename() + TEXT("_*.bin")), /* Files */ true, /* Directories */ false);
+	IFileManager::Get().FindFiles(CachePaths, *(GetCacheBaseFilename() + TEXT("_*.bin")), /* Files */ true, /* Directories */ false);
 	if (CachePaths.Num())
 	{
-		FString Directory = FPaths::GetPath(GetMonolithicCacheBaseFilename());
+		FString Directory = FPaths::GetPath(GetCacheBaseFilename());
 		for (FString& Path : CachePaths)
 		{
 			Path = Directory / Path;
@@ -364,7 +365,7 @@ struct FCachePayload
 
 void SerializeCacheSave(FAssetRegistryWriter& Ar, const TArray<TPair<FName, FDiskCachedAssetData*>>& AssetsToSave);
 FCachePayload SerializeCacheLoad(FAssetRegistryReader& Ar);
-TArray<FCachePayload> LoadCacheFiles(TConstArrayView<FString> CacheFilenames, bool bIsMonolithicCache);
+TArray<FCachePayload> LoadCacheFiles(TConstArrayView<FString> CacheFilenames);
 
 
 /** InOutResult = Value, but without shrinking the string to fit. */
@@ -3506,7 +3507,7 @@ private:
 };
 
 #if UE_EDITOR
-/** A class to preload the monolithic cache used by FAssetDataGatherer. Preloading the cache allows us to
+/** A class to preload the cache used by FAssetDataGatherer. Preloading the cache allows us to
  * start very early in editor startup, so that we have time to finish the cache load before the engine starts
  * making package load requests that need to use the gathered data.
  *
@@ -3551,9 +3552,9 @@ private:
 	{
 		GPreloadSettings.Initialize();
 
-		if (GPreloadSettings.IsPreloadMonolithicCache())
+		if (GPreloadSettings.IsPreloadCache())
 		{
-			TArray<FString> CachePaths = GPreloadSettings.FindShardedMonolithicCacheFiles();
+			TArray<FString> CachePaths = GPreloadSettings.FindShardedCacheFiles();
 			if (CachePaths.Num())
 			{
 				PreloadReady = Async(EAsyncExecution::TaskGraph, [this, CachePaths=MoveTemp(CachePaths)]() { LoadAsync(CachePaths); });
@@ -3565,7 +3566,7 @@ private:
 	void LoadAsync(const TArray<FString>& Paths)
 	{
 		LLM_SCOPE(ELLMTag::AssetRegistry);
-		Payloads = LoadCacheFiles(Paths, true /* bIsMonolithicCache */);
+		Payloads = LoadCacheFiles(Paths);
 	}
 
 	TFuture<void> PreloadReady;
@@ -3590,14 +3591,14 @@ FAssetDataGatherer::FAssetDataGatherer(const TArray<FString>& InLongPackageNames
 	, bSaveAsyncCacheTriggered(false)
 	, CurrentSearchTime(0.)
 	, LastCacheWriteTime(0.0)
-	, bHasLoadedMonolithicCache(false)
+	, bHasLoadedCache(false)
 	, bDiscoveryIsComplete(false)
 	, bIsComplete(false)
 	, bIsIdle(false)
 	, bFirstTickAfterIdle(true)
 	, bFinishedInitialDiscovery(false)
 	, WaitBatchCount(-1)
-	, LastMonolithicCacheSaveUncachedAssetFiles(0)
+	, LastCacheSaveNumUncachedAssetFiles(0)
 	, CacheInUseCount(0)
 	, bIsSavingAsyncCache(false)
 	, bFlushedRetryFiles(false)
@@ -3612,14 +3613,11 @@ FAssetDataGatherer::FAssetDataGatherer(const TArray<FString>& InLongPackageNames
 	TickInternalBatchSize = NumGatherThreads * AssetDataGathererConstants::SingleThreadFilesPerBatch;
 
 	GPreloadSettings.Initialize();
-	bGatherAssetPackageData = GIsEditor || GPreloadSettings.IsForceDependsGathering();
+	constexpr bool bEditorExecutable = WITH_EDITOR;
+	bGatherAssetPackageData = bEditorExecutable || GPreloadSettings.IsForceDependsGathering();
 	bGatherDependsData = GPreloadSettings.IsGatherDependsData();
 	bCacheReadEnabled = GPreloadSettings.IsGatherCacheReadEnabled();
 	bCacheWriteEnabled = GPreloadSettings.IsGatherCacheWriteEnabled();
-	// If IsMonolithicCacheActivatedDuringPreload is true, we are already instructed to use the MonolithicCache.
-	// Otherwise it may be set to true later if game/commandlet calls SearchAllAssets.
-	bReadMonolithicCache = bCacheReadEnabled && GPreloadSettings.IsMonolithicCacheActivatedDuringPreload();
-	bWriteMonolithicCache = bCacheWriteEnabled && GPreloadSettings.IsMonolithicCacheActivatedDuringPreload();
 	LastCacheWriteTime = FPlatformTime::Seconds();
 
 #if WITH_EDITOR || !UE_BUILD_SHIPPING
@@ -3668,22 +3666,6 @@ void FAssetDataGatherer::OnInitialSearchCompleted()
 	}
 }
 
-void FAssetDataGatherer::ActivateMonolithicCache()
-{
-	FGathererScopeLock ResultsScopeLock(&ResultsLock);
-	bool bNewWriteMonolithicCache = bCacheWriteEnabled;
-	bool bNewReadMonolithicCache = bCacheReadEnabled;
-	if ((!bNewWriteMonolithicCache || bWriteMonolithicCache) &&
-		(!bNewReadMonolithicCache || bReadMonolithicCache))
-	{
-		return;
-	}
-
-	bWriteMonolithicCache = bNewWriteMonolithicCache;
-	bReadMonolithicCache = bNewReadMonolithicCache;
-	LastCacheWriteTime = FPlatformTime::Seconds();
-}
-
 void FAssetDataGatherer::StartAsync()
 {
 	if (bAsyncEnabled && !Thread)
@@ -3749,7 +3731,7 @@ void FAssetDataGatherer::InnerTickLoop(bool bInSynchronousTick, bool bContribute
 	// The dedicated async thread always contributes
 	bContributeToCacheSave = !bInSynchronousTick || (Thread == nullptr && bContributeToCacheSave);
 
-	bool bShouldSaveMonolithicCache = false;
+	bool bShouldSaveCache = false;
 	TArray<TPair<FName, FDiskCachedAssetData*>> AssetsToSave;
 	{
 		CHECK_IS_NOT_LOCKED_CURRENT_THREAD(ResultsLock);
@@ -3820,23 +3802,23 @@ void FAssetDataGatherer::InnerTickLoop(bool bInSynchronousTick, bool bContribute
 
 		if (bContributeToCacheSave)
 		{
-			TryReserveSaveMonolithicCache(bShouldSaveMonolithicCache, AssetsToSave);
-			if (bShouldSaveMonolithicCache)
+			TryReserveSaveCache(bShouldSaveCache, AssetsToSave);
+			if (bShouldSaveCache)
 			{
 				CacheInUseCount++;
 			}
 		}
 	}
-	if (bShouldSaveMonolithicCache)
+	if (bShouldSaveCache)
 	{
-		SaveMonolithicCacheFile(AssetsToSave);
+		SaveCacheFile(AssetsToSave);
 	}
 }
 
-void FAssetDataGatherer::SaveMonolithicCacheFile(const TArray<TPair<FName,FDiskCachedAssetData*>>& AssetsToSave)
+void FAssetDataGatherer::SaveCacheFile(const TArray<TPair<FName,FDiskCachedAssetData*>>& AssetsToSave)
 {
 	using namespace UE::AssetDataGather::Private;
-	TRACE_CPUPROFILER_EVENT_SCOPE_STR("Save Monolithic Cache")
+	TRACE_CPUPROFILER_EVENT_SCOPE_STR("Save Cache")
 	
 	// If we recently saved or loaded the file then pause for 0.5 seconds before trying to save on
 	// top of it, to avoid failure to be able to delete the file we just saved/loaded.
@@ -3887,19 +3869,21 @@ void FAssetDataGatherer::SaveMonolithicCacheFile(const TArray<TPair<FName,FDiskC
 	ParallelFor(CacheShards, [this, &AssetsToSave, &TotalCacheSize, &DataPerShard, ShardMask](int32 Shard) 
 	{
 		TArray<TPair<FName, FDiskCachedAssetData*>>& ShardData = DataPerShard[Shard];
-		FString Filename = FString::Printf(TEXT("%s_%d.bin"), *GPreloadSettings.GetMonolithicCacheBaseFilename(), Shard);
+		FString Filename = FString::Printf(TEXT("%s_%d.bin"), *GPreloadSettings.GetCacheBaseFilename(), Shard);
 		int64 CacheSize = SaveCacheFileInternal(Filename, ShardData);
 		TotalCacheSize += CacheSize;
 	}, EParallelForFlags::BackgroundPriority);
 
 	UE_LOG(LogAssetRegistry, Display, TEXT("Asset registry cache written as %.1f MiB to %s_*.bin"), 
-		static_cast<float>(TotalCacheSize)/1024.f/1024.f, *GPreloadSettings.GetMonolithicCacheBaseFilename());
-	// Delete old monolithic cache file if it exists
-	FString OldMonolithicName = GPreloadSettings.GetLegacyMonolithicCacheFilename();
-	IFileManager::Get().Delete(*OldMonolithicName, /* bRequireExists */ false, /* EvenReadOnly */ true, /* Quiet */ true);
-	
+		static_cast<float>(TotalCacheSize)/1024.f/1024.f, *GPreloadSettings.GetCacheBaseFilename());
+	// Delete old name of monolithic cache file and old non-monolithic cache directory if they exist
+	IFileManager::Get().Delete(*GPreloadSettings.GetLegacyCacheFilename(), false /* bRequireExists */,
+		true /* EvenReadOnly */, true /* Quiet */);
+	IFileManager::Get().DeleteDirectory(*GPreloadSettings.GetLegacyNonMonolithicCacheDirectory(),
+		false /* bRequireExists */, true /* Tree */);
+
 	// Delete any other shards if number of shards was reduced
-	TArray<FString> CacheFiles = GPreloadSettings.FindShardedMonolithicCacheFiles();
+	TArray<FString> CacheFiles = GPreloadSettings.FindShardedCacheFiles();
 	for (const FString& CacheFile : CacheFiles)
 	{
 		FStringView BaseName = FPathViews::GetBaseFilename(CacheFile);
@@ -3990,7 +3974,7 @@ FAssetDataGatherer::ETickResult FAssetDataGatherer::TickInternal(double& TickSta
 	TArray<FString, FBatchInlineAllocator> LocalCookedPackageNamesWithoutAssetDataResults;
 	TArray<FName, FBatchInlineAllocator> LocalVerseResults;
 	TArray<FString, FBatchInlineAllocator> LocalBlockedResults;
-	bool bLoadMonolithicCache = false;
+	bool bLoadCache = false;
 	bool bLocalIsCacheWriteEnabled = false;
 	double LocalLastCacheWriteTime = 0.0;
 	bool bWaitBatchCountDecremented = false;
@@ -4075,38 +4059,38 @@ FAssetDataGatherer::ETickResult FAssetDataGatherer::TickInternal(double& TickSta
 		FilesToSearch->PopFront(LocalFilesToSearch, NumToProcess);
 		check(LocalFilesToSearch.Num() > 0);
 
-		if (bReadMonolithicCache && !bHasLoadedMonolithicCache)
+		if (bCacheReadEnabled && !bHasLoadedCache)
 		{
-			bLoadMonolithicCache = true;
+			bLoadCache = true;
 		}
 		LocalLastCacheWriteTime = LastCacheWriteTime;
 	}
 	bLocalIsCacheWriteEnabled = bCacheWriteEnabled;
 
 	// Load the async cache if not yet loaded
-	if (bLoadMonolithicCache)
+	if (bLoadCache)
 	{
 		double CacheLoadStartTime = FPlatformTime::Seconds();
 		TArray<FCachePayload> Payloads;
 #if UE_EDITOR
-		if (GPreloadSettings.IsPreloadMonolithicCache())
+		if (GPreloadSettings.IsPreloadCache())
 		{
 			Payloads = GPreloader.Consume();
 		}
 		else
 #endif
 		{
-			TArray<FString> CachePaths = GPreloadSettings.FindShardedMonolithicCacheFiles();
-			Payloads = UE::AssetDataGather::Private::LoadCacheFiles(CachePaths, true /* bIsMonolithicCache */);
+			TArray<FString> CachePaths = GPreloadSettings.FindShardedCacheFiles();
+			Payloads = UE::AssetDataGather::Private::LoadCacheFiles(CachePaths);
 		}
 		ConsumeCacheFiles(MoveTemp(Payloads));
 		UE_LOG(LogAssetRegistry, Display, TEXT("AssetDataGatherer spent %.3fs loading caches %s_*.bin."),
-			FPlatformTime::Seconds() - CacheLoadStartTime, *GPreloadSettings.GetMonolithicCacheBaseFilename());
+			FPlatformTime::Seconds() - CacheLoadStartTime, *GPreloadSettings.GetCacheBaseFilename());
 
 		FGathererScopeLock ResultsScopeLock(&ResultsLock);
-		bHasLoadedMonolithicCache = true;
+		bHasLoadedCache = true;
 
-		// After we load the monolithic cache, restart the write timer for it. We don't need to save to if we just
+		// After we load the cache, restart the write timer for it. We don't need to save to if we just
 		// finished loading it (which we do before gathering anything) and we want to avoid failure to save due to
 		// writing a file that we just closed a readhandle for.
 		LastCacheWriteTime = FPlatformTime::Seconds();
@@ -4376,8 +4360,8 @@ FAssetDataGatherer::ETickResult FAssetDataGatherer::TickInternal(double& TickSta
 			}
 		}
 
-		const int32 NumAssetsReadSinceLastCacheWrite = NumUncachedAssetFiles - LastMonolithicCacheSaveUncachedAssetFiles;
-		if (bWriteMonolithicCache && !bIsSavingAsyncCache 
+		const int32 NumAssetsReadSinceLastCacheWrite = NumUncachedAssetFiles - LastCacheSaveNumUncachedAssetFiles;
+		if (bCacheWriteEnabled && !bIsSavingAsyncCache 
 			&& FPlatformTime::Seconds() - LocalLastCacheWriteTime >= AssetDataGathererConstants::MinSecondsToElapseBeforeCacheWrite
 			&& NumAssetsReadSinceLastCacheWrite >= AssetDataGathererConstants::MinAssetReadsBeforeCacheWrite)
 		{
@@ -4407,9 +4391,8 @@ bool FAssetDataGatherer::ReadAssetFile(const FString& AssetLongPackageName, cons
 	if (!PackageReader.OpenPackageFile(AssetLongPackageName, AssetFilename, &OpenPackageResult))
 	{
 		// If we're missing a custom version, we might be able to load this package later once the module containing that version is loaded...
-		//   -	We can only attempt a retry in editors (not commandlets) that haven't yet finished initializing (!GIsRunning), as we 
-		//		have no guarantee that a commandlet or an initialized editor is going to load any more modules/plugins
-		const bool bAllowRetry = GIsEditor && !IsEngineStartupModuleLoadingComplete();
+		//   -	Attempting a retry is only useful when engine startup module is not yet complete and therefore more plugins are expected.
+		const bool bAllowRetry = !IsEngineStartupModuleLoadingComplete();
 		if (OpenPackageResult == FPackageReader::EOpenPackageResult::CustomVersionMissing)
 		{
 			OutCanRetry = bAllowRetry;
@@ -4589,7 +4572,7 @@ void FAssetDataGatherer::WaitOnPath(FStringView InPath)
 	UE::AssetDataGather::Private::FPathExistence QueryPath(LocalAbsPath);
 	Discovery->SetPropertiesAndWait(TArrayView<UE::AssetDataGather::Private::FPathExistence>(&QueryPath, 1),
 		false /* bAddToAllowList */, false /* bForceRescan */, false /* bIgnoreDenyListScanFilters */);
-	WaitOnPathsInternal(TArrayView<UE::AssetDataGather::Private::FPathExistence>(&QueryPath, 1), FString(), TArray<FString>());
+	WaitOnPathsInternal(TArrayView<UE::AssetDataGather::Private::FPathExistence>(&QueryPath, 1));
 }
 
 void FAssetDataGatherer::ClearCache()
@@ -4601,11 +4584,6 @@ void FAssetDataGatherer::ClearCache()
 		bWasCacheEnabled = bCacheWriteEnabled || bCacheReadEnabled;
 		bCacheWriteEnabled = false;
 		bCacheReadEnabled = false;
-		{
-			FGathererScopeLock ResultsScopeLock(&ResultsLock);
-			bReadMonolithicCache = false;
-			bWriteMonolithicCache = false;
-		}
 		bCacheIsInUseOnOtherThread = CacheInUseCount > 0;
 	}
 
@@ -4645,7 +4623,7 @@ void FAssetDataGatherer::ClearCache()
 
 
 void FAssetDataGatherer::ScanPathsSynchronous(const TArray<FString>& InLocalPaths, bool bForceRescan,
-	bool bIgnoreDenyListScanFilters, const FString& SaveCacheFilename, const TArray<FString>& SaveCacheLongPackageNameDirs)
+	bool bIgnoreDenyListScanFilters)
 {
 	TArray<UE::AssetDataGather::Private::FPathExistence> QueryPaths;
 	QueryPaths.Reserve(InLocalPaths.Num());
@@ -4661,11 +4639,10 @@ void FAssetDataGatherer::ScanPathsSynchronous(const TArray<FString>& InLocalPath
 		SetIsIdle(false);
 	}
 
-	WaitOnPathsInternal(QueryPaths, SaveCacheFilename, SaveCacheLongPackageNameDirs);
+	WaitOnPathsInternal(QueryPaths);
 }
 
-void FAssetDataGatherer::WaitOnPathsInternal(TArrayView<UE::AssetDataGather::Private::FPathExistence> QueryPaths,
-	const FString& SaveCacheFilename, const TArray<FString>& SaveCacheLongPackageNameDirs)
+void FAssetDataGatherer::WaitOnPathsInternal(TArrayView<UE::AssetDataGather::Private::FPathExistence> QueryPaths)
 {
 	LLM_SCOPE(ELLMTag::AssetRegistry);
 
@@ -4688,36 +4665,17 @@ void FAssetDataGatherer::WaitOnPathsInternal(TArrayView<UE::AssetDataGather::Pri
 		}
 	}
 
-	// We do not contribute to the async cache save if we have been given a modular cache to save 
-	bool bContributeToCacheSave = SaveCacheFilename.IsEmpty();
-
 	// Tick until NumDiscoveredPaths have been read
 	TRACE_CPUPROFILER_EVENT_SCOPE(FAssetDataGatherer::Tick);
 	for (;;)
 	{
-		InnerTickLoop(true /* bInSynchronousTick */, bContributeToCacheSave, -1. /* EndTimeSeconds */);
+		InnerTickLoop(true /* bInSynchronousTick */, true /* bContributeToCacheSave */, -1. /* EndTimeSeconds */);
 		FGathererScopeLock ResultsScopeLock(&ResultsLock); // WaitBatchCount requires the lock
 		if (WaitBatchCount < 0)
 		{
 			break;
 		}
 	}
-
-	if (!SaveCacheFilename.IsEmpty() && bCacheWriteEnabled)
-	{
-		TArray<TPair<FName, FDiskCachedAssetData*>> AssetsToSave;
-		{
-			FGathererScopeLock TickScopeLock(&TickLock);
-			GetAssetsToSave(SaveCacheLongPackageNameDirs, AssetsToSave);
-			CacheInUseCount++;
-		}
-		SaveCacheFileInternal(SaveCacheFilename, AssetsToSave);
-		{
-			FGathererScopeLock TickScopeLock(&TickLock);
-			check(CacheInUseCount > 0);
-			CacheInUseCount--;
-		}
-	}
 }
 
 void FAssetDataGatherer::WaitForIdle(float TimeoutSeconds)
@@ -4795,47 +4753,6 @@ bool FAssetDataGatherer::IsCacheWriteEnabled() const
 	return bCacheWriteEnabled;
 }
 
-FString FAssetDataGatherer::GetCacheFilename(TConstArrayView<FString> CacheFilePackagePaths)
-{
-	// Try and build a consistent hash for this input
-	// Normalize the paths; removing any trailing /
-	TArray<FString> SortedPaths(CacheFilePackagePaths);
-	for (FString& PackagePath : SortedPaths)
-	{
-		while (PackagePath.Len() > 1 && PackagePath.EndsWith(TEXT("/")))
-		{
-			PackagePath.LeftChopInline(1);
-		}
-	}
-
-	// Sort the paths
-	SortedPaths.StableSort();
-
-	// todo: handle hash collisions?
-	uint32 CacheHash = SortedPaths.Num() > 0 ? GetTypeHash(SortedPaths[0]) : 0;
-	for (int32 PathIndex = 1; PathIndex < SortedPaths.Num(); ++PathIndex)
-	{
-		CacheHash = HashCombine(CacheHash, GetTypeHash(SortedPaths[PathIndex]));
-	}
-
-	return UE::AssetDataGather::Private::GPreloadSettings.GetAssetRegistryCacheRootFolder() / TEXT("AssetRegistryCache") / FString::Printf(TEXT("%08x%s.bin"), CacheHash, bGatherDependsData ? TEXT("") : TEXT("NoDeps"));
-}
-
-void FAssetDataGatherer::LoadCacheFiles(TConstArrayView<FString> CacheFilenames)
-{
-	using namespace UE::AssetDataGather::Private;
-	if (!bCacheReadEnabled)
-	{
-		return;
-	}
-
-	TArray<FCachePayload> Payloads = UE::AssetDataGather::Private::LoadCacheFiles(CacheFilenames, false /* bIsMonolithicCache */);
-	FScopedGatheringPause ScopedPause(*this);
-	CHECK_IS_NOT_LOCKED_CURRENT_THREAD(ResultsLock);
-	FGathererScopeLock TickScopeLock(&TickLock);
-	ConsumeCacheFiles(MoveTemp(Payloads));
-}
-
 void FAssetDataGatherer::ConsumeCacheFiles(TArray<UE::AssetDataGather::Private::FCachePayload> Payloads)
 {
 	CHECK_IS_LOCKED_CURRENT_THREAD(TickLock);
@@ -4869,7 +4786,7 @@ void FAssetDataGatherer::ConsumeCacheFiles(TArray<UE::AssetDataGather::Private::
 	AssetResults.Reserve(DiskCachedAssetDataMap.Num());
 }
 
-void FAssetDataGatherer::TryReserveSaveMonolithicCache(bool& bOutShouldSave, TArray<TPair<FName,FDiskCachedAssetData*>>& AssetsToSave)
+void FAssetDataGatherer::TryReserveSaveCache(bool& bOutShouldSave, TArray<TPair<FName,FDiskCachedAssetData*>>& AssetsToSave)
 {
 	bOutShouldSave = false;
 	if (IsStopped)
@@ -4884,14 +4801,14 @@ void FAssetDataGatherer::TryReserveSaveMonolithicCache(bool& bOutShouldSave, TAr
 	int32 LocalNumUncachedAssetFiles;
 	{
 		FGathererScopeLock ResultsScopeLock(&ResultsLock);
-		bOutShouldSave = bWriteMonolithicCache;
+		bOutShouldSave = bCacheWriteEnabled;
 		LocalNumUncachedAssetFiles = NumUncachedAssetFiles;
 	}
 	if (bOutShouldSave)
 	{
-		GetMonolithicCacheAssetsToSave(AssetsToSave);
+		GetCacheAssetsToSave(AssetsToSave);
 		bIsSavingAsyncCache = true;
-		LastMonolithicCacheSaveUncachedAssetFiles = LocalNumUncachedAssetFiles ;
+		LastCacheSaveNumUncachedAssetFiles = LocalNumUncachedAssetFiles;
 	}
 	bSaveAsyncCacheTriggered = false;
 }
@@ -4926,7 +4843,7 @@ void FAssetDataGatherer::GetAssetsToSave(TArrayView<const FString> SaveCacheLong
 	}
 }
 
-void FAssetDataGatherer::GetMonolithicCacheAssetsToSave(TArray<TPair<FName,FDiskCachedAssetData*>>& OutAssetsToSave)
+void FAssetDataGatherer::GetCacheAssetsToSave(TArray<TPair<FName,FDiskCachedAssetData*>>& OutAssetsToSave)
 {
 	CHECK_IS_LOCKED_CURRENT_THREAD(TickLock);
 
@@ -4939,7 +4856,7 @@ void FAssetDataGatherer::GetMonolithicCacheAssetsToSave(TArray<TPair<FName,FDisk
 
 	for (const TPair<FName, FDiskCachedAssetData*>& Pair : DiskCachedAssetDataMap)
 	{
-		if(NewCachedAssetDataMap.Contains(Pair.Key))
+		if (NewCachedAssetDataMap.Contains(Pair.Key))
 		{
 			continue; // Data was replaced when populating NewCachedAssetDataMap
 		}	
@@ -5085,12 +5002,12 @@ FCachePayload SerializeCacheLoad(FAssetRegistryReader& Ar)
 	return Result;
 }
 
-TArray<FCachePayload> LoadCacheFiles(TConstArrayView<FString> InCacheFilenames, bool bIsMonolithicCache)
+TArray<FCachePayload> LoadCacheFiles(TConstArrayView<FString> InCacheFilenames)
 {
 	TRACE_CPUPROFILER_EVENT_SCOPE(LoadCacheFiles);
 	TArray<FCachePayload> Results;
 	Results.AddDefaulted(InCacheFilenames.Num());
-	ParallelFor(InCacheFilenames.Num(), [&InCacheFilenames, &Results, bIsMonolithicCache](int32 Index)
+	ParallelFor(InCacheFilenames.Num(), [&InCacheFilenames, &Results](int32 Index)
 	{
 		const FString& CacheFilename = InCacheFilenames[Index];
 		auto DoLoad = [&](FArchive& ChecksummingReader)
@@ -5122,7 +5039,7 @@ TArray<FCachePayload> LoadCacheFiles(TConstArrayView<FString> InCacheFilenames,
 			{
 				FChecksumViewReader ChecksummingReader(MoveTemp(FileReader), CacheFilename);
 				Payload = DoLoad(ChecksummingReader);
-				UE_CLOG(bIsMonolithicCache && Payload.bSucceeded, LogAssetRegistry, Display, TEXT("Asset registry cache read as %.1f MiB from %s"), static_cast<float>(FileReader.GetTotalSize())/1024.f/1024.f, *CacheFilename);
+				UE_CLOG(Payload.bSucceeded, LogAssetRegistry, Display, TEXT("Asset registry cache read as %.1f MiB from %s"), static_cast<float>(FileReader.GetTotalSize())/1024.f/1024.f, *CacheFilename);
 				UE_CLOG(!Payload.bSucceeded, LogAssetRegistry, Warning, TEXT("There was an error loading the asset registry cache using memory mapping"));
 			}
 
@@ -5140,7 +5057,7 @@ TArray<FCachePayload> LoadCacheFiles(TConstArrayView<FString> InCacheFilenames,
 				{
 					FChecksumArchiveReader ChecksummingReader(*FileAr);
 					Payload = DoLoad(ChecksummingReader);
-					UE_CLOG(bIsMonolithicCache && Payload.bSucceeded, LogAssetRegistry, Display, TEXT("Asset registry cache read as %.1f MiB from %s"), static_cast<float>(FileAr->TotalSize())/1024.f/1024.f, *CacheFilename);
+					UE_CLOG(Payload.bSucceeded, LogAssetRegistry, Display, TEXT("Asset registry cache read as %.1f MiB from %s"), static_cast<float>(FileAr->TotalSize())/1024.f/1024.f, *CacheFilename);
 					UE_CLOG(!Payload.bSucceeded, LogAssetRegistry, Warning, TEXT("There was an error loading the asset registry cache"));
 				}
 			}
@@ -5198,7 +5115,7 @@ SIZE_T FAssetDataGatherer::GetAllocatedSize() const
 	Result += BlockedResults.GetAllocatedSize();
 	Result += SearchTimes.GetAllocatedSize();
 	Result += GetArrayRecursiveAllocatedSize(DiscoveredPaths);
-	Result += GPreloadSettings.GetMonolithicCacheBaseFilename().GetAllocatedSize();
+	Result += GPreloadSettings.GetCacheBaseFilename().GetAllocatedSize();
 
 	Result += NewCachedAssetData.GetAllocatedSize();
 	for (const FDiskCachedAssetData* Value : NewCachedAssetData)
diff --git a/Engine/Source/Runtime/AssetRegistry/Private/AssetDataGatherer.h b/Engine/Source/Runtime/AssetRegistry/Private/AssetDataGatherer.h
--- a/Engine/Source/Runtime/AssetRegistry/Private/AssetDataGatherer.h
+++ b/Engine/Source/Runtime/AssetRegistry/Private/AssetDataGatherer.h
@@ -84,12 +84,6 @@ public:
 
 	void OnInitialSearchCompleted();
 
-	// Extra at-construction configuration 
-
-	/** Configure the gatherer to use a single monolithic cache, and read/write this cache during ticks. */
-	void ActivateMonolithicCache();
-
-
 	// Controlling Async behavior
 
 	/** Start the async thread, if this Gatherer was created async. Does nothing if not async or already started. */
@@ -167,11 +161,8 @@ public:
 	 * Add a set of paths to the allow list, optionally force rescanning and ignore deny list on them,
 	 * and wait for all assets in the paths to be added to search results.
 	 * Wait time is minimized by prioritizing the paths and transferring async scanning to the current thread.
-	 * If SaveCacheFilename is non-empty, save a cachefile to it with all discovered paths that are in one of
-	 * the SaveCacheLongPackageNameDirs paths.
 	 */
-	void ScanPathsSynchronous(const TArray<FString>& InPaths, bool bForceRescan, bool bIgnoreDenyListScanFilters,
-		const FString& SaveCacheFilename, const TArray<FString>& SaveCacheLongPackageNameDirs);
+	void ScanPathsSynchronous(const TArray<FString>& InPaths, bool bForceRescan, bool bIgnoreDenyListScanFilters);
 	/** Wait for all monitored assets to be added to search results. */
 	void WaitForIdle(float TimeoutSeconds = -1.0f);
 	/**
@@ -192,10 +183,6 @@ public:
 	bool IsCacheReadEnabled() const;
 	/** Return whether the current process enables writing AssetDataGatherer cache files. */
 	bool IsCacheWriteEnabled() const;
-	/** Calculate the cache filename that should be used for the given list of package paths. */
-	FString GetCacheFilename(TConstArrayView<FString> CacheFilePackagePaths);
-	/** Attempt to read the cache file at the given LocalPath, and store all of its results in the in-memory cache. */
-	void LoadCacheFiles(TConstArrayView<FString> CacheFilename);
 	/** Return the memory used by the gatherer. Used for performance metrics. */
 	SIZE_T GetAllocatedSize() const;
 
@@ -295,8 +282,7 @@ private:
 	 * Wait for all monitored assets under the given path to be added to search results.
 	 * Returns immediately if the given path are not monitored.
 	 */
-	void WaitOnPathsInternal(TArrayView<UE::AssetDataGather::Private::FPathExistence> QueryPaths,
-		const FString& SaveCacheFilename, const TArray<FString>& SaveCacheFilterDirs);
+	void WaitOnPathsInternal(TArrayView<UE::AssetDataGather::Private::FPathExistence> QueryPaths);
 
 	/** Sort the pending list of filepaths so that assets under the given directory/filename are processed first. */
 	void SortPathsByPriority(TArrayView<UE::AssetDataGather::Private::FPathExistence> QueryPaths,
@@ -320,14 +306,12 @@ private:
 	/** Add the given AssetDatas into DiskCachedAssetDataMap and DiskCachedAssetBlocks. */
 	void ConsumeCacheFiles(TArray<UE::AssetDataGather::Private::FCachePayload> Payloads);
 	/**
-	 * If a save of the monolithic cache has been triggered, get the cache filename and pointers to all elements that
+	 * If a cache save has been triggered, get the cache filename and pointers to all elements that
 	 * should be saved, for later saving outside of the critical section.
 	 */
-	void TryReserveSaveMonolithicCache(bool& bOutShouldSave, TArray<TPair<FName,FDiskCachedAssetData*>>& AssetsToSave);
-	/** 
-	 * Save a monolithic cache for the main asset discovery process, possibly sharded into multiple files.
-	*/
-	void SaveMonolithicCacheFile(const TArray<TPair<FName,FDiskCachedAssetData*>>& AssetsToSave);
+	void TryReserveSaveCache(bool& bOutShouldSave, TArray<TPair<FName,FDiskCachedAssetData*>>& AssetsToSave);
+	/** Save cache file for the assetdatas read from package headers, possibly sharded into multiple files. */
+	void SaveCacheFile(const TArray<TPair<FName,FDiskCachedAssetData*>>& AssetsToSave);
 	/**
 	 * If the CacheFilename/AssetsToSave are non empty, save the cache file. 
 	 * This function reads the read-only-after-creation data from each FDiskCachedAssetData*, but otherwise does not use
@@ -343,10 +327,10 @@ private:
 	void GetAssetsToSave(TArrayView<const FString> SaveCacheLongPackageNameDirs,
 		TArray<TPair<FName,FDiskCachedAssetData*>>& OutAssetsToSave);
 	/**
-	 * Get the list of FDiskCachedAssetData* for saving into the monolithic cache.
-	 * Includes both assets that were loaded in the gatherer and assets which were loaded from the monolithic cache and have not been pruned.
+	 * Get the list of FDiskCachedAssetData* for saving into the cache.
+	 * Includes both assets that were loaded in the gatherer and assets which were loaded from the cache and have not been pruned.
 	 */
-	void GetMonolithicCacheAssetsToSave(TArray<TPair<FName,FDiskCachedAssetData*>>& OutAssetsToSave);
+	void GetCacheAssetsToSave(TArray<TPair<FName,FDiskCachedAssetData*>>& OutAssetsToSave);
 
 	/* Adds the given pair into NewCachedAssetDataMap. Detects collisions for multiple files with the same PackageName */
 	void AddToCache(FName PackageName, FDiskCachedAssetData* DiskCachedAssetData);
@@ -474,16 +458,8 @@ private:
 	int32 NumCachedAssetFiles = 0;
 	/** The total number of files in the search results that were not in the cache and were read by parsing the file. */
 	int32 NumUncachedAssetFiles = 0;
-	/**
-	 * Track whether we are allowed to read from a monolithic cache that should be loaded during tick.
-	 * Even if we are or not, if bCacheReadEnabled the AssetRegistry can also call LoadCacheFile/ScanPathsSynchronous to
-	 * load/save smaller files.
-	 */
-	bool bReadMonolithicCache;
-	/** Track whether we are allowed to write to the monolithic cache. */
-	bool bWriteMonolithicCache;
-	/** If bHasLoadedMonolithicCache is true, track whether the cache has been loaded. */
-	bool bHasLoadedMonolithicCache;
+	/** Track whether the cache has been loaded. */
+	bool bHasLoadedCache;
 	/** Track whether the Discovery subsystem has gone idle and we have read all filenames from it. */
 	bool bDiscoveryIsComplete;
 	/** Track whether this Gather has gone idle and a caller has read all search data from it. */
@@ -513,7 +489,7 @@ private:
 	/** Used to block on gather results. If non-negative, tick should end when WaitBatchCount files have been processed. */
 	int32 WaitBatchCount;
 	/** How many uncached asset files had been discovered at the last async cache save */
-	int32 LastMonolithicCacheSaveUncachedAssetFiles;
+	int32 LastCacheSaveNumUncachedAssetFiles;
 	/**
 	 * Incremented when a thread is in the middle of saving any cache and therefore the cache cannot be deleted,
 	 * decremented when the thread is done. Only incremented when bCacheEnabled has been recently confirmed to be true.
diff --git a/Engine/Source/Runtime/AssetRegistry/Private/AssetDataGathererPrivate.h b/Engine/Source/Runtime/AssetRegistry/Private/AssetDataGathererPrivate.h
--- a/Engine/Source/Runtime/AssetRegistry/Private/AssetDataGathererPrivate.h
+++ b/Engine/Source/Runtime/AssetRegistry/Private/AssetDataGathererPrivate.h
@@ -858,17 +858,17 @@ public:
 	bool IsDiscoveryCacheReadEnabled() const;
 	EFeatureEnabled IsDiscoveryCacheWriteEnabled() const;
 	bool IsDiscoveryCacheInvalidateEnabled() const;
-	bool IsMonolithicCacheActivatedDuringPreload() const;
-	bool IsPreloadMonolithicCache() const;
+	bool IsPreloadCache() const;
 	bool IsGatherDependsData() const;
 	bool IsForceDependsGathering() const;
-	FString GetLegacyMonolithicCacheFilename() const;
-	const FString& GetMonolithicCacheBaseFilename() const;
+	FString GetLegacyCacheFilename() const;
+	FString GetLegacyNonMonolithicCacheDirectory() const;
+	const FString& GetCacheBaseFilename() const;
 	const FString& GetAssetRegistryCacheRootFolder() const;
-	TArray<FString> FindShardedMonolithicCacheFiles() const;
+	TArray<FString> FindShardedCacheFiles() const;
 
 private:
-	FString MonolithicCacheBaseFilename;
+	FString CacheBaseFilename;
 	FString AssetRegistryCacheRootFolder;
 	bool bForceDependsGathering = false;
 	bool bGatherDependsData = false;
@@ -877,7 +877,7 @@ private:
 	bool bDiscoveryCacheReadEnabled = false;
 	EFeatureEnabled DiscoveryCacheWriteEnabled = EFeatureEnabled::Never;
 	bool bDiscoveryCacheInvalidateEnabled = false;
-	bool bMonolithicCacheActivatedDuringPreload = false;
+	bool bPreloadCache = false;
 	bool bInitialized = false;
 };
 extern FPreloadSettings GPreloadSettings;
diff --git a/Engine/Source/Runtime/AssetRegistry/Private/AssetRegistry.cpp b/Engine/Source/Runtime/AssetRegistry/Private/AssetRegistry.cpp
--- a/Engine/Source/Runtime/AssetRegistry/Private/AssetRegistry.cpp
+++ b/Engine/Source/Runtime/AssetRegistry/Private/AssetRegistry.cpp
@@ -989,7 +989,7 @@ void FAssetRegistryImpl::Initialize(Impl::FInitializeContext& Context)
 #if WITH_EDITOR
 	if (GIsEditor)
 	{
-		// Double check mount point is still valid because it could have been umounted
+		// Double check mount point is still valid because it could have been unmounted
 		bVerifyMountPointAfterGather = true;
 	}
 #endif // WITH_EDITOR
@@ -1016,8 +1016,10 @@ void FAssetRegistryImpl::Initialize(Impl::FInitializeContext& Context)
 		}
 		else
 		{
-			// For the Editor we need to take responsibility for the synchronous search; Commandlets will handle that themselves
-			Context.bNeedsSearchAllAssetsAtStartSynchronous = GIsEditor && !IsRunningCommandlet();
+			// For the Editor and editor game we need to take responsibility for the synchronous search;
+			// Commandlets and cooked game will handle it themselves.
+			constexpr bool bEditorExecutable = WITH_EDITOR;
+			Context.bNeedsSearchAllAssetsAtStartSynchronous = bEditorExecutable && !IsRunningCommandlet();
 		}
 	}
 
@@ -1241,7 +1243,7 @@ void FAssetRegistryImpl::InitRedirectors(Impl::FEventContext& EventContext,
 	bOutRedirectorsNeedSubscribe = false;
 
 	// plugins can't initialize redirectors in the editor, it will mess up the saving of content.
-	if ( GIsEditor )
+	if (GIsEditor)
 	{
 		return;
 	}
@@ -1812,9 +1814,6 @@ bool FAssetRegistryImpl::TryConstructGathererIfNeeded()
 void FAssetRegistryImpl::SearchAllAssetsInitialAsync(Impl::FEventContext& EventContext,
 	Impl::FClassInheritanceContext& InheritanceContext)
 {
-	InitialSearchStartTime = FPlatformTime::Seconds();
-	bInitialSearchStarted = true;
-	bInitialSearchCompleted = false;
 	SetPerformanceMode(Impl::EPerformanceMode::BulkLoading);
 	SearchAllAssets(EventContext, InheritanceContext, false /* bSynchronousSearch */);
 }
@@ -1877,14 +1876,13 @@ void UAssetRegistryImpl::SearchAllAssets(bool bSynchronousSearch)
 		}
 		GuardedData.SearchAllAssets(EventContext, InheritanceContext, bSynchronousSearch);
 	}
-#if WITH_EDITOR
+	Broadcast(EventContext);
+
 	if (bSynchronousSearch)
 	{
-		UE::AssetRegistry::Impl::FInterruptionContext InterruptionContext(-1., -1.);
-		ProcessLoadedAssetsToUpdateCache(EventContext, EGatherStatus::Complete, InterruptionContext);
+		// Continue calling TickGatherer until completion is signaled, and call ProcessLoadedAssetsToUpdateCache
+		WaitForCompletion();
 	}
-#endif
-	Broadcast(EventContext);
 }
 
 bool UAssetRegistryImpl::IsSearchAllAssets() const
@@ -1912,6 +1910,12 @@ void FAssetRegistryImpl::SearchAllAssets(Impl::FEventContext& EventContext,
 	{
 		return;
 	}
+	if (!bInitialSearchStarted)
+	{
+		InitialSearchStartTime = FPlatformTime::Seconds();
+		bInitialSearchStarted = true;
+		bInitialSearchCompleted = false;
+	}
 
 	FAssetDataGatherer& Gatherer = *GlobalGatherer;
 	if (!Gatherer.IsAsyncEnabled())
@@ -1920,8 +1924,6 @@ void FAssetRegistryImpl::SearchAllAssets(Impl::FEventContext& EventContext,
 		bSynchronousSearch = true;
 	}
 
-	Gatherer.ActivateMonolithicCache();
-
 	// Add all existing mountpoints to the GlobalGatherer
 	// This will include Engine content, Game content, but also may include mounted content directories for one or more plugins.
 	TArray<FString> PackagePathsToSearch;
@@ -1939,14 +1941,6 @@ void FAssetRegistryImpl::SearchAllAssets(Impl::FEventContext& EventContext,
 		Gatherer.WaitForIdle();
 		Impl::FInterruptionContext InterruptionContext(-1., -1.);
 		Impl::EGatherStatus UnusedStatus = TickGatherer(EventContext, InheritanceContext, InterruptionContext);
-#if WITH_EDITOR
-		if (!bInitialSearchStarted)
-		{
-			// We have a contract that we call UpdateRedirectCollector after the call to SearchAllAssets completes.
-			// If we ran the initial async call asynchronously it is done in TickGatherer; for later synchronous calls it is done here
-			UpdateRedirectCollector();
-		}
-#endif
 	}
 	else
 	{
@@ -1964,6 +1958,7 @@ void UAssetRegistryImpl::WaitForCompletion()
 
 	bool bInitialSearchStarted = false;
 	bool bInitialSearchCompleted = false;
+	bool bAsyncGathering = false;
 
 	// Try taking over the gather thread for a short time in case it is mostly done.
 	// But if it has more than a small amount of work to do, let the gather thread do that work
@@ -1980,6 +1975,7 @@ void UAssetRegistryImpl::WaitForCompletion()
 		GuardedData.WaitForGathererIdle(TimeToJoinSeconds);
 		bInitialSearchStarted = GuardedData.IsInitialSearchStarted();
 		bInitialSearchCompleted = GuardedData.IsInitialSearchCompleted();
+		bAsyncGathering = GuardedData.GlobalGatherer && GuardedData.GlobalGatherer->IsAsyncEnabled();
 	}
 
 #if WITH_EDITOR
@@ -2069,9 +2065,12 @@ void UAssetRegistryImpl::WaitForCompletion()
 		}
 
 		FThreadHeartBeat::Get().HeartBeat();
-		// Sleep long enough to avoid causing contention on the CriticalSection in GetAndTrimSearchResults
-		constexpr float SleepTimeSeconds = 0.010f;
-		FPlatformProcess::SleepNoStats(SleepTimeSeconds);
+		if (Status == EGatherStatus::TickActiveGatherActive && bAsyncGathering)
+		{
+			// Sleep long enough to avoid causing contention on the CriticalSection in GetAndTrimSearchResults
+			constexpr float SleepTimeSeconds = 0.010f;
+			FPlatformProcess::SleepNoStats(SleepTimeSeconds);
+		}
 	}
 }
 
@@ -5701,18 +5700,6 @@ void FAssetRegistryImpl::ScanPathsSynchronous(Impl::FScanPathContext& Context)
 	}
 	FAssetDataGatherer& Gatherer = *GlobalGatherer;
 
-	// Add a cache file for any not-yet-scanned dirs
-	TArray<FString> CacheFilePackagePaths;
-	if (!Context.bForceRescan && (Gatherer.IsCacheReadEnabled() || Gatherer.IsCacheWriteEnabled()))
-	{
-		for (int n = 0; n < Context.LocalDirs.Num(); ++n)
-		{
-			if (!Gatherer.IsOnAllowList(Context.LocalDirs[n]))
-			{
-				CacheFilePackagePaths.Add(Context.PackageDirs[n]);
-			}
-		}
-	}
 	Context.LocalPaths.Reserve(Context.LocalFiles.Num() + Context.LocalDirs.Num());
 	Context.LocalPaths.Append(MoveTemp(Context.LocalDirs));
 	Context.LocalPaths.Append(MoveTemp(Context.LocalFiles));
@@ -5722,13 +5709,6 @@ void FAssetRegistryImpl::ScanPathsSynchronous(Impl::FScanPathContext& Context)
 	}
 	Gatherer.AddRequiredMountPoints(Context.LocalPaths);
 
-	FString CacheFilename;
-	if (!CacheFilePackagePaths.IsEmpty())
-	{
-		CacheFilename = Gatherer.GetCacheFilename(CacheFilePackagePaths);
-		Gatherer.LoadCacheFiles({CacheFilename});
-	}
-
 	// If we are forcing a rescan, then delete any old assets that no longer exist. If we are not forcing a rescan,
 	// then there should not be any old assets that no longer exist, so we skip the cost of searching for them.
 	TSet<FSoftObjectPath> OldAssetsToRemove;
@@ -5780,7 +5760,7 @@ void FAssetRegistryImpl::ScanPathsSynchronous(Impl::FScanPathContext& Context)
 		}
 	}
 
-	Gatherer.ScanPathsSynchronous(Context.LocalPaths, Context.bForceRescan, Context.bIgnoreDenyListScanFilters, CacheFilename, Context.PackageDirs);
+	Gatherer.ScanPathsSynchronous(Context.LocalPaths, Context.bForceRescan, Context.bIgnoreDenyListScanFilters);
 	TArray<FName> FoundAssetPackageNames;
 
 	auto IsInRequestedDir = [&Context](const FAssetData& AssetData)
@@ -9153,8 +9133,17 @@ void GetAssetForPackages(TConstArrayView<FName> PackageNames, TMap<FName, FAsset
 
 bool ShouldSearchAllAssetsAtStart()
 {
-	// If in the editor or cookcommandlet, or an allowlist commandlet, we start the GlobalGatherer now
-	// In the game or other commandlets, we do not construct it until project or commandlet code calls SearchAllAssets or ScanPathsSynchronous
+	// Search at start for configurations that need the entire assetregistry and that do not load it from serialized:
+	// Need it: Editor IDE, CookCommandlet, other Allowlist Commandlets
+	// Possibly need it: editor running as -game or -server
+	// Do not need it: Commandlets not on the Allowlist
+	// Load it from serialized: Non-editor-executable
+	// 
+	// This behavior can be overridden with commandline option.
+	// 
+	// For the editor-executable configurations that do not search at start, the search will be triggered when
+	// SearchAllAssets or ScanPathsSynchronous is called.
+
 	bool bSearchAllAssetsAtStart = false;
 	if (GIsEditor)
 	{
@@ -9190,6 +9179,14 @@ bool ShouldSearchAllAssetsAtStart()
 			}
 		}
 	}
+#if WITH_EDITOR
+	else
+	{
+		bool bEditorGameScansAR = true;
+		GConfig->GetBool(TEXT("AssetRegistry"), TEXT("EditorGameScansAR"), bEditorGameScansAR, GEngineIni);
+		bSearchAllAssetsAtStart = bEditorGameScansAR;
+	}
+#endif
 #if WITH_EDITOR || !UE_BUILD_SHIPPING
 	bool bCommandlineAllAssetsAtStart;
 	if (FParse::Bool(FCommandLine::Get(), TEXT("AssetGatherAll="), bCommandlineAllAssetsAtStart))
