AssetRegistry: Fix static analysis warning about expression always 0 in non-editor buidls.
#rnx

diff --git a/Engine/Source/Runtime/AssetRegistry/Private/AssetDataGatherer.cpp b/Engine/Source/Runtime/AssetRegistry/Private/AssetDataGatherer.cpp
--- a/Engine/Source/Runtime/AssetRegistry/Private/AssetDataGatherer.cpp
+++ b/Engine/Source/Runtime/AssetRegistry/Private/AssetDataGatherer.cpp
@@ -183,8 +183,11 @@ void FPreloadSettings::Initialize()
 		AssetRegistryCacheRootFolder = FPaths::ProjectIntermediateDir();
 	}
 	bForceDependsGathering = FParse::Param(FCommandLine::Get(), TEXT("ForceDependsGathering"));
-	constexpr bool bEditorExecutable = WITH_EDITOR;
-	bGatherDependsData = (bEditorExecutable && !FParse::Param(FCommandLine::Get(), TEXT("NoDependsGathering"))) || bForceDependsGathering;
+#if WITH_EDITOR
+	bGatherDependsData = bForceDependsGathering || !FParse::Param(FCommandLine::Get(), TEXT("NoDependsGathering"));
+#else
+	bGatherDependsData = bForceDependsGathering;
+#endif
 	bool bNoAssetRegistryCache = FParse::Param(FCommandLine::Get(), TEXT("NoAssetRegistryCache"));
 	bool bNoAssetRegistryDiscoveryCache = bNoAssetRegistryCache || FParse::Param(FCommandLine::Get(), TEXT("NoAssetRegistryDiscoveryCache"));
 	bool bNoAssetRegistryCacheRead = FParse::Param(FCommandLine::Get(), TEXT("NoAssetRegistryCacheRead"));
