AssetRegistry: Fix unsynchronized multithreaded access to FCoreRedirects from gatherer thread due to TickGatherer being called too early.
Add speculative fix for AssetManager failing to receive FileLoadedEvent; do not allow setting bInitialSearchCompleted without broadcasting the FileLoadedEvent.
Maintainability improvement: create FTickContext struct for TickGatherer's arguments.
Remove the separate flag for bScanEndedEventBroadcast; it is synonymous with bFileLoadedEventBroadcast.
#rnx

diff --git a/Engine/Source/Runtime/AssetRegistry/Private/AssetRegistry.cpp b/Engine/Source/Runtime/AssetRegistry/Private/AssetRegistry.cpp
--- a/Engine/Source/Runtime/AssetRegistry/Private/AssetRegistry.cpp
+++ b/Engine/Source/Runtime/AssetRegistry/Private/AssetRegistry.cpp
@@ -1190,7 +1190,7 @@ void UAssetRegistryImpl::InitializeEvents(UE::AssetRegistry::Impl::FInitializeCo
 	FPackageName::OnContentPathDismounted().AddUObject(this, &UAssetRegistryImpl::OnContentPathDismounted);
 
 	// If we were called before engine has fully initialized, refresh classes on initialize. If not this won't do anything as it already happened
-	FCoreDelegates::OnPostEngineInit.AddUObject(this, &UAssetRegistryImpl::OnRefreshNativeClasses);
+	FCoreDelegates::OnPostEngineInit.AddUObject(this, &UAssetRegistryImpl::OnPostEngineInit);
 
 	IPluginManager& PluginManager = IPluginManager::Get();
 	if (!IsEngineStartupModuleLoadingComplete())
@@ -1243,7 +1243,7 @@ void FAssetRegistryImpl::InitRedirectors(Impl::FEventContext& EventContext,
 	bOutRedirectorsNeedSubscribe = false;
 
 	// plugins can't initialize redirectors in the editor, it will mess up the saving of content.
-	if (GIsEditor)
+	if ( GIsEditor )
 	{
 		return;
 	}
@@ -1300,7 +1300,7 @@ void UAssetRegistryImpl::OnInitialPluginLoadingComplete()
 	{
 		LLM_SCOPE(ELLMTag::AssetRegistry);
 		UE::AssetRegistry::FInterfaceWriteScopeLock InterfaceScopeLock(InterfaceLock);
-		GuardedData.OnPostEngineInit(true);
+		GuardedData.OnPluginLoadingComplete(true);
 	}
 
 	FCoreDelegates::OnAllModuleLoadingPhasesComplete.RemoveAll(this);
@@ -1309,7 +1309,7 @@ void UAssetRegistryImpl::OnInitialPluginLoadingComplete()
 namespace UE::AssetRegistry
 {
 
-void FAssetRegistryImpl::OnPostEngineInit(bool bPhaseSuccessful)
+void FAssetRegistryImpl::OnPluginLoadingComplete(bool bPhaseSuccessful)
 {
 	// If we have constructed the GlobalGatherer then we need to readscriptpackages,
 	// otherwise we will read them when constructing the gatherer.
@@ -1622,7 +1622,7 @@ void FAssetRegistryImpl::CollectCodeGeneratorClasses()
 
 }
 
-void UAssetRegistryImpl::OnRefreshNativeClasses()
+void UAssetRegistryImpl::OnPostEngineInit()
 {
 	LLM_SCOPE(ELLMTag::AssetRegistry);
 	UE::AssetRegistry::FInterfaceWriteScopeLock InterfaceScopeLock(InterfaceLock);
@@ -1939,8 +1939,10 @@ void FAssetRegistryImpl::SearchAllAssets(Impl::FEventContext& EventContext,
 	if (bSynchronousSearch)
 	{
 		Gatherer.WaitForIdle();
-		Impl::FInterruptionContext InterruptionContext(-1., -1.);
-		Impl::EGatherStatus UnusedStatus = TickGatherer(EventContext, InheritanceContext, InterruptionContext);
+		Impl::FTickContext TickContext(EventContext, InheritanceContext);
+		TickContext.bHandleDeferred = true;
+		TickContext.bHandleCompletion = false; // Our caller will call WaitForCompletion which will handle this
+		Impl::EGatherStatus UnusedStatus = TickGatherer(TickContext);
 	}
 	else
 	{
@@ -2013,14 +2015,16 @@ void UAssetRegistryImpl::WaitForCompletion()
 
 			GuardedData.WaitForGathererIdleIfSynchronous();
 
-			UE::AssetRegistry::Impl::FInterruptionContext InterruptionContext(-1., -1.);
-			Status = GuardedData.TickGatherer(EventContext, InheritanceContext, InterruptionContext);
+			UE::AssetRegistry::Impl::FTickContext TickContext(EventContext, InheritanceContext);
+			TickContext.bHandleCompletion = true;
+			TickContext.bHandleDeferred = true;
+			Status = GuardedData.TickGatherer(TickContext);
 		}
 #if WITH_EDITOR
-		UE::AssetRegistry::Impl::FInterruptionContext InterruptionContext(-1., -1.);
+		UE::AssetRegistry::Impl::FInterruptionContext InterruptionContext;
 		ProcessLoadedAssetsToUpdateCache(EventContext, Status, InterruptionContext);
 #endif
-		Broadcast(EventContext);
+		Broadcast(EventContext, true /* bAllowFileLoadedEvent */);
 		bLocalHasSentFileLoadedEventBroadcast |= EventContext.bHasSentFileLoadedEventBroadcast;
 		if (!IsTickActive(Status) && Status != EGatherStatus::WaitingForEvents)
 		{
@@ -4071,7 +4075,7 @@ void UAssetRegistryImpl::ScanPathsSynchronousInternal(const TArray<FString>& InD
 #if WITH_EDITOR
 	if (bWaitForInMemoryObjects)
 	{
-		UE::AssetRegistry::Impl::FInterruptionContext InterruptionContext(-1., -1.);
+		UE::AssetRegistry::Impl::FInterruptionContext InterruptionContext;
 		ProcessLoadedAssetsToUpdateCache(EventContext, Context.Status, InterruptionContext);
 	}
 #endif
@@ -4466,9 +4470,12 @@ UE::AssetRegistry::Impl::EGatherStatus UAssetRegistryImpl::TickOnBackgroundThrea
 			UE::AssetRegistry::Impl::FInterruptionContext::ShouldExitEarlyCallbackType EarlyExitHelper = 
 				[this]()->bool { return InterfaceLock.HasWaiters() || IsBackgroundProcessingPaused(); };
 
-			UE::AssetRegistry::Impl::FInterruptionContext InterruptionContext(FPlatformTime::Seconds(), 
-				UE::AssetRegistry::Impl::MaxSecondsPerTickBackgroundThread, EarlyExitHelper);
-			Status = GuardedData.TickGatherer(EventContext, InitializeContext.InheritanceContext, InterruptionContext);
+			UE::AssetRegistry::Impl::FTickContext TickContext(EventContext, InheritanceContext);
+			TickContext.InterruptionContext.SetLimitedTickTime(FPlatformTime::Seconds(),
+				UE::AssetRegistry::Impl::MaxSecondsPerTickBackgroundThread);
+			TickContext.InterruptionContext.SetEarlyExitCallback(EarlyExitHelper);
+			TickContext.bHandleDeferred = true;
+			Status = GuardedData.TickGatherer(TickContext);
 
 			DeferredEvents.Append(MoveTemp(EventContext));
 			EventContext.Clear();
@@ -4516,7 +4523,6 @@ void UAssetRegistryImpl::Tick(float DeltaTime)
 	do
 	{
 		bInterruptedOrShouldProcessDeferredEvents = false;
-		UE::AssetRegistry::Impl::FInterruptionContext InterruptionContext(TickStartTime, UE::AssetRegistry::Impl::MaxSecondsPerFrame);
 
 		UE::AssetRegistry::Impl::FEventContext EventContext;
 
@@ -4534,7 +4540,12 @@ void UAssetRegistryImpl::Tick(float DeltaTime)
 			if (EventContext.IsEmpty())
 			{
 				// Tick the Gatherer
-				Status = GuardedData.TickGatherer(EventContext, InheritanceContext, InterruptionContext);
+				UE::AssetRegistry::Impl::FTickContext TickContext(EventContext, InheritanceContext);
+				TickContext.InterruptionContext.SetLimitedTickTime(TickStartTime, UE::AssetRegistry::Impl::MaxSecondsPerFrame);
+				TickContext.bHandleCompletion = true;
+				TickContext.bHandleDeferred = true;
+				Status = GuardedData.TickGatherer(TickContext);
+				bInterruptedOrShouldProcessDeferredEvents = TickContext.InterruptionContext.WasInterrupted();
 			}
 			else
 			{
@@ -4552,19 +4563,20 @@ void UAssetRegistryImpl::Tick(float DeltaTime)
 			DeferredEvents.Clear();
 		}
 
-		bInterruptedOrShouldProcessDeferredEvents = bInterruptedOrShouldProcessDeferredEvents || InterruptionContext.WasInterrupted();
-
 #if WITH_EDITOR
 		if (!bInterruptedOrShouldProcessDeferredEvents)
 		{
+			UE::AssetRegistry::Impl::FInterruptionContext InterruptionContext;
+			InterruptionContext.SetLimitedTickTime(TickStartTime, UE::AssetRegistry::Impl::MaxSecondsPerFrame);
 			ProcessLoadedAssetsToUpdateCache(EventContext, Status, InterruptionContext);
-			bInterruptedOrShouldProcessDeferredEvents = bInterruptedOrShouldProcessDeferredEvents || InterruptionContext.WasInterrupted();
+			bInterruptedOrShouldProcessDeferredEvents = bInterruptedOrShouldProcessDeferredEvents
+				|| InterruptionContext.WasInterrupted();
 		}
 #endif
 
 		{
 			TRACE_CPUPROFILER_EVENT_SCOPE_STR("Asset Registry Event Broadcast");
-			Broadcast(EventContext);
+			Broadcast(EventContext, true /* bAllowFileLoadedEvent */);
 		}
 	} while ((bInterruptedOrShouldProcessDeferredEvents || Status == UE::AssetRegistry::Impl::EGatherStatus::WaitingForEvents) &&
 		(TickStartTime < 0 || (FPlatformTime::Seconds() - TickStartTime) <= UE::AssetRegistry::Impl::MaxSecondsPerFrame));
@@ -4596,12 +4608,13 @@ bool FAssetRegistryImpl::ClassRequiresGameThreadProcessing(const UClass* Class)
 	return true;
 }
 
-Impl::EGatherStatus FAssetRegistryImpl::TickGatherer(Impl::FEventContext& EventContext,
-	Impl::FClassInheritanceContext& InheritanceContext, Impl::FInterruptionContext& InOutInterruptionContext,
-	TOptional<FAssetsFoundCallback> AssetsFoundCallback, TOptional<FVerseFilesFoundCallback> VerseFilesFoundCallback)
+Impl::EGatherStatus FAssetRegistryImpl::TickGatherer(Impl::FTickContext& TickContext)
 {
 	using namespace UE::AssetRegistry::Impl;
+	FEventContext& EventContext = TickContext.EventContext;
+	FInterruptionContext& InOutInterruptionContext = TickContext.InterruptionContext;
 	bool bLocalIsInGameThread = IsInGameThread();
+	bool bCanAccessCoreRedirects = bLocalIsInGameThread || IsEngineStartupModuleLoadingComplete();
 
 	EGatherStatus OutStatus = EGatherStatus::Complete;
 	if (!GlobalGatherer.IsValid())
@@ -4675,7 +4688,27 @@ Impl::EGatherStatus FAssetRegistryImpl::TickGatherer(Impl::FEventContext& EventC
 		}
 		else
 		{
-			OutStatus = ResultContext.bAbleToProgress ? EGatherStatus::TickActiveGatherActive : EGatherStatus::TickActiveGatherIdle;
+			EGatherStatus NewStatus = ResultContext.bAbleToProgress ? EGatherStatus::TickActiveGatherActive
+				: EGatherStatus::TickActiveGatherIdle;
+			if (InOutInterruptionContext.WasInterrupted())
+			{
+				// When interrupted we don't know the current status, so just keep the previous status, unless
+				// the previous status is a temporary status, in which case just switch it over to TickActive
+				switch (GatherStatus)
+				{
+				case EGatherStatus::WaitingForEvents:
+				case EGatherStatus::UnableToProgress:
+					OutStatus = NewStatus;
+					break;
+				default:
+					OutStatus = this->GatherStatus;
+					break;
+				}
+			}
+			else
+			{
+				OutStatus = NewStatus;
+			}
 		}
 		if (OutStatus == Impl::EGatherStatus::TickActiveGatherIdle)
 		{
@@ -4724,7 +4757,8 @@ Impl::EGatherStatus FAssetRegistryImpl::TickGatherer(Impl::FEventContext& EventC
 		return OutStatus;
 	}
 
-	auto RunAssetSearchDataGathered = [this, &EventContext, &AssetsFoundCallback, &LazyStartTimer, &InOutInterruptionContext]
+	auto RunAssetSearchDataGathered = [this, &EventContext, &TickContext, &LazyStartTimer,
+		&InOutInterruptionContext, bCanAccessCoreRedirects]
 	(TMultiMap<FName, TUniquePtr<FAssetData>>& InAssetResults,
 		TMultiMap<FName, TUniquePtr<FAssetData>>& OutDeferredAssetResults)
 		{
@@ -4733,17 +4767,18 @@ Impl::EGatherStatus FAssetRegistryImpl::TickGatherer(Impl::FEventContext& EventC
 			{
 				LazyStartTimer();
 				// Mark the first amortize time
-				if (AssetsFoundCallback.IsSet())
+				if (TickContext.AssetsFoundCallback.IsSet())
 				{
 					TMultiMap<FName, FAssetData*> NonOwningContainer;
 					for (auto Iter = InAssetResults.CreateIterator(); Iter; ++Iter)
 					{
 						NonOwningContainer.Add(Iter.Key(), Iter.Value().Get());
 					}
-					AssetsFoundCallback.GetValue()(NonOwningContainer);
+					TickContext.AssetsFoundCallback.GetValue()(NonOwningContainer);
 				}
 
-				AssetSearchDataGathered(EventContext, InAssetResults, OutDeferredAssetResults, InOutInterruptionContext);
+				AssetSearchDataGathered(EventContext, InAssetResults, OutDeferredAssetResults,
+					InOutInterruptionContext, bCanAccessCoreRedirects);
 			}
 		};
 	auto RunDependencyDataGathered = [this, &bLocalIsInGameThread, &LazyStartTimer, &InOutInterruptionContext]
@@ -4756,13 +4791,14 @@ Impl::EGatherStatus FAssetRegistryImpl::TickGatherer(Impl::FEventContext& EventC
 			{
 				LazyStartTimer();
 
-				DependencyDataGathered(DependenciesToProcess, OutDeferredDependencies, OutPackagesNeedingDependencyCalculation, InOutInterruptionContext);
-
+				DependencyDataGathered(DependenciesToProcess, OutDeferredDependencies,
+					OutPackagesNeedingDependencyCalculation, InOutInterruptionContext);
 			}
 		};
 
 	bool bRetryAssetGathering = true;
-	bool bHasRetriedDeferredResults = false;
+	int32 OriginalDeferredAssetsCount = 0;
+	int32 NumRetries = 0;
 	while (bRetryAssetGathering)
 	{
 		bRetryAssetGathering = false;
@@ -4815,21 +4851,43 @@ Impl::EGatherStatus FAssetRegistryImpl::TickGatherer(Impl::FEventContext& EventC
 
 		// Retry deferred assets if we've finished all the other assets; we need to do this in the current tick
 		// so we avoid spuriously reporting status == UnableToProgress
-		if (!bHasRetriedDeferredResults)
+		if (BackgroundResults.Assets.IsEmpty()
+			&& (!bLocalIsInGameThread || BackgroundResults.AssetsForGameThread.IsEmpty())
+			&& bCanAccessCoreRedirects && TickContext.bHandleDeferred)
 		{
-			bHasRetriedDeferredResults = true;
-			if (BackgroundResults.Assets.IsEmpty() &&
-				(!bLocalIsInGameThread || BackgroundResults.AssetsForGameThread.IsEmpty()))
+			if (!DeferredAssets.IsEmpty() || !DeferredDependencies.IsEmpty() ||
+				(bLocalIsInGameThread && (!DeferredAssetsForGameThread.IsEmpty() || !DeferredDependenciesForGameThread.IsEmpty())))
 			{
-				if (!DeferredAssets.IsEmpty() || !DeferredDependencies.IsEmpty() ||
-					(bLocalIsInGameThread && (!DeferredAssetsForGameThread.IsEmpty() || !DeferredDependenciesForGameThread.IsEmpty())))
+				if (bProcessedAnyAssetsAfterRetryDeferred)
 				{
 					bRetryAssetGathering = true;
+				}
+				else
+				{
+					if (!bForceCompletionEvenIfPostLoadsFail && bPreloadingComplete && IsEngineStartupModuleLoadingComplete())
+					{
+						bForceCompletionEvenIfPostLoadsFail = true;
+						bRetryAssetGathering = true;
+					}
+				}
+				if (bRetryAssetGathering)
+				{
+					bProcessedAnyAssetsAfterRetryDeferred = false;
+					if (NumRetries == 0)
+					{
+						OriginalDeferredAssetsCount = DeferredAssets.Num() + DeferredAssetsForGameThread.Num()
+							+ 10; // fudge factor to make sure an edge case of 0 does not cause a problem
+					}
+					if (NumRetries++ >= OriginalDeferredAssetsCount)
+					{
+						UE_LOG(LogAssetRegistry, Error, TEXT("Runaway loop detected in handling of deferred assets"));
+						// This will cause us to return UnableToProgress status
+						break;
+					}
 					BackgroundResults.Assets.Append(MoveTemp(DeferredAssets));
 					BackgroundResults.AssetsForGameThread.Append(MoveTemp(DeferredAssetsForGameThread));
 					BackgroundResults.Dependencies.Append(MoveTemp(DeferredDependencies));
 					BackgroundResults.DependenciesForGameThread.Append(MoveTemp(DeferredDependenciesForGameThread));
-					bForceCompletionEvenIfPostLoadsFail = bPreloadingComplete && IsEngineStartupModuleLoadingComplete();
 				}
 			}
 		}
@@ -4852,9 +4910,9 @@ Impl::EGatherStatus FAssetRegistryImpl::TickGatherer(Impl::FEventContext& EventC
 	if (BackgroundResults.VerseFiles.Num())
 	{
 		LazyStartTimer();
-		if (VerseFilesFoundCallback.IsSet())
+		if (TickContext.VerseFilesFoundCallback.IsSet())
 		{
-			VerseFilesFoundCallback.GetValue()(BackgroundResults.VerseFiles);
+			TickContext.VerseFilesFoundCallback.GetValue()(BackgroundResults.VerseFiles);
 		}
 
 		VerseFilesGathered(EventContext, BackgroundResults.VerseFiles, InOutInterruptionContext);
@@ -4902,7 +4960,8 @@ Impl::EGatherStatus FAssetRegistryImpl::TickGatherer(Impl::FEventContext& EventC
 
 		if (PackagesNeedingDependencyCalculationOnGameThread.Num() && bLocalIsInGameThread)
 		{
-			LoadCalculatedDependencies(nullptr, InheritanceContext, &PackagesNeedingDependencyCalculationOnGameThread, InOutInterruptionContext);
+			LoadCalculatedDependencies(nullptr, TickContext.InheritanceContext,
+				&PackagesNeedingDependencyCalculationOnGameThread, InOutInterruptionContext);
 			if (InOutInterruptionContext.ShouldExitEarly())
 			{
 				CalculateStatus(NumGatherFromDiskPending);
@@ -4922,7 +4981,8 @@ Impl::EGatherStatus FAssetRegistryImpl::TickGatherer(Impl::FEventContext& EventC
 			// Finishing the background search is blocked until preloading complete because plugins can be mounted during
 			// startup up until that point, and we need to wait for all the plugins to load before declaring completion.
 			// Only the main thread can know that we're complete because we need to wait until we've broadcast the events
-			bool bCanCompleteInitialSearch = bPreloadingComplete && IsEngineStartupModuleLoadingComplete() && bLocalIsInGameThread;
+			bool bCanCompleteInitialSearch = bPreloadingComplete && IsEngineStartupModuleLoadingComplete()
+				&& bLocalIsInGameThread && TickContext.bHandleCompletion;
 
 			if (bCanCompleteInitialSearch)
 			{
@@ -4938,7 +4998,7 @@ Impl::EGatherStatus FAssetRegistryImpl::TickGatherer(Impl::FEventContext& EventC
 			}
 			else
 			{
-				if (bLocalIsInGameThread)
+				if (bLocalIsInGameThread && TickContext.bHandleCompletion)
 				{
 					UE_LOG(LogAssetRegistry, Display, TEXT("TickGatherer returning UnableToProgress because bCanCompleteInitialSearch is false but our work is otherwise complete. "
 						"bPreloadingComplete == %s; IsEngineStartupModuleLoadingComplete() == %s; bLocalIsInGameThread == %s"),
@@ -4994,7 +5054,6 @@ void FAssetRegistryImpl::OnInitialSearchCompleted(Impl::FEventContext& EventCont
 	GlobalGatherer->OnInitialSearchCompleted();
 
 	EventContext.bFileLoadedEventBroadcast = true;
-	EventContext.bScanEndedEventBroadcast = true;
 }
 
 void FAssetRegistryImpl::LogSearchDiagnostics(double StartTime)
@@ -5061,6 +5120,7 @@ void FAssetRegistryImpl::TickGatherPackage(Impl::FEventContext& EventContext, co
 		}
 	};
 
+	bool bCanAccessCoreRedirects = IsInGameThread() || IsEngineStartupModuleLoadingComplete();
 	FName PackageFName(PackageName);
 
 	// Gather results from the background search
@@ -5078,6 +5138,13 @@ void FAssetRegistryImpl::TickGatherPackage(Impl::FEventContext& EventContext, co
 	BackgroundResults.DependenciesForGameThread.MultiFind(PackageFName, PackageDependencyDatas);
 	BackgroundResults.DependenciesForGameThread.Remove(PackageFName);
 
+	DeferredAssets.MultiFindPointer(PackageFName, PackageAssets);
+	DeferredAssetsForGameThread.MultiFindPointer(PackageFName, PackageAssets);
+	DeferredDependencies.MultiFind(PackageFName, PackageDependencyDatas);
+	DeferredDependencies.Remove(PackageFName);
+	DeferredDependenciesForGameThread.MultiFind(PackageFName, PackageDependencyDatas);
+	DeferredDependenciesForGameThread.Remove(PackageFName);
+
 	if (PackageAssets.Num() > 0)
 	{
 		LazyStartTimer();
@@ -5090,10 +5157,13 @@ void FAssetRegistryImpl::TickGatherPackage(Impl::FEventContext& EventContext, co
 		// Ownership transfer is now complete so remove these packages from the results arrays
 		BackgroundResults.Assets.Remove(PackageFName);
 		BackgroundResults.AssetsForGameThread.Remove(PackageFName);
+		DeferredAssets.Remove(PackageFName);
+		DeferredAssetsForGameThread.Remove(PackageFName);
 
 		TMultiMap<FName, TUniquePtr<FAssetData>> DeferredPackageAssetsMap;
-		UE::AssetRegistry::Impl::FInterruptionContext InterruptionContext(-1., -1.);
-		AssetSearchDataGathered(EventContext, PackageAssetsMap, DeferredPackageAssetsMap, InterruptionContext);
+		UE::AssetRegistry::Impl::FInterruptionContext InterruptionContext;
+		AssetSearchDataGathered(EventContext, PackageAssetsMap, DeferredPackageAssetsMap,
+			InterruptionContext, bCanAccessCoreRedirects);
 		if (DeferredPackageAssetsMap.Num())
 		{
 			UE_LOG(LogAssetRegistry, Warning, TEXT("Attempted to add package '%s' to the registry before its UClass was available. \
@@ -5124,8 +5194,9 @@ Could not execute PostLoadAssetRegistryTags. We will try again later. Until then
 #if WITH_EDITOR
 		OutPackagesNeedingDependencyCalculation = &PackagesNeedingDependencyCalculation;
 #endif
-		UE::AssetRegistry::Impl::FInterruptionContext InterruptionContext(-1., -1.);
-		DependencyDataGathered(PackageDependencyDatasMap, DeferredDependenciesForGameThread, OutPackagesNeedingDependencyCalculation, InterruptionContext);
+		UE::AssetRegistry::Impl::FInterruptionContext InterruptionContext;
+		DependencyDataGathered(PackageDependencyDatasMap, DeferredDependenciesForGameThread,
+			OutPackagesNeedingDependencyCalculation, InterruptionContext);
 	}
 }
 
@@ -5273,6 +5344,11 @@ void FAssetRegistryImpl::RemoveDirectoryReferencer(FName PackageName)
 
 void UAssetRegistryImpl::Serialize(FArchive& Ar)
 {
+	if (Ar.IsObjectReferenceCollector())
+	{
+		// The Asset Registry does not have any object references, and its serialization function is expensive
+		return;
+	}
 	UE::AssetRegistry::Impl::FEventContext EventContext;
 	{
 		LLM_SCOPE(ELLMTag::AssetRegistry);
@@ -5287,12 +5363,8 @@ namespace UE::AssetRegistry
 
 void FAssetRegistryImpl::Serialize(FArchive& Ar, Impl::FEventContext& EventContext)
 {
-	if (Ar.IsObjectReferenceCollector())
-	{
-		// The Asset Registry does not have any object references, and its serialization function is expensive
-		return;
-	}
-	else if (Ar.IsLoading())
+	check(!Ar.IsObjectReferenceCollector()); // Caller should not call in this case
+	if (Ar.IsLoading())
 	{
 		State.Load(Ar);
 		CachePathsFromState(EventContext, State);
@@ -5784,8 +5856,7 @@ void FAssetRegistryImpl::ScanPathsSynchronous(Impl::FScanPathContext& Context)
 	{
 		Context.NumFoundAssets = InFoundAssets.Num();
 
-		FoundAssetPackageNames.Reset();
-		FoundAssetPackageNames.Reserve(Context.NumFoundAssets);
+		FoundAssetPackageNames.Reserve(FoundAssetPackageNames.Num() + Context.NumFoundAssets);
 
 		// The gatherer may have added other assets that were scanned as part of the ongoing background scan,
 		// so remove any assets that were not in the requested paths
@@ -5842,9 +5913,10 @@ void FAssetRegistryImpl::ScanPathsSynchronous(Impl::FScanPathContext& Context)
 			}
 		};
 
-	Impl::FInterruptionContext InterruptionContext(-1., -1.);
-	Context.Status = TickGatherer(Context.EventContext, Context.InheritanceContext,	InterruptionContext,
-		FAssetsFoundCallback(AssetsFoundCallback), FVerseFilesFoundCallback(VerseFileFoundCallback));
+	Impl::FTickContext TickContext(Context.EventContext, Context.InheritanceContext);
+	TickContext.AssetsFoundCallback = Impl::FAssetsFoundCallback(AssetsFoundCallback);
+	TickContext.VerseFilesFoundCallback = Impl::FVerseFilesFoundCallback(VerseFileFoundCallback);
+	Context.Status = TickGatherer(TickContext);
 
 	// Temporary hack/partial solution. The expectation is that this function will return cause all assets
 	// under the specified directories to be ingested into the registry. However, one of the early steps
@@ -5867,6 +5939,7 @@ void FAssetRegistryImpl::ScanPathsSynchronous(Impl::FScanPathContext& Context)
 		{
 			if (IsInRequestedDir(*Iter.Value()))
 			{
+				FoundAssetPackageNames.Add(Iter.Key());
 				CollectedDeferredAssets.Add(MoveTemp(*Iter));
 				Iter.RemoveCurrent();
 			}
@@ -5875,27 +5948,50 @@ void FAssetRegistryImpl::ScanPathsSynchronous(Impl::FScanPathContext& Context)
 		{
 			if (IsInRequestedDir(*Iter.Value()))
 			{
+				FoundAssetPackageNames.Add(Iter.Key());
 				CollectedDeferredAssets.Add(MoveTemp(*Iter));
 				Iter.RemoveCurrent();
 			}
 		}
+		bool bCanAccessCoreRedirects = IsInGameThread() || IsEngineStartupModuleLoadingComplete();
 		// Force AssetSearchDataGathered to process these assets, skipping the PostLoadAssetRegistryTags if needed
+		// But don't allow this if !bCanAccessCoreRedirects. The caller should not be calling ScanPathsSynchronous
+		// in that case.
 		const bool bOldForceCompletionEvenIfPostLoadsFail = bForceCompletionEvenIfPostLoadsFail;
-		bForceCompletionEvenIfPostLoadsFail = true;
-		// We don't call the AssetsFoundCallback here because even for deferred assets it will already have been called.
-		// We pass DeferredAssetsForGameThread as the OutDeferred parameter, but the expectation is that nothing will be deferred
-		AssetSearchDataGathered(Context.EventContext, CollectedDeferredAssets, DeferredAssetsForGameThread, InterruptionContext);
-		// All of the assets we collected should have been processed.
+		bForceCompletionEvenIfPostLoadsFail = bCanAccessCoreRedirects;
+
+		int32 OriginalNumDeferredAssetsForGameThread = DeferredAssetsForGameThread.Num();
+		// We don't call AssetsFoundCallback here because even for deferred assets it will already have been called.
+		// We pass DeferredAssetsForGameThread as the OutDeferred parameter, but we expect nothing will be deferred.
+		AssetSearchDataGathered(Context.EventContext, CollectedDeferredAssets, DeferredAssetsForGameThread,
+			TickContext.InterruptionContext, bCanAccessCoreRedirects);
+		// All of the assets we collected should have been processed or deferred.
 		ensure(CollectedDeferredAssets.Num() == 0);
+		// We should not have deferred any assets because we set bForceCompletionEvenIfPostLoadsFail. If we couldn't
+		// set it because the caller is calling to early, log an error about it.
+		if (DeferredAssetsForGameThread.Num() > OriginalNumDeferredAssetsForGameThread)
+		{
+			// With current logic the only way we could defer is if !bCanAccessCoreRedirects
+			check(!bCanAccessCoreRedirects);
+			ensureMsgf(false,
+				TEXT("ScanPathsSynchronous was called on a thread other than the game thread, before IsEngineStartupModuleLoadingComplete. ")
+				TEXT("This caused us not to be able to complete the scan because we can not access CoreRedirects. Some assets will be missing. ")
+				TEXT("Calling code should change to not allow the ScanPathsSynchronous call this early, or queue it to the gamethread."));
+		}
 		bForceCompletionEvenIfPostLoadsFail = bOldForceCompletionEvenIfPostLoadsFail;
-		// Now run one more tick to perform any subsequent processing required for these assets beyond AssetSearchDataGathered
-		Context.Status = TickGatherer(Context.EventContext, Context.InheritanceContext, InterruptionContext,
-			FAssetsFoundCallback(AssetsFoundCallback));
+		// Tick to perform any subsequent processing required for these assets beyond AssetSearchDataGathered
+		Impl::FTickContext AssetTickContext(Context.EventContext, Context.InheritanceContext);
+		AssetTickContext.AssetsFoundCallback = Impl::FAssetsFoundCallback(AssetsFoundCallback);
+		Context.Status = TickGatherer(AssetTickContext);
 	}
+	FoundAssetPackageNames.Sort(FNameFastLess());
+	FoundAssetPackageNames.SetNum(Algo::Unique(FoundAssetPackageNames));
 
 #if WITH_EDITOR
-	LoadCalculatedDependencies(&FoundAssetPackageNames, Context.InheritanceContext, &PackagesNeedingDependencyCalculation, InterruptionContext);
-	LoadCalculatedDependencies(&FoundAssetPackageNames, Context.InheritanceContext, &PackagesNeedingDependencyCalculationOnGameThread, InterruptionContext);
+	LoadCalculatedDependencies(&FoundAssetPackageNames, Context.InheritanceContext,
+		&PackagesNeedingDependencyCalculation, TickContext.InterruptionContext);
+	LoadCalculatedDependencies(&FoundAssetPackageNames, Context.InheritanceContext,
+		&PackagesNeedingDependencyCalculationOnGameThread, TickContext.InterruptionContext);
 #endif
 	for (FSoftObjectPath& OldAssetToRemove : OldAssetsToRemove)
 	{
@@ -5963,14 +6059,25 @@ FAssetData* FAssetRegistryImpl::ResolveAssetIdCollision(FAssetData& A, FAssetDat
 	return Keep;
 }
 
-bool FAssetRegistryImpl::TryPostLoadAssetRegistryTags(FAssetData* AssetData)
+bool FAssetRegistryImpl::TryPostLoadAssetRegistryTags(FAssetData* AssetData, bool bCanAccessCoreRedirects)
 {
 	check(AssetData);
-	bool CouldPostLoadAssetRegistryTags = true;
-	if (AssetData->TagsAndValues.Num())
+	if (!AssetData->TagsAndValues.Num())
 	{
-		FTopLevelAssetPath AssetClassPath = AssetData->AssetClassPath;
-		UClass* AssetClass = FindObject<UClass>(AssetClassPath, true);
+		return true;
+	}
+
+	bool bCouldPostLoadAssetRegistryTags = true;
+	UClass* AssetClass = nullptr;
+	FTopLevelAssetPath AssetClassPath;
+	if (!bCanAccessCoreRedirects)
+	{
+		bCouldPostLoadAssetRegistryTags = false;
+	}
+	else
+	{
+		AssetClassPath = AssetData->AssetClassPath;
+		AssetClass = FindObject<UClass>(AssetClassPath, true);
 
 		while (!AssetClass)
 		{
@@ -5998,103 +6105,107 @@ bool FAssetRegistryImpl::TryPostLoadAssetRegistryTags(FAssetData* AssetData)
 						AssetClassPath = NewName.ToString();
 					}
 				}
-				
+
 				if (AssetClassPath != LastAssetClassPath && !AssetClassPath.IsNull())
 				{
 					AssetClass = FindObject<UClass>(AssetClassPath, true);
 				}
 				else
 				{
-					CouldPostLoadAssetRegistryTags = false;
+					bCouldPostLoadAssetRegistryTags = false;
 					break;
 				}
 			}
 		}
+	}
 
-		// Now identify the most derived native class in the class hierarchy
-		if (AssetClass)
+	// Now identify the most derived native class in the class hierarchy
+	if (AssetClass)
+	{
+		while (!AssetClass->HasAnyClassFlags(CLASS_Native))
 		{
-			while (!AssetClass->HasAnyClassFlags(CLASS_Native))
-			{
-				AssetClass = AssetClass->GetSuperClass();
-			}
+			AssetClass = AssetClass->GetSuperClass();
 		}
+	}
 
-		bool MakeFinalChecks = false;
-		if (bForceCompletionEvenIfPostLoadsFail && bPreloadingComplete && IsEngineStartupModuleLoadingComplete())
-		{
-			// Okay, we think we're done loading and now we need to make some expensive final checks to try to either
-			// track down the classes for fixup or just give up
-			MakeFinalChecks = true;
-		}
-		if (!AssetClass && bForceCompletionEvenIfPostLoadsFail)
+	bool bMakeFinalChecks = false;
+	if (bForceCompletionEvenIfPostLoadsFail && bPreloadingComplete && IsEngineStartupModuleLoadingComplete())
+	{
+		// Okay, we think we're done loading and now we need to make some expensive final checks to try to either
+		// track down the classes for fixup or just give up
+		bMakeFinalChecks = true;
+		// bCanAccessCoreRedirects should be true when IsEngineStartupModuleLoadingComplete
+		check(bCanAccessCoreRedirects);
+	}
+	if (!AssetClass && bForceCompletionEvenIfPostLoadsFail)
+	{
+		if (bMakeFinalChecks)
 		{
-			if (MakeFinalChecks)
+			FString Reason;
+			if (AssetClassPath.ToString().StartsWith(TEXT("/Script/")))
 			{
-				FString Reason;
-				if (AssetClassPath.ToString().StartsWith(TEXT("/Script/")))
-				{
-					Reason = TEXT("The missing class is native--perhaps a CoreRedirector is missing?");
-				}
-				else
+				Reason = TEXT("The missing class is native--perhaps a CoreRedirector is missing?");
+			}
+			else
+			{
+				if (State.GetAssetPackageData(AssetClassPath.GetPackageName()) == nullptr)
 				{
-					if (State.GetAssetPackageData(AssetClassPath.GetPackageName()) == nullptr)
-					{
-						Reason = TEXT("The class is missing on disk or could not be loaded. Perhaps it has been deleted from perforce and the referencing object is broken?");
-					}
+					Reason = TEXT("The class is missing on disk or could not be loaded. Perhaps it has been deleted from perforce and the referencing object is broken?");
 				}
-				//@TODO this should become a Warning once UE-209846 is finished
-				UE_LOG(LogAssetRegistry, Verbose, TEXT("Unable to PostLoadAssetRegistryTags for '%s' because ancestor class '%s' cannot be found. %s"), 
-					*AssetData->GetObjectPathString(), *AssetClassPath.ToString(), *Reason);
 			}
-
-			// Force this so that we can move on
-			CouldPostLoadAssetRegistryTags = true;
+			//@TODO this should become a Warning once UE-209846 is finished
+			UE_LOG(LogAssetRegistry, Verbose,
+				TEXT("Unable to PostLoadAssetRegistryTags for '%s' because ancestor class '%s' cannot be found. %s"), 
+				*AssetData->GetObjectPathString(), *AssetClassPath.ToString(), *Reason);
 		}
 
-		if (AssetClass)
+		// Force this so that we can move on
+		bCouldPostLoadAssetRegistryTags = true;
+	}
+
+	if (AssetClass)
+	{
+		UObject* ClassDefaultObject = AssetClass->GetDefaultObject(false);
+		if (ClassDefaultObject && !ClassDefaultObject->HasAnyFlags(RF_NeedInitialization))
 		{
-			UObject* ClassDefaultObject = AssetClass->GetDefaultObject(false);
-			if (ClassDefaultObject && !ClassDefaultObject->HasAnyFlags(RF_NeedInitialization))
+			// We are using RF_NeedInitialization to guarantee that ClassDefaultObject is fully initialized
+			// potentially on another thread. For weakly ordered memory platforms, we need to 
+			// ensure that our read of the vtable ptr isn't performed prior to the read of the class flags
+			// otherwise we might see a stale vtable despite seeing RF_NeedInit clear.
+			std::atomic_thread_fence(std::memory_order_acquire);
+			TArray<UObject::FAssetRegistryTag> TagsToModify;
+			UObject::FPostLoadAssetRegistryTagsContext Context(*AssetData, TagsToModify);
+			ClassDefaultObject->ThreadedPostLoadAssetRegistryTags(Context);
+			if (TagsToModify.Num())
 			{
-				// We are using RF_NeedInitialization to guarantee that ClassDefaultObject is fully initialized
-				// potentially on another thread. For weakly ordered memory platforms, we need to 
-				// ensure that our read of the vtable ptr isn't performed prior to the read of the class flags
-				// otherwise we might see a stale vtable despite seeing RF_NeedInit clear.
-				std::atomic_thread_fence(std::memory_order_acquire);
-				TArray<UObject::FAssetRegistryTag> TagsToModify;
-				UObject::FPostLoadAssetRegistryTagsContext Context(*AssetData, TagsToModify);
-				ClassDefaultObject->ThreadedPostLoadAssetRegistryTags(Context);
-				if (TagsToModify.Num())
+				FAssetDataTagMap TagsAndValues = AssetData->TagsAndValues.CopyMap();
+				for (const UObject::FAssetRegistryTag& Tag : TagsToModify)
 				{
-					FAssetDataTagMap TagsAndValues = AssetData->TagsAndValues.CopyMap();
-					for (const UObject::FAssetRegistryTag& Tag : TagsToModify)
+					if (!Tag.Value.IsEmpty())
 					{
-						if (!Tag.Value.IsEmpty())
-						{
-							FString& Value = TagsAndValues.FindOrAdd(Tag.Name);
-							Value = Tag.Value;
-						}
-						else
-						{
-							TagsAndValues.Remove(Tag.Name);
-						}
+						FString& Value = TagsAndValues.FindOrAdd(Tag.Name);
+						Value = Tag.Value;
+					}
+					else
+					{
+						TagsAndValues.Remove(Tag.Name);
 					}
-					AssetData->TagsAndValues = FAssetDataTagMapSharedView(MoveTemp(TagsAndValues));
 				}
+				AssetData->TagsAndValues = FAssetDataTagMapSharedView(MoveTemp(TagsAndValues));
 			}
-			else if (!bForceCompletionEvenIfPostLoadsFail)
-			{
-				CouldPostLoadAssetRegistryTags = false;
-			}
-			else 
-			{
-				ensureMsgf(!MakeFinalChecks, TEXT("Unable to PostLoadAssetRegistryTags for '%s' because the CDO for ancestor class '%s' could not be found or was not ready."),
-					*AssetData->GetObjectPathString(), *AssetClassPath.ToString());
-			}
+		}
+		else if (!bForceCompletionEvenIfPostLoadsFail)
+		{
+			bCouldPostLoadAssetRegistryTags = false;
+		}
+		else 
+		{
+			ensureMsgf(!bMakeFinalChecks,
+				TEXT("Unable to PostLoadAssetRegistryTags for '%s' because the CDO for ancestor class '%s' could not be found or was not ready."),
+				*AssetData->GetObjectPathString(), *AssetClassPath.ToString());
 		}
 	}
-	return CouldPostLoadAssetRegistryTags;
+	return bCouldPostLoadAssetRegistryTags;
 }
 #endif
 
@@ -6210,8 +6321,10 @@ bool FAssetRegistryImpl::ShouldSkipGatheredAsset(FAssetData& AssetData)
 	return false;
 }
 
-void FAssetRegistryImpl::AssetSearchDataGathered(Impl::FEventContext& EventContext, TMultiMap<FName, TUniquePtr<FAssetData>>& AssetResults,
-	TMultiMap<FName, TUniquePtr<FAssetData>>& OutDeferredAssetResults, Impl::FInterruptionContext& InOutInterruptionContext)
+void FAssetRegistryImpl::AssetSearchDataGathered(Impl::FEventContext& EventContext,
+	TMultiMap<FName, TUniquePtr<FAssetData>>& AssetResults,
+	TMultiMap<FName, TUniquePtr<FAssetData>>& OutDeferredAssetResults,
+	Impl::FInterruptionContext& InOutInterruptionContext, bool bCanAccessCoreRedirects)
 {
 	TRACE_CPUPROFILER_EVENT_SCOPE(AssetSearchDataGathered);
 
@@ -6280,13 +6393,14 @@ void FAssetRegistryImpl::AssetSearchDataGathered(Impl::FEventContext& EventConte
 
 #if WITH_EDITOR
 		// Postload assets based on their declared class. Queue them for for later retry if their class has not yet loaded.
-		bool CouldPostLoad = TryPostLoadAssetRegistryTags(BackgroundResult.Get());
+		bool CouldPostLoad = TryPostLoadAssetRegistryTags(BackgroundResult.Get(), bCanAccessCoreRedirects);
 		if (!CouldPostLoad)
 		{
 			OutDeferredAssetResults.Add(BackgroundAssetPackageName, MoveTemp(BackgroundResult));
 			continue;
 		}
 #endif
+		bProcessedAnyAssetsAfterRetryDeferred = true;
 
 		// Look for an existing asset to check whether we need to add or update
 		FCachedAssetKey Key(*BackgroundResult);
@@ -7503,7 +7617,7 @@ void UAssetRegistryImpl::ScanModifiedAssetFiles(const TArray<FString>& InFilePat
 	// but in-memory results will override the on-disk results we just scanned,
 	// and our in-memory results might be out of date due to being queued but not yet processed.
 	// So ProcessLoadedAssetsToUpdateCache before returning to make sure results are up to date.
-	UE::AssetRegistry::Impl::FInterruptionContext InterruptionContext(-1., -1.);
+	UE::AssetRegistry::Impl::FInterruptionContext InterruptionContext;
 	ProcessLoadedAssetsToUpdateCache(EventContext, UE::AssetRegistry::Impl::EGatherStatus::Complete, InterruptionContext);
 #endif
 
@@ -8637,6 +8751,7 @@ namespace Impl
 
 void FEventContext::Clear()
 {
+	bScanStartedEventBroadcast = false;
 	bFileLoadedEventBroadcast = false;
 	bHasSentFileLoadedEventBroadcast = false;
 	ProgressUpdateData.Reset();
@@ -8648,7 +8763,8 @@ void FEventContext::Clear()
 
 bool FEventContext::IsEmpty() const
 {
-	return !bFileLoadedEventBroadcast &&
+	return !bScanStartedEventBroadcast &&
+		!bFileLoadedEventBroadcast &&
 		!ProgressUpdateData.IsSet() &&
 		PathEvents.Num() == 0 &&
 		AssetEvents.Num() == 0 &&
@@ -8662,7 +8778,10 @@ void FEventContext::Append(FEventContext&& Other)
 	{
 		return;
 	}
+	bScanStartedEventBroadcast |= Other.bScanStartedEventBroadcast;
+	Other.bScanStartedEventBroadcast = false;
 	bFileLoadedEventBroadcast |= Other.bFileLoadedEventBroadcast;
+	Other.bFileLoadedEventBroadcast = false;
 	if (Other.ProgressUpdateData.IsSet())
 	{
 		ProgressUpdateData = MoveTemp(Other.ProgressUpdateData);
@@ -8696,7 +8815,7 @@ void UAssetRegistryImpl::ReadLockEnumerateAllTagToAssetDatas(TFunctionRef<bool(F
 	GuardedData.GetState().EnumerateTagToAssetDatas(Callback);
 }
 
-void UAssetRegistryImpl::Broadcast(UE::AssetRegistry::Impl::FEventContext& EventContext)
+void UAssetRegistryImpl::Broadcast(UE::AssetRegistry::Impl::FEventContext& EventContext, bool bAllowFileLoadedEvent)
 {
 	using namespace UE::AssetRegistry::Impl;
 	if (!IsInGameThread() || FUObjectThreadContext::Get().IsRoutingPostLoad)
@@ -8704,8 +8823,14 @@ void UAssetRegistryImpl::Broadcast(UE::AssetRegistry::Impl::FEventContext& Event
 		// By contract events (and PackageLoads) can only be sent on the game thread; some legacy systems depend on 
 		// this and are not threadsafe. If we're not in the game thread, defer all events in the EventContext
 		// instead of broadcasting them on this thread
+		if (EventContext.IsEmpty())
+		{
+			return;
+		}
 		UE::AssetRegistry::FInterfaceWriteScopeLock InterfaceScopeLock(InterfaceLock);
-		check(&EventContext != &DeferredEvents); // Only the GameThread should be calling Broadcast on DeferredEvents
+		// Broadcast should not be called on DeferredEvents; DeferredEvents should be moved to a separate EventContext
+		// while under the InterfaceLock, and broadcast called on that separate Eventcontext outside of the lock.
+		check(&EventContext != &DeferredEvents);
 		DeferredEvents.Append(MoveTemp(EventContext));
 		EventContext.Clear();
 		return;
@@ -8873,29 +8998,43 @@ void UAssetRegistryImpl::Broadcast(UE::AssetRegistry::Impl::FEventContext& Event
 
 	if (EventContext.bFileLoadedEventBroadcast)
 	{
-		if (&EventContext != &DeferredEvents)
+		if (!bAllowFileLoadedEvent)
 		{
+			// Do not send the file loaded event yet; pass the flag on instead
 			UE::AssetRegistry::FInterfaceWriteScopeLock InterfaceScopeLock(InterfaceLock);
-			// Do not send the file loaded event yet if there are still deferred events and pass the flag on instead
-			if (!DeferredEvents.IsEmpty())
-			{
-				EventContext.bFileLoadedEventBroadcast = false;
-				DeferredEvents.bFileLoadedEventBroadcast = true;
-				return;
-			}
+			// Broadcast should not be called on DeferredEvents; DeferredEvents should be moved to a separate EventContext
+			// while under the InterfaceLock, and broadcast called on that separate Eventcontext outside of the lock.
+			check(&EventContext != &DeferredEvents);
+			DeferredEvents.Append(MoveTemp(EventContext));
+			EventContext.Clear();
+			check(!EventContext.bFileLoadedEventBroadcast); // was cleared by Append and by Clear
+			check(DeferredEvents.bFileLoadedEventBroadcast); // was set by Append
+			return;
 		}
 
+		FEventContext CopiedDeferredEvents;
+		{
+			UE::AssetRegistry::FInterfaceWriteScopeLock InterfaceScopeLock(InterfaceLock);
+			check(&EventContext != &DeferredEvents);
+			CopiedDeferredEvents = MoveTemp(DeferredEvents);
+			DeferredEvents.Clear();
+		}
+		if (!CopiedDeferredEvents.IsEmpty())
+		{
+			// Recursively send all of the deferred events, except for the fileloaded event (which should not exist
+			// on DeferredEvents at this point, but it's not a problem for us if it does).
+			CopiedDeferredEvents.bFileLoadedEventBroadcast = false;
+			Broadcast(CopiedDeferredEvents, false /* bAllowFileLoadedEvent */);
+		}
+		// Now it is safe to broadcast the FileLoadedEvent. If other deferred events come in on another thread after we
+		// copied from DeferredEvents, that is okay; the contract for FileLoadedEvent is that it can not be sent before
+		// sending the events that were fired before it was fired, and new deferred events were fired after it.
+
 		FileLoadedEvent.Broadcast();
+		ScanEndedEvent.Broadcast();
 		EventContext.bFileLoadedEventBroadcast = false;
 		EventContext.bHasSentFileLoadedEventBroadcast = true;
 	}
-
-	if (EventContext.bScanEndedEventBroadcast)
-	{
-		// Raise event when the scan is ended
-		ScanEndedEvent.Broadcast();
-		EventContext.bScanEndedEventBroadcast = false;
-	}
 }
 
 
diff --git a/Engine/Source/Runtime/AssetRegistry/Private/AssetRegistry.h b/Engine/Source/Runtime/AssetRegistry/Private/AssetRegistry.h
--- a/Engine/Source/Runtime/AssetRegistry/Private/AssetRegistry.h
+++ b/Engine/Source/Runtime/AssetRegistry/Private/AssetRegistry.h
@@ -226,7 +226,7 @@ private:
 	void OnAssetDependencyGathererRegistered();
 #endif
 	void InitializeEvents(UE::AssetRegistry::Impl::FInitializeContext& Context);
-	void Broadcast(UE::AssetRegistry::Impl::FEventContext& EventContext);
+	void Broadcast(UE::AssetRegistry::Impl::FEventContext& EventContext, bool bAllowFileLoadedEvent = false);
 
 	bool OnResolveRedirect(const FString& InPackageName, FString& OutPackageName);
 
@@ -257,9 +257,12 @@ private:
 	void OnContentPathDismounted(const FString& AssetPath, const FString& FileSystemPath);
 
 	/** Called to refresh the native classes list, called at end of engine initialization. */
-	void OnRefreshNativeClasses();
+	void OnPostEngineInit();
 
-	/** Called from LaunchEngineLoop via SetEngineStartupModuleLoadingComplete after plugins are loaded, used to scan classes that were loaded by plugins. */
+	/**
+	 * Called from LaunchEngineLoop via SetEngineStartupModuleLoadingComplete after plugins are loaded, used to scan
+	 * classes that were loaded by plugins, and enable some global multithreaded access.
+	 */
 	void OnInitialPluginLoadingComplete();
 
 	/** Shared helper for Scan*Synchronous function */
diff --git a/Engine/Source/Runtime/AssetRegistry/Private/AssetRegistryImpl.h b/Engine/Source/Runtime/AssetRegistry/Private/AssetRegistryImpl.h
--- a/Engine/Source/Runtime/AssetRegistry/Private/AssetRegistryImpl.h
+++ b/Engine/Source/Runtime/AssetRegistry/Private/AssetRegistryImpl.h
@@ -31,6 +31,7 @@ namespace UE::AssetRegistry::Premade { enum class ELoadResult : uint8; }
 namespace UE::AssetRegistry::Premade { struct FAsyncConsumer; }
 namespace UE::AssetRegistry::Impl { struct FInitializeContext; }
 namespace UE::AssetRegistry::Impl { struct FScanPathContext; }
+namespace UE::AssetRegistry::Impl { struct FTickContext; }
 namespace UE::AssetRegistry::Impl { struct FClassInheritanceContext; }
 namespace UE::AssetRegistry { template<typename TScopeLockType> class TRWScopeLockWithPriority; }
 namespace UE::AssetRegistry { using FInterfaceReadScopeLock = TRWScopeLockWithPriority<FReadScopeLock>; }
@@ -153,15 +154,32 @@ namespace Impl
 	public:
 		typedef TFunction<bool(void)> ShouldExitEarlyCallbackType;
 
+		FInterruptionContext() = default;
 		explicit FInterruptionContext(double InTickStartTime, double InMaxRunningTime)
 			: TickStartTime(InTickStartTime)
-			, MaxRunningTime(InMaxRunningTime) {}
+			, MaxRunningTime(InMaxRunningTime)
+		{
+		}
 
 		FInterruptionContext(double InTickStartTime, double InMaxRunningTime, ShouldExitEarlyCallbackType& Callback)
 			: TickStartTime(InTickStartTime)
 			, MaxRunningTime(InMaxRunningTime)
 			, EarlyExitCallback(Callback) {}
 
+		void SetEarlyExitCallback(const ShouldExitEarlyCallbackType& InCallback)
+		{
+			EarlyExitCallback = InCallback;
+		}
+		void SetUnlimitedTickTime()
+		{
+			TickStartTime = -1.;
+			MaxRunningTime = -1.;
+		}
+		void SetLimitedTickTime(double InTickStartTime, double InMaxRunningTime)
+		{
+			TickStartTime = InTickStartTime;
+			MaxRunningTime = InMaxRunningTime;
+		}
 		double GetTickStartTime() const { return TickStartTime; }
 		bool IsTimeSlicingEnabled() const { return TickStartTime > 0; }
 		bool WasInterrupted() const { return OutInterrupted; }
@@ -170,9 +188,9 @@ namespace Impl
 
 	private:
 		// A negative value disables time slicing
-		const double TickStartTime;
+		double TickStartTime = -1.;
 		// The maximum time should allow before interruption. If TickStartTime is negative, this is ignored
-		const double MaxRunningTime;
+		double MaxRunningTime = -1.;
 		// If provided, this is always checked
 		ShouldExitEarlyCallbackType EarlyExitCallback;
 		// True if we ran out of time, EarlyExitCallback returned true, or RequestEarlyExit() was called
@@ -254,9 +272,12 @@ public:
 	// Other helper functions called by UAssetRegistryImpl
 
 
-	/** Update cached values that need to be reread when IPluginManager reaches PostEngineInit phase */
-	void OnPostEngineInit(bool bPhaseSuccessful);
-	/** Update cached values that need to be reread from the Engine's PostEngineInit callback */
+	/**
+	 * Update cached values about classes, and enable some global multithreaded access,
+	 * when all startup plugins finish loading.
+	 */
+	void OnPluginLoadingComplete(bool bPhaseSuccessful);
+	/** Update cached values about native classes after they have possibly changed. */
 	void RefreshNativeClasses();
 
 	/** Enumerate assets in the State, filtering by filter and package not in PackagesToSkip */
@@ -269,14 +290,8 @@ public:
 	void WaitForGathererIdleIfSynchronous();
 	/** Waits for the gatherer to be idle. */
 	void WaitForGathererIdle(float TimeoutSeconds);
-	/** Callback types for TickGatherer */
-	typedef TFunctionRef<void(const TMultiMap<FName, FAssetData*>&)> FAssetsFoundCallback;
-	typedef TFunctionRef<void(const TRingBuffer<FName>&)> FVerseFilesFoundCallback;
 	/** Consume any results from the gatherer and return its status */
-	Impl::EGatherStatus TickGatherer(Impl::FEventContext& EventContext,
-		Impl::FClassInheritanceContext& InheritanceContext, Impl::FInterruptionContext& InOutInterruptionContext,
-		TOptional<FAssetsFoundCallback> AssetsFoundCallback = TOptional<FAssetsFoundCallback>(),
-		TOptional<FVerseFilesFoundCallback> VerseFilesFoundCallback = TOptional<FVerseFilesFoundCallback>());
+	Impl::EGatherStatus TickGatherer(Impl::FTickContext& TickContext);
 
 	/** Send a log message with the search statistics. 
 	 *  StartTime is used to report wall clock search time in the case of background scan
@@ -410,8 +425,10 @@ private:
 	 * If TickStartTime is < 0, the entire list of gathered assets will be cached. Also used in sychronous searches
 	 * The DeferredResults array contains assets that were not ready for processing due to required UClass's not being loaded
 	 */
-	void AssetSearchDataGathered(Impl::FEventContext& EventContext, TMultiMap<FName, TUniquePtr<FAssetData>>& AssetResults, 
-		TMultiMap<FName, TUniquePtr<FAssetData>>& OutDeferredResults, Impl::FInterruptionContext& InOutInterruptionContext);
+	void AssetSearchDataGathered(Impl::FEventContext& EventContext,
+		TMultiMap<FName, TUniquePtr<FAssetData>>& AssetResults, 
+		TMultiMap<FName, TUniquePtr<FAssetData>>& OutDeferredResults,\
+		Impl::FInterruptionContext& InOutInterruptionContext, bool bCanAccessCoreRedirects);
 	/** Validate assets gathered from disk before adding them to the AssetRegistry. */
 	bool ShouldSkipGatheredAsset(FAssetData& AssetData);
 
@@ -454,7 +471,7 @@ private:
 	 * @param AssetData Existing asset data
 	 * @return Returns false if the required parent UClass is not yet available
 	 */
-	bool TryPostLoadAssetRegistryTags(FAssetData* AssetData);
+	bool TryPostLoadAssetRegistryTags(FAssetData* AssetData, bool bCanAccessCoreRedirects);
 
 	/** Update Redirect collector with redirects loaded from asset registry */
 	void UpdateRedirectCollector();
@@ -586,6 +603,7 @@ private:
 	bool bVerboseLogging;
 
 	bool bForceCompletionEvenIfPostLoadsFail = false;
+	bool bProcessedAnyAssetsAfterRetryDeferred = true;
 
 	/** List of all class names derived from Blueprint (including Blueprint itself) */
 	TSet<FTopLevelAssetPath> ClassGeneratorNames;
@@ -662,7 +680,6 @@ struct FEventContext
 	TArray<FString> BlockedFiles;
 	bool bFileLoadedEventBroadcast = false;
 	bool bScanStartedEventBroadcast = false;
-	bool bScanEndedEventBroadcast = false;
 	bool bHasSentFileLoadedEventBroadcast = false;
 
 	/** Remove all stored events */
@@ -720,6 +737,27 @@ struct FScanPathContext
 	EGatherStatus Status = EGatherStatus::TickActiveGatherActive;
 };
 
+/** Callback types for FTickContext */
+typedef TFunctionRef<void(const TMultiMap<FName, FAssetData*>&)> FAssetsFoundCallback;
+typedef TFunctionRef<void(const TRingBuffer<FName>&)> FVerseFilesFoundCallback;
+
+/** Input and output variables for the TickGatherer function. */
+struct FTickContext
+{
+	FTickContext(Impl::FEventContext& InEventContext, Impl::FClassInheritanceContext& InInheritanceContext)
+		: EventContext(InEventContext)
+		, InheritanceContext(InInheritanceContext)
+	{
+	}
+	Impl::FInterruptionContext InterruptionContext;
+	Impl::FEventContext& EventContext;
+	Impl::FClassInheritanceContext& InheritanceContext;
+	TOptional<FAssetsFoundCallback> AssetsFoundCallback;
+	TOptional<FVerseFilesFoundCallback> VerseFilesFoundCallback;
+	bool bHandleCompletion = false;
+	bool bHandleDeferred = false;
+};
+
 }
 
 namespace Utils
