AssetRegistry: Fix UnrealEditor.exe -server to do a synchronous search rather than asynchronous, like -game does. Asynchronous is too expensive for game or dedicatedserver, and additionally causes the AssetManager to misbehave because it assumes asynchronous is only GIsEditor.
#rnx

diff --git a/Engine/Source/Runtime/AssetRegistry/Private/AssetRegistry.cpp b/Engine/Source/Runtime/AssetRegistry/Private/AssetRegistry.cpp
--- a/Engine/Source/Runtime/AssetRegistry/Private/AssetRegistry.cpp
+++ b/Engine/Source/Runtime/AssetRegistry/Private/AssetRegistry.cpp
@@ -1266,7 +1266,7 @@ void FAssetRegistryImpl::InitRedirectors(Impl::FEventContext& EventContext,
 		}
 
 		// if we are -game in editor build we might need to initialize the asset registry manually for this plugin
-		if (!FPlatformProperties::RequiresCookedData() && IsRunningGame())
+		if (!FPlatformProperties::RequiresCookedData() && (IsRunningGame() || IsRunningDedicatedServer()))
 		{
 			TArray<FString> PathsToSearch;
 			
@@ -1799,7 +1799,7 @@ bool FAssetRegistryImpl::TryConstructGathererIfNeeded()
 		EngineIni->GetArray(TEXT("AssetRegistry"), TEXT("BlacklistContentSubPathScanFilters"), ContentSubPathsDenyList);
 	}
 
-	bool bAsyncGatherEnabled = !IsRunningGame();
+	bool bAsyncGatherEnabled = !IsRunningGame() && !IsRunningDedicatedServer();
 	GlobalGatherer = MakeUnique<FAssetDataGatherer>(PathsDenyList, ContentSubPathsDenyList, bAsyncGatherEnabled, *this);
 
 	// Read script packages if all initial plugins have been loaded, otherwise do nothing; we wait for the callback.
