AssetRegistry: update comment to match code change.
#rnx

diff --git a/Engine/Source/Runtime/AssetRegistry/Private/AssetRegistry.cpp b/Engine/Source/Runtime/AssetRegistry/Private/AssetRegistry.cpp
--- a/Engine/Source/Runtime/AssetRegistry/Private/AssetRegistry.cpp
+++ b/Engine/Source/Runtime/AssetRegistry/Private/AssetRegistry.cpp
@@ -1265,7 +1265,7 @@ void FAssetRegistryImpl::InitRedirectors(Impl::FEventContext& EventContext,
 			continue;
 		}
 
-		// if we are -game in editor build we might need to initialize the asset registry manually for this plugin
+		// if we are -game or -server in editor build we might need to initialize the asset registry manually for this plugin
 		if (!FPlatformProperties::RequiresCookedData() && (IsRunningGame() || IsRunningDedicatedServer()))
 		{
 			TArray<FString> PathsToSearch;
