Reduce allocations when ticking FAssetDataGatherer

#rnx

- Use the stringview version of GetExtension to avoid creating a temp FString.

diff --git a/Engine/Source/Runtime/AssetRegistry/Private/AssetDataGatherer.cpp b/Engine/Source/Runtime/AssetRegistry/Private/AssetDataGatherer.cpp
--- a/Engine/Source/Runtime/AssetRegistry/Private/AssetDataGatherer.cpp
+++ b/Engine/Source/Runtime/AssetRegistry/Private/AssetDataGatherer.cpp
@@ -4160,7 +4160,7 @@ FAssetDataGatherer::ETickResult FAssetDataGatherer::TickInternal(double& TickSta
 				DiskCachedAssetData = nullptr;
 			}
 			else if ((!DiskCachedAssetData->DependencyData.PackageName.IsEqual(PackageName, ENameCase::CaseSensitive) && DiskCachedAssetData->DependencyData.PackageName != NAME_None) ||
-				DiskCachedAssetData->Extension != FName(*FPaths::GetExtension(AssetFileData.LocalAbsPath)))
+				DiskCachedAssetData->Extension != FName(FPathViews::GetExtension(AssetFileData.LocalAbsPath)))
 			{
 				UE_LOG(LogAssetRegistry, Display, TEXT("Cached dependency data for package '%s' is invalid. Discarding cached data."), *PackageName.ToString());
 				DiskCachedAssetData = nullptr;
@@ -4212,7 +4212,7 @@ FAssetDataGatherer::ETickResult FAssetDataGatherer::TickInternal(double& TickSta
 		else
 		{
 			// Not found in cache (or stale) - schedule to be read from disk
-			ReadContexts.Emplace(PackageName, FName(*FPaths::GetExtension(AssetFileData.LocalAbsPath)), AssetFileData);
+			ReadContexts.Emplace(PackageName, FName(FPathViews::GetExtension(AssetFileData.LocalAbsPath)), AssetFileData);
 		}
 	}
 
