Check for sandbox paths when checking for correct capitalization.  When operating in a sandbox we must get the LowerLevel platform file and query before we obtain the filename on disk.

#rnx

diff --git a/Engine/Source/Runtime/AssetRegistry/Private/AssetDataGatherer.cpp b/Engine/Source/Runtime/AssetRegistry/Private/AssetDataGatherer.cpp
--- a/Engine/Source/Runtime/AssetRegistry/Private/AssetDataGatherer.cpp
+++ b/Engine/Source/Runtime/AssetRegistry/Private/AssetDataGatherer.cpp
@@ -2346,6 +2346,24 @@ FDateTime FPathExistence::GetModificationTime()
 	return ModificationTime;
 }
 
+FString GetCorrectedCapitalization(const TCHAR* Filename)
+{
+	// If we are in a sandbox environment we have to ask the LowerLevel platform file for the path otherwise we
+	// may get the sandbox path when we check for capitalization.
+	//
+	if (IFileManager::Get().IsSandboxEnabled())
+	{
+		IPlatformFile& CurrentPlatformFile = FPlatformFileManager::Get().GetPlatformFile();
+		IPlatformFile* LowerLevel = CurrentPlatformFile.GetLowerLevel();
+		if (ensure(LowerLevel))
+		{
+			return LowerLevel->GetFilenameOnDisk(Filename);
+		}
+	}
+
+	return IFileManager::Get().GetFilenameOnDisk(Filename);
+}
+
 void FPathExistence::LoadExistenceData()
 {
 	if (bHasExistenceData)
@@ -2355,7 +2373,7 @@ void FPathExistence::LoadExistenceData()
 	FFileStatData StatData = IFileManager::Get().GetStatData(*LocalAbsPath);
 	if (StatData.bIsValid)
 	{
-		FString CorrectedCapitalization = IFileManager::Get().GetFilenameOnDisk(*LocalAbsPath);
+		FString CorrectedCapitalization = UE::AssetDataGather::Private::GetCorrectedCapitalization(*LocalAbsPath);
 		if (LocalAbsPath == CorrectedCapitalization)
 		{
 			LocalAbsPath = MoveTemp(CorrectedCapitalization);
@@ -2377,7 +2395,7 @@ void FPathExistence::LoadExistenceData()
 		StatData = IFileManager::Get().GetStatData(*ParentPath);
 		if (StatData.bIsValid && StatData.bIsDirectory)
 		{
-			FString CorrectedCapitalization = IFileManager::Get().GetFilenameOnDisk(*ParentPath);
+			FString CorrectedCapitalization = UE::AssetDataGather::Private::GetCorrectedCapitalization(*LocalAbsPath);
 			CorrectedCapitalization = FPaths::Combine(CorrectedCapitalization, BaseName) +
 				(!Extension.IsEmpty() ? TEXT(".") : TEXT("")) + Extension;
 			if (LocalAbsPath == CorrectedCapitalization)
