Asset Registry - Fixed issue with GetRedirectedObjectPath() potentially not resolving redirectors while a full scan is still in progress

diff --git a/Engine/Source/Runtime/AssetRegistry/Private/AssetRegistry.cpp b/Engine/Source/Runtime/AssetRegistry/Private/AssetRegistry.cpp
--- a/Engine/Source/Runtime/AssetRegistry/Private/AssetRegistry.cpp
+++ b/Engine/Source/Runtime/AssetRegistry/Private/AssetRegistry.cpp
@@ -3189,10 +3189,14 @@ bool UAssetRegistryImpl::DoesPackageExistOnDisk(FName PackageName, FString* OutC
 
 FSoftObjectPath UAssetRegistryImpl::GetRedirectedObjectPath(const FSoftObjectPath& ObjectPath)
 {
-	if (IsSearchAllAssets())
+	// Fast path, if a full registry scan was triggered & has completed
+	// In that case, we can skip further scanning while looking for a redirected path
 	{
 		UE::AssetRegistry::FInterfaceReadScopeLock InterfaceScopeLock(InterfaceLock);
-		return GuardedData.GetRedirectedObjectPath(ObjectPath, nullptr, nullptr, /*bNeedsScanning*/ false);
+		if (GuardedData.IsSearchAllAssets() && GuardedData.IsInitialSearchCompleted())
+		{
+			return GuardedData.GetRedirectedObjectPath(ObjectPath, nullptr, nullptr, /*bNeedsScanning*/ false);
+		}
 	}
 
 	FSoftObjectPath RedirectedObjectPath;
