Apply package relocations to the FAssetDataTagMap when reading them
[RN] Asset Registry tag values are now updated when an asset is moved on disk

diff --git a/Engine/Source/Runtime/AssetRegistry/Private/PackageReader.cpp b/Engine/Source/Runtime/AssetRegistry/Private/PackageReader.cpp
--- a/Engine/Source/Runtime/AssetRegistry/Private/PackageReader.cpp
+++ b/Engine/Source/Runtime/AssetRegistry/Private/PackageReader.cpp
@@ -39,6 +39,51 @@ const TCHAR* LexToString(FPackageReader::EOpenPackageResult Result)
 	return TEXT("UNKNOWN");
 }
 
+namespace UE::Private
+{
+static void ApplyRelocationToTagsAndValues(FAssetDataTagMap& TagsAndValues, UE::Package::Relocation::Private::FPackageRelocationContext RelocationArgs)
+{
+	using namespace UE::Package::Relocation::Private;
+	for (TPair<FName, FString>& Iter : TagsAndValues)
+	{
+		FString& Value = Iter.Value;
+		if (Value.IsEmpty())
+		{
+			continue;
+		}
+
+		if (FPackageName::IsValidObjectPath(Value))
+		{
+			FNameBuilder RelocatedPackageName;
+			if (TryRelocateReference(RelocationArgs, Value, RelocatedPackageName)
+				&& RelocatedPackageName.Len() != 0)
+			{
+				Value = *RelocatedPackageName;
+			}
+		}
+		else if (
+			FStringView ClassName, ObjectPath;
+			FPackageName::ParseExportTextPath(Value, &ClassName, &ObjectPath))
+		{
+			FNameBuilder RelocatedClassName;
+			const bool bHasNewClassName = TryRelocateReference(RelocationArgs, ClassName, RelocatedClassName)
+				&& RelocatedClassName.Len() != 0;
+			FNameBuilder RelocatedObjectName;
+			const bool bHasNewObjectName = TryRelocateReference(RelocationArgs, ObjectPath, RelocatedObjectName)
+				&& RelocatedObjectName.Len() != 0;
+
+			// want to be careful with all these string views around:
+			FString NewValue = FString::Format(TEXT("{0}'{1}'"),
+				{ bHasNewClassName ? *RelocatedClassName : ClassName,
+				bHasNewObjectName ? *RelocatedObjectName : ObjectPath });
+			Value = MoveTemp(NewValue);
+			// validate value:
+			ensure(FPackageName::ParseExportTextPath(Value, &ClassName, &ObjectPath));
+		}
+	}
+}
+}
+
 FPackageReader::FPackageReader()
 	: Loader(nullptr)
 	, PackageFileSize(0)
@@ -1670,6 +1715,10 @@ namespace UE::AssetRegistry
 		}
 
 		OutDependencyDataOffset = DeserializePackageData.DependencyDataOffset;
+		
+		// support package relocation:
+		UE::Package::Relocation::Private::FPackageRelocationContext RelocationArgs;
+		const bool bIsRelocated = UE::Package::Relocation::Private::ShouldApplyRelocation(PackageFileSummary, PackageName, RelocationArgs);
 
 		// Worlds that were saved before they were marked public do not have asset data so we will synthesize it here to make sure we see all legacy umaps
 		// We will also do this for maps saved after they were marked public but no asset data was saved for some reason. A bug caused this to happen for some maps.
@@ -1710,6 +1759,11 @@ namespace UE::AssetRegistry
 				}
 			}
 
+			if(bIsRelocated)
+			{
+				UE::Private::ApplyRelocationToTagsAndValues(TagsAndValues, RelocationArgs);
+			}
+
 			// Before worlds were RF_Public, other non-public assets were added to the asset data table in map packages.
 			// Here we simply skip over them
 			if (bIsMapPackage && PackageFileSummary.GetFileVersionUE() < VER_UE4_PUBLIC_WORLDS)
