AssetRegistry PackageReader: Add more feedback for the version mismatch error when a package has a newer custom version.
#rnx

diff --git a/Engine/Source/Runtime/AssetRegistry/Private/PackageReader.cpp b/Engine/Source/Runtime/AssetRegistry/Private/PackageReader.cpp
--- a/Engine/Source/Runtime/AssetRegistry/Private/PackageReader.cpp
+++ b/Engine/Source/Runtime/AssetRegistry/Private/PackageReader.cpp
@@ -242,9 +242,20 @@ bool FPackageReader::OpenPackageFile(EOpenPackageResult& OutErrorCode)
 		{
 			if (IsEnforcePackageCompatibleVersionCheck())
 			{
+				int32 PackageVersion = -1;
+				int32 HeadCodeVersion = -1;
+				if (const FCustomVersion* PackagePtr
+					= PackageFileSummary.GetCustomVersionContainer().GetVersion(Diff.Version->Key))
+				{
+					PackageVersion = PackagePtr->Version;
+				}
+				if (TOptional<FCustomVersion> CurrentPtr = FCurrentCustomVersions::Get(Diff.Version->Key))
+				{
+					HeadCodeVersion = CurrentPtr->Version;
+				}
 				UE_LOG(LogAssetRegistry, Error,
-					TEXT("Package is unloadable: %s. Reason: Custom version is too new; the package has newer custom version of %s."),
-					*PackageFilename, *Diff.Version->GetFriendlyName().ToString());
+					TEXT("Package is unloadable: %s. Reason: Custom version is too new; the package has newer custom version of %s: Package: %d, HeadCode: %d."),
+					*PackageFilename, *Diff.Version->GetFriendlyName().ToString(), PackageVersion, HeadCodeVersion);
 				OutErrorCode = EOpenPackageResult::VersionTooNew;
 			}
 		}
