Removed AssetRegistry's ScanPathsSynchronous warning "<package> does not exist, will not scan" because it rarely indicates an error, and is often used to check whether a package exists in the state before the scan has finished.

diff --git a/Engine/Source/Runtime/AssetRegistry/Private/AssetRegistry.cpp b/Engine/Source/Runtime/AssetRegistry/Private/AssetRegistry.cpp
--- a/Engine/Source/Runtime/AssetRegistry/Private/AssetRegistry.cpp
+++ b/Engine/Source/Runtime/AssetRegistry/Private/AssetRegistry.cpp
@@ -5718,11 +5718,9 @@ FScanPathContext::FScanPathContext(FEventContext& InEventContext, FClassInherita
 				if (FPackageName::InternalDoesPackageExistEx(PackagePath, FPackageName::EPackageLocationFilter::Any,
 					false /* bMatchCaseOnDisk */, &PackagePath) == FPackageName::EPackageLocationFilter::None)
 				{
-					if (!bIgnoreInvalidPathWarning)
-					{
-						UE_LOG(LogAssetRegistry, Warning, TEXT("ScanPathsSynchronous: Package %s does not exist, will not scan."), *InFile);
-						bLogCallstack = true;
-					}
+					// Requesting to scan a non-existent package is not a condition we need to warn about, because it rarely indicates an error,
+					// and is often used to check whether a package exists in the state before the scan has finished. Silently ignore it,
+					// even if !bIgnoreInvalidPathWarning.
 					continue;
 				}
 				Extension = PackagePath.GetExtensionString(EPackageSegment::Header);
