AssetRegistry now properly removes redirects from GRedirectCollector when object redirector assets are removed. Added additional conditional timing scopes to registry to help identify stalls.

diff --git a/Engine/Source/Runtime/AssetRegistry/Private/AssetDataGatherer.cpp b/Engine/Source/Runtime/AssetRegistry/Private/AssetDataGatherer.cpp
--- a/Engine/Source/Runtime/AssetRegistry/Private/AssetDataGatherer.cpp
+++ b/Engine/Source/Runtime/AssetRegistry/Private/AssetDataGatherer.cpp
@@ -3600,6 +3600,7 @@ FAssetDataGatherer::FAssetDataGatherer(const TArray<FString>& InLongPackageNames
 	, bIsIdle(false)
 	, bFirstTickAfterIdle(true)
 	, bFinishedInitialDiscovery(false)
+	, bIsInitialSearchCompleted(false)
 	, WaitBatchCount(-1)
 	, LastCacheSaveNumUncachedAssetFiles(0)
 	, CacheInUseCount(0)
@@ -3667,6 +3668,8 @@ void FAssetDataGatherer::OnInitialSearchCompleted()
 	{
 		Discovery->OnInitialSearchCompleted();
 	}
+
+	bIsInitialSearchCompleted.store(true, std::memory_order_relaxed);
 }
 
 void FAssetDataGatherer::StartAsync()
@@ -3719,6 +3722,8 @@ uint32 FAssetDataGatherer::Run()
 			if ((Status != UE::AssetRegistry::Impl::EGatherStatus::TickActiveGatherActive)
 				&& (Status != UE::AssetRegistry::Impl::EGatherStatus::TickActiveGatherIdle))
 			{
+				// Consider creating a TRACE_CPUPROFILER_EVENT_SCOPE_STR_CONDITIONAL
+				TRACE_CPUPROFILER_EVENT_SCOPE_TEXT_CONDITIONAL("FAssetDataGatherer Sleep", !bIsInitialSearchCompleted.load(std::memory_order_relaxed));
 				FPlatformProcess::Sleep(bLocalIdle ? IdleSleepTime : PausedSleepTime);
 			}
 		}
diff --git a/Engine/Source/Runtime/AssetRegistry/Private/AssetDataGatherer.h b/Engine/Source/Runtime/AssetRegistry/Private/AssetDataGatherer.h
--- a/Engine/Source/Runtime/AssetRegistry/Private/AssetDataGatherer.h
+++ b/Engine/Source/Runtime/AssetRegistry/Private/AssetDataGatherer.h
@@ -470,6 +470,8 @@ private:
 	bool bFirstTickAfterIdle;
 	/** True if we have finished discovering our first wave of files, to report metrics for that most-important wave. */
 	bool bFinishedInitialDiscovery;
+	/** True if OnInitialSearchCompleted has been called. */
+	std::atomic<bool> bIsInitialSearchCompleted;
 
 	// Variable section for variables that are read/writable only within TickLock.
 
diff --git a/Engine/Source/Runtime/AssetRegistry/Private/AssetRegistry.cpp b/Engine/Source/Runtime/AssetRegistry/Private/AssetRegistry.cpp
--- a/Engine/Source/Runtime/AssetRegistry/Private/AssetRegistry.cpp
+++ b/Engine/Source/Runtime/AssetRegistry/Private/AssetRegistry.cpp
@@ -4614,6 +4614,9 @@ bool FAssetRegistryImpl::ClassRequiresGameThreadProcessing(const UClass* Class)
 
 Impl::EGatherStatus FAssetRegistryImpl::TickGatherer(Impl::FTickContext& TickContext)
 {
+	// Consider creating a TRACE_CPUPROFILER_EVENT_SCOPE_STR_CONDITIONAL
+	TRACE_CPUPROFILER_EVENT_SCOPE_TEXT_CONDITIONAL("FAssetRegistryImpl::TickGatherer", !bInitialSearchCompleted);
+
 	using namespace UE::AssetRegistry::Impl;
 	FEventContext& EventContext = TickContext.EventContext;
 	FInterruptionContext& InOutInterruptionContext = TickContext.InterruptionContext;
@@ -6910,6 +6913,14 @@ bool FAssetRegistryImpl::RemoveAssetData(Impl::FEventContext& EventContext, FAss
 			EventContext.AssetEvents.Emplace(*AssetData, Impl::FEventContext::EEvent::Removed);
 		}
 
+		// Make sure to consider redirections!
+#if WITH_EDITOR
+		if (AssetData->IsRedirector())
+		{
+			GRedirectCollector.RemoveAssetPathRedirection(AssetData->GetSoftObjectPath());
+		}
+#endif
+
 		// Remove from the class map if removing a blueprint
 		if (ClassGeneratorNames.Contains(AssetData->AssetClassPath))
 		{
