Added overload on SoftObjectPath for no SubPath strings.
Changed call sites to use new static name creation methods. (removing explicit FString construction)

diff --git a/Engine/Plugins/Experimental/PythonScriptPlugin/Source/PythonScriptPlugin/Private/PyWrapperObject.cpp b/Engine/Plugins/Experimental/PythonScriptPlugin/Source/PythonScriptPlugin/Private/PyWrapperObject.cpp
--- a/Engine/Plugins/Experimental/PythonScriptPlugin/Source/PythonScriptPlugin/Private/PyWrapperObject.cpp
+++ b/Engine/Plugins/Experimental/PythonScriptPlugin/Source/PythonScriptPlugin/Private/PyWrapperObject.cpp
@@ -1947,7 +1947,7 @@ void UPythonGeneratedClass::PostRename(UObject* OldOuter, const FName OldName)
 
 	if (PyType)
 	{
-		FPyWrapperTypeRegistry::Get().UnregisterWrappedClassType(FSoftObjectPath(OldOuter->GetFName(), OldName, FString()), PyType, false);
+		FPyWrapperTypeRegistry::Get().UnregisterWrappedClassType(FSoftObjectPath::ConstructFromPackageAsset(OldOuter->GetFName(), OldName), PyType, false);
 		FPyWrapperTypeRegistry::Get().RegisterWrappedClassType(this, PyType, false);
 	}
 }
diff --git a/Engine/Plugins/Experimental/PythonScriptPlugin/Source/PythonScriptPlugin/Private/PyWrapperStruct.cpp b/Engine/Plugins/Experimental/PythonScriptPlugin/Source/PythonScriptPlugin/Private/PyWrapperStruct.cpp
--- a/Engine/Plugins/Experimental/PythonScriptPlugin/Source/PythonScriptPlugin/Private/PyWrapperStruct.cpp
+++ b/Engine/Plugins/Experimental/PythonScriptPlugin/Source/PythonScriptPlugin/Private/PyWrapperStruct.cpp
@@ -1632,7 +1632,7 @@ void UPythonGeneratedStruct::PostRename(UObject* OldOuter, const FName OldName)
 
 	if (PyType)
 	{
-		FPyWrapperTypeRegistry::Get().UnregisterWrappedStructType(FSoftObjectPath(OldOuter->GetFName(), OldName, FString()), PyType, false);
+		FPyWrapperTypeRegistry::Get().UnregisterWrappedStructType(FSoftObjectPath::ConstructFromPackageAsset(OldOuter->GetFName(), OldName), PyType, false);
 		FPyWrapperTypeRegistry::Get().RegisterWrappedStructType(this, PyType, false);
 	}
 }
diff --git a/Engine/Source/Runtime/AssetRegistry/Private/AssetRegistry.cpp b/Engine/Source/Runtime/AssetRegistry/Private/AssetRegistry.cpp
--- a/Engine/Source/Runtime/AssetRegistry/Private/AssetRegistry.cpp
+++ b/Engine/Source/Runtime/AssetRegistry/Private/AssetRegistry.cpp
@@ -2521,7 +2521,7 @@ void EnumerateMemoryAssetsHelper(const FARCompiledFilter& InFilter, TSet<FName>&
 		}
 
 		// Asset Path
-		FilterData.ObjectPath = FSoftObjectPath(Obj);
+		FilterData.ObjectPath = FSoftObjectPath::ConstructFromObject(Obj);
 		if (InFilter.SoftObjectPaths.Num() > 0)
 		{
 			if (!InFilter.SoftObjectPaths.Contains(FilterData.ObjectPath))
@@ -2668,7 +2668,7 @@ void EnumerateMemoryAssets(const FARCompiledFilter& InFilter, TSet<FName>& OutPa
 				// Our contract for on-disk versus in-memory tags is that in-memory tags override on-disk tags, but we
 				// keep any on-disk tags that do not exist in the in-memory tags because they may be extended tags.
 				UE::AssetRegistry::FInterfaceReadScopeLock InterfaceScopeLock(InterfaceLock);
-				const FAssetData* OnDiskAssetData = GuardedDataState.GetAssetByObjectPath(FSoftObjectPath(Object));
+				const FAssetData* OnDiskAssetData = GuardedDataState.GetAssetByObjectPath(FSoftObjectPath::ConstructFromObject(Object));
 				if (OnDiskAssetData)
 				{
 					TOptional<FAssetDataTagMap> ModifiedTags = Utils::AddNonOverlappingTags(PartialAssetData, *OnDiskAssetData);
@@ -3265,7 +3265,7 @@ FSoftObjectPath FAssetRegistryImpl::GetRedirectedObjectPath(const FSoftObjectPat
 	UObject* Asset = ObjectPath.ResolveObject();
 	if (Asset)
 	{
-		RedirectedPath = FSoftObjectPath(Asset);
+		RedirectedPath = FSoftObjectPath::ConstructFromObject(Asset);
 		UObjectRedirector* Redirector = Cast<UObjectRedirector>(Asset);
 		if (!Redirector || !Redirector->DestinationObject)
 		{
@@ -3284,7 +3284,7 @@ FSoftObjectPath FAssetRegistryImpl::GetRedirectedObjectPath(const FSoftObjectPat
 		{
 			// If we found no Asset because it is a subobject, then look for its toplevelobject's Asset
 			SubPathString = RedirectedPath.GetSubPathString();
-			RedirectedPath = FSoftObjectPath(RedirectedPath.GetAssetPath(), FString());
+			RedirectedPath = FSoftObjectPath::ConstructFromAssetPath(RedirectedPath.GetAssetPath());
 			AssetData = State.GetAssetByObjectPath(RedirectedPath);
 		}
 		return AssetData;
@@ -4249,7 +4249,7 @@ void UAssetRegistryImpl::AssetDeleted(UObject* DeletedAsset)
 
 		{
 			// Need to remove from GRedirectCollector
-			GRedirectCollector.RemoveAssetPathRedirection(FSoftObjectPath(DeletedAsset));
+			GRedirectCollector.RemoveAssetPathRedirection(FSoftObjectPath::ConstructFromObject(DeletedAsset));
 		}
 #endif
 
@@ -6117,7 +6117,7 @@ bool FAssetRegistryImpl::TryPostLoadAssetRegistryTags(FAssetData* AssetData)
 		{
 			FTopLevelAssetPath LastAssetClassPath = AssetClassPath;
 			// Maybe it's a redirector
-			FSoftObjectPath RedirectedPath = GRedirectCollector.GetAssetPathRedirection(FSoftObjectPath(AssetClassPath));
+			FSoftObjectPath RedirectedPath = GRedirectCollector.GetAssetPathRedirection(FSoftObjectPath::ConstructFromAssetPath(AssetClassPath));
 			if (RedirectedPath.IsValid())
 			{
 				AssetClassPath = RedirectedPath.GetAssetPath();
@@ -8182,7 +8182,7 @@ void UAssetRegistryImpl::OnGetExtraObjectTags(FAssetRegistryTagsContext Context)
 	if (bAddMetaDataTagsToOnGetExtraObjectTags)
 	{
 		// It is critical that bIncludeOnlyOnDiskAssets=true otherwise this will cause an infinite loop
-		const FAssetData AssetData = GetAssetByObjectPath(FSoftObjectPath(Context.GetObject()), /*bIncludeOnlyOnDiskAssets=*/true);
+		const FAssetData AssetData = GetAssetByObjectPath(FSoftObjectPath::ConstructFromObject(Context.GetObject()), /*bIncludeOnlyOnDiskAssets=*/true);
 		// Adding metadata tags from disk is only necessary for cooked assets; uncooked assets still have the metadata and add them elsewhere
 		// in UObject::GetAssetRegistryTags. Adding the tags from disk into uncooked assets would make the tags impossible to remove when
 		// the uncooked assets are resaved.
diff --git a/Engine/Source/Runtime/AssetRegistry/Private/PackageReader.cpp b/Engine/Source/Runtime/AssetRegistry/Private/PackageReader.cpp
--- a/Engine/Source/Runtime/AssetRegistry/Private/PackageReader.cpp
+++ b/Engine/Source/Runtime/AssetRegistry/Private/PackageReader.cpp
@@ -549,7 +549,7 @@ bool FPackageReader::ReadLinkerObjects(TMap<FSoftObjectPath, FObjectData>& OutEx
 	{
 		FObjectData Data;
 		const FObjectImport& Import = ImportMap[Index];
-		Data.ClassPath = FSoftObjectPath(Import.ClassPackage, Import.ClassName, FString());
+		Data.ClassPath = FSoftObjectPath::ConstructFromPackageAsset(Import.ClassPackage, Import.ClassName);
 		Data.bUsedInGame = ImportUsedInGame[Index];
 		OutImports.Add(ImportsPaths[Index], MoveTemp(Data));
 	}
@@ -1490,16 +1490,16 @@ void FPackageReader::ConvertLinkerTableToPaths(FName InPackageName, TArray<FObje
 			{
 				if (Index.IsExport())
 				{
-					Result = FSoftObjectPath(InPackageName, ObjectName, FString());
+					Result = FSoftObjectPath::ConstructFromPackageAsset(InPackageName, ObjectName);
 				}
 				else
 				{
-					Result = FSoftObjectPath(ObjectName, NAME_None, FString());
+					Result = FSoftObjectPath::ConstructFromPackageAsset(ObjectName, NAME_None);
 				}
 			}
 			else if (ParentPath.GetAssetFName().IsNone())
 			{
-				Result = FSoftObjectPath(ParentPath.GetLongPackageFName(), ObjectName, FString());
+				Result = FSoftObjectPath::ConstructFromPackageAsset(ParentPath.GetLongPackageFName(), ObjectName);
 			}
 			else if (ParentPath.GetSubPathString().IsEmpty())
 			{
diff --git a/Engine/Source/Runtime/CoreUObject/Private/AssetRegistry/AssetData.cpp b/Engine/Source/Runtime/CoreUObject/Private/AssetRegistry/AssetData.cpp
--- a/Engine/Source/Runtime/CoreUObject/Private/AssetRegistry/AssetData.cpp
+++ b/Engine/Source/Runtime/CoreUObject/Private/AssetRegistry/AssetData.cpp
@@ -384,7 +384,7 @@ FSoftObjectPath FAssetData::GetSoftObjectPath() const
 {
 	if (IsTopLevelAsset())
 	{
-		return FSoftObjectPath(PackageName, AssetName, FString());
+		return FSoftObjectPath::ConstructFromPackageAsset(PackageName, AssetName);
 	}
 	else
 	{
diff --git a/Engine/Source/Runtime/CoreUObject/Private/Misc/RedirectCollector.cpp b/Engine/Source/Runtime/CoreUObject/Private/Misc/RedirectCollector.cpp
--- a/Engine/Source/Runtime/CoreUObject/Private/Misc/RedirectCollector.cpp
+++ b/Engine/Source/Runtime/CoreUObject/Private/Misc/RedirectCollector.cpp
@@ -114,7 +114,7 @@ void FRedirectCollector::CollectSavedSoftPackageReferences(FName ReferencingPack
 	TArray<FSoftObjectPathProperty, TInlineAllocator<4>> SoftObjectPathArray;
 	Algo::Transform(PackageNames, SoftObjectPathArray, [ReferencingPackage, bEditorOnlyReferences](const FName& PackageName)
 		{
-			return FSoftObjectPathProperty(FSoftObjectPath(PackageName, NAME_None, FString()), NAME_None, bEditorOnlyReferences);
+			return FSoftObjectPathProperty(FSoftObjectPath::ConstructFromPackageAsset(PackageName, NAME_None), NAME_None, bEditorOnlyReferences);
 		});
 
 	FScopeLock ScopeLock(&CriticalSection);
diff --git a/Engine/Source/Runtime/CoreUObject/Private/UObject/SoftObjectPath.cpp b/Engine/Source/Runtime/CoreUObject/Private/UObject/SoftObjectPath.cpp
--- a/Engine/Source/Runtime/CoreUObject/Private/UObject/SoftObjectPath.cpp
+++ b/Engine/Source/Runtime/CoreUObject/Private/UObject/SoftObjectPath.cpp
@@ -43,6 +43,13 @@ FSoftObjectPath FSoftObjectPath::ConstructFromAssetPath(FTopLevelAssetPath InAss
 	return FSoftObjectPath(InAssetPath);
 }
 
+FSoftObjectPath FSoftObjectPath::ConstructFromStringPath(FString&& InPath)
+{
+	FSoftObjectPath Tmp;
+	Tmp.SetPath(FStringView(InPath));
+	return Tmp;
+}
+
 FSoftObjectPath FSoftObjectPath::ConstructFromStringPath(FStringView InPath)
 {
 	FSoftObjectPath Tmp;
@@ -554,7 +561,7 @@ UObject* FSoftObjectPath::TryLoad(FUObjectSerializeContext* InLoadContext) const
 		if (IsSubobject())
 		{
 			// For subobjects, it's not safe to call LoadObject directly, so we want to load the parent object and then resolve again
-			FSoftObjectPath TopLevelPath = FSoftObjectPath(AssetPath, FString());
+			FSoftObjectPath TopLevelPath = FSoftObjectPath::ConstructFromAssetPath(AssetPath);
 			UObject* TopLevelObject = TopLevelPath.TryLoad(InLoadContext);
 
 			// This probably loaded the top-level object, so re-resolve ourselves
@@ -676,7 +683,7 @@ UObject* FSoftObjectPath::ResolveObjectInternal(const TCHAR* PathString) const
 	if (!FoundObject && IsSubobject())
 	{
 		// Try to resolve through the top level object
-		FSoftObjectPath TopLevelPath = FSoftObjectPath(AssetPath, FString());
+		FSoftObjectPath TopLevelPath(AssetPath);
 		UObject* TopLevelObject = TopLevelPath.ResolveObject();
 
 		// If the the top-level object exists but we can't find the object, defer the resolving to the top-level container object in case
diff --git a/Engine/Source/Runtime/CoreUObject/Public/UObject/SoftObjectPath.h b/Engine/Source/Runtime/CoreUObject/Public/UObject/SoftObjectPath.h
--- a/Engine/Source/Runtime/CoreUObject/Public/UObject/SoftObjectPath.h
+++ b/Engine/Source/Runtime/CoreUObject/Public/UObject/SoftObjectPath.h
@@ -105,6 +105,7 @@ PRAGMA_ENABLE_DEPRECATION_WARNINGS;
 	COREUOBJECT_API static FSoftObjectPath ConstructFromPackageAssetSubpath(FName InPackageName, FName InAssetName, FString&& InSubPathString);
 	COREUOBJECT_API static FSoftObjectPath ConstructFromPackageAsset(FName InPackageName, FName InAssetName);
 	COREUOBJECT_API static FSoftObjectPath ConstructFromAssetPath(FTopLevelAssetPath InAssetPath);
+	COREUOBJECT_API static FSoftObjectPath ConstructFromStringPath(FString&& InPath);
 	COREUOBJECT_API static FSoftObjectPath ConstructFromStringPath(FStringView InPath);
 	COREUOBJECT_API static FSoftObjectPath ConstructFromStringPath(FUtf8StringView InPath);
 	COREUOBJECT_API static FSoftObjectPath ConstructFromObject(const UObject* InObject);
