AssetRegistry - whitespace and comment changes to make FAssetRegistryState and FDependsNode match the rest of AssetRegistry whitespace style.
#rnx

diff --git a/Engine/Source/Runtime/AssetRegistry/Private/AssetRegistryImpl.h b/Engine/Source/Runtime/AssetRegistry/Private/AssetRegistryImpl.h
--- a/Engine/Source/Runtime/AssetRegistry/Private/AssetRegistryImpl.h
+++ b/Engine/Source/Runtime/AssetRegistry/Private/AssetRegistryImpl.h
@@ -72,7 +72,7 @@ private:
 	FEvent* Consumed = nullptr;
 };
 
-}
+} // namespace UE::AssetRegistry::Premade
 
 namespace UE::AssetRegistry
 {
@@ -103,7 +103,7 @@ namespace Private
 	};
 
 	using FInterfaceRWLock = FRWLockWithPriority;
-}
+} // namespace Private within namespace UE::AssetRegistry
 
 namespace Impl
 {
@@ -196,8 +196,8 @@ namespace Impl
 		// True if we ran out of time, EarlyExitCallback returned true, or RequestEarlyExit() was called
 		bool OutInterrupted = false;
 	};
-}
-}
+} // namespace Impl within namespace UE::AssetRegistry
+} // namespace UE::AssetRegistry
 
 /** Returns true if ANY work remains to be done. This work might require the game thread. */
 FORCEINLINE bool IsTickActive(UE::AssetRegistry::Impl::EGatherStatus Status) 
@@ -758,7 +758,7 @@ struct FTickContext
 	bool bHandleDeferred = false;
 };
 
-}
+} // namespace Impl within namespace UE::AssetRegistry
 
 namespace Utils
 {
@@ -837,6 +837,6 @@ void EnumerateMemoryAssets(const FARCompiledFilter& InFilter, TSet<FName>& OutPa
 	bool& bOutStopIteration, UE::AssetRegistry::Private::FInterfaceRWLock& InterfaceLock, const FAssetRegistryState& GuardedDataState,
 	TFunctionRef<bool(FAssetData&&)> Callback, bool bSkipARFilteredAssets);
 
-}
+} // namespace Utils within namespace UE::AssetRegistry
 
-}
+} // namespace UE::AssetRegistry
diff --git a/Engine/Source/Runtime/AssetRegistry/Private/AssetRegistryState.cpp b/Engine/Source/Runtime/AssetRegistry/Private/AssetRegistryState.cpp
--- a/Engine/Source/Runtime/AssetRegistry/Private/AssetRegistryState.cpp
+++ b/Engine/Source/Runtime/AssetRegistry/Private/AssetRegistryState.cpp
@@ -177,8 +177,12 @@ void FAssetRegistryState::FilterTags(const FAssetDataTagMapSharedView& InTagsAnd
 		}
 		else
 		{
-			const bool bInAllClassesList = AllClassesFilterList && (AllClassesFilterList->Contains(TagPair.Key) || AllClassesFilterList->Contains(UE::AssetRegistry::WildcardFName));
-			const bool bInClassSpecificList = ClassSpecificFilterList && (ClassSpecificFilterList->Contains(TagPair.Key) || ClassSpecificFilterList->Contains(UE::AssetRegistry::WildcardFName));
+			const bool bInAllClassesList = AllClassesFilterList
+				&& (AllClassesFilterList->Contains(TagPair.Key)
+					|| AllClassesFilterList->Contains(UE::AssetRegistry::WildcardFName));
+			const bool bInClassSpecificList = ClassSpecificFilterList
+				&& (ClassSpecificFilterList->Contains(TagPair.Key)
+					|| ClassSpecificFilterList->Contains(UE::AssetRegistry::WildcardFName));
 			if (Options.bUseAssetRegistryTagsAllowListInsteadOfDenyList)
 			{
 			// It's an allow list, only include it if it is in the all classes list or in the class specific list
@@ -197,7 +201,8 @@ void FAssetRegistryState::FilterTags(const FAssetDataTagMapSharedView& InTagsAnd
 	}
 }
 
-void FAssetRegistryState::InitializeFromExistingAndPrune(const FAssetRegistryState & ExistingState, const TSet<FName>& RequiredPackages, const TSet<FName>& RemovePackages,
+void FAssetRegistryState::InitializeFromExistingAndPrune(const FAssetRegistryState & ExistingState,
+	const TSet<FName>& RequiredPackages, const TSet<FName>& RemovePackages,
 	const TSet<int32> ChunksToKeep, const FAssetRegistrySerializationOptions& Options)
 {
 	const bool bIsFilteredByChunkId = ChunksToKeep.Num() != 0;
@@ -213,7 +218,10 @@ void FAssetRegistryState::InitializeFromExistingAndPrune(const FAssetRegistrySta
 		bool bRemoveDependencyData = true;
 
 		if (bIsFilteredByChunkId &&
-			!AssetData->GetChunkIDs().ContainsByPredicate([&](int32 ChunkId) { return ChunksToKeep.Contains(ChunkId); }))
+			!AssetData->GetChunkIDs().ContainsByPredicate([&](int32 ChunkId)
+				{
+					return ChunksToKeep.Contains(ChunkId);
+				}))
 		{
 			bRemoveAssetData = true;
 		}
@@ -242,7 +250,8 @@ void FAssetRegistryState::InitializeFromExistingAndPrune(const FAssetRegistrySta
 		}
 
 		FAssetDataTagMap NewTagsAndValues;
-		FAssetRegistryState::FilterTags(AssetData->TagsAndValues, NewTagsAndValues, Options.CookFilterlistTagsByClass.Find(AssetData->AssetClassPath), Options);
+		FAssetRegistryState::FilterTags(AssetData->TagsAndValues, NewTagsAndValues,
+			Options.CookFilterlistTagsByClass.Find(AssetData->AssetClassPath), Options);
 
 		FAssetData* NewAssetData = new FAssetData(AssetData->PackageName, AssetData->PackagePath, AssetData->AssetName,
 			AssetData->AssetClassPath, NewTagsAndValues, AssetData->GetChunkIDs(), AssetData->PackageFlags);
@@ -305,7 +314,11 @@ void FAssetRegistryState::InitializeFromExistingAndPrune(const FAssetRegistrySta
 	for (FDependsNode* OldNode : ValidDependsNodes)
 	{
 		FDependsNode* NewNode = CreateOrFindDependsNode(OldNode->GetIdentifier());
-		OldNode->IterateOverDependencies([&, OldNode, NewNode](FDependsNode* InDependency, UE::AssetRegistry::EDependencyCategory InCategory, UE::AssetRegistry::EDependencyProperty InFlags, bool bDuplicate) {
+		OldNode->IterateOverDependencies([&, OldNode, NewNode](FDependsNode* InDependency,
+			UE::AssetRegistry::EDependencyCategory InCategory,
+			UE::AssetRegistry::EDependencyProperty InFlags,
+			bool bDuplicate)
+		{
 			if (ValidDependsNodes.Contains(InDependency))
 			{
 				// Only add link if it's part of the filtered asset set
@@ -334,8 +347,11 @@ void FAssetRegistryState::InitializeFromExistingAndPrune(const FAssetRegistrySta
 	SetDependencyNodeSorting(true, true);
 }
 
-void FAssetRegistryState::InitializeFromExisting(const FAssetDataMap& AssetDataMap, const TMap<FAssetIdentifier, FDependsNode*>& DependsNodeMap, 
-	const TMap<FName, FAssetPackageData*>& AssetPackageDataMap, const FAssetRegistrySerializationOptions& Options, EInitializationMode InInitializationMode)
+void FAssetRegistryState::InitializeFromExisting(const FAssetDataMap& AssetDataMap,
+	const TMap<FAssetIdentifier, FDependsNode*>& DependsNodeMap, 
+	const TMap<FName, FAssetPackageData*>& AssetPackageDataMap,
+	const FAssetRegistrySerializationOptions& Options,
+	EInitializationMode InInitializationMode)
 {
 	if (InInitializationMode == EInitializationMode::Rebuild)
 	{
@@ -371,7 +387,8 @@ void FAssetRegistryState::InitializeFromExisting(const FAssetDataMap& AssetDataM
 
 		// Filter asset registry tags now
 		FAssetDataTagMap LocalTagsAndValues;
-		FAssetRegistryState::FilterTags(AssetData.TagsAndValues, LocalTagsAndValues, Options.CookFilterlistTagsByClass.Find(AssetData.AssetClassPath), Options);
+		FAssetRegistryState::FilterTags(AssetData.TagsAndValues, LocalTagsAndValues,
+			Options.CookFilterlistTagsByClass.Find(AssetData.AssetClassPath), Options);
 		
 		if (ExistingData)
 		{
@@ -420,7 +437,8 @@ void FAssetRegistryState::InitializeFromExisting(const FAssetDataMap& AssetDataM
 		const TMap<FAssetIdentifier, FDependsNode*>* DependsNodesToAdd = &DependsNodeMap;
 		if (InInitializationMode == EInitializationMode::OnlyUpdateNew)
 		{
-			// Keep the original DependsNodeMap for reference, but remove from NodesToAdd all nodes that already have dependency data
+			// Keep the original DependsNodeMap for reference,
+			// but remove from NodesToAdd all nodes that already have dependency data
 			// Also reserve up-front all (unfiltered) nodes we are adding, to avoid reallocating the Referencers array.
 			FilteredDependsNodeMap.Reserve(DependsNodeMap.Num());
 			DependsNodesToAdd = &FilteredDependsNodeMap;
@@ -452,7 +470,10 @@ void FAssetRegistryState::InitializeFromExisting(const FAssetDataMap& AssetDataM
 			FDependsNode* SourceNode = Pair.Value;
 			FDependsNode* TargetNode = CreateOrFindDependsNode(Pair.Key);
 			SourceNode->IterateOverDependencies([this, &DependsNodeMap, &ScriptPackages, TargetNode]
-			(FDependsNode* InDependency, UE::AssetRegistry::EDependencyCategory InCategory, UE::AssetRegistry::EDependencyProperty InFlags, bool bDuplicate) {
+			(FDependsNode* InDependency, UE::AssetRegistry::EDependencyCategory InCategory,
+				UE::AssetRegistry::EDependencyProperty InFlags,
+				bool bDuplicate)
+			{
 				const FAssetIdentifier& Identifier = InDependency->GetIdentifier();
 				if (DependsNodeMap.Find(Identifier) || ScriptPackages.Contains(Identifier))
 				{
@@ -472,12 +493,14 @@ void FAssetRegistryState::InitializeFromExisting(const FAssetDataMap& AssetDataM
 	}
 }
 
-void FAssetRegistryState::PruneAssetData(const TSet<FName>& RequiredPackages, const TSet<FName>& RemovePackages, const FAssetRegistrySerializationOptions& Options)
+void FAssetRegistryState::PruneAssetData(const TSet<FName>& RequiredPackages, const TSet<FName>& RemovePackages,
+	const FAssetRegistrySerializationOptions& Options)
 {
 	PruneAssetData(RequiredPackages, RemovePackages, TSet<int32>(), Options);
 }
 
-void FAssetRegistryState::PruneAssetData(const TSet<FName>& RequiredPackages, const TSet<FName>& RemovePackages, const TSet<int32> ChunksToKeep, const FAssetRegistrySerializationOptions& Options)
+void FAssetRegistryState::PruneAssetData(const TSet<FName>& RequiredPackages, const TSet<FName>& RemovePackages,
+	const TSet<int32> ChunksToKeep, const FAssetRegistrySerializationOptions& Options)
 {
 	FAssetRegistryPruneOptions PruneOptions;
 	PruneOptions.RequiredPackages = RequiredPackages;
@@ -513,7 +536,10 @@ void FAssetRegistryState::Prune(const FAssetRegistryPruneOptions& PruneOptions)
 		bool bRemoveDependencyData = true;
 
 		if (bIsFilteredByChunkId &&
-			!AssetData->GetChunkIDs().ContainsByPredicate([&](int32 ChunkId) { return ChunksToKeep.Contains(ChunkId); }))
+			!AssetData->GetChunkIDs().ContainsByPredicate([&](int32 ChunkId)
+				{
+					return ChunksToKeep.Contains(ChunkId);
+				}))
 		{
 			bRemoveAssetData = true;
 		}
@@ -594,7 +620,8 @@ void FAssetRegistryState::Prune(const FAssetRegistryPruneOptions& PruneOptions)
 			{
 				if (!KnownPrimaryAssetIds.Contains(PrimaryAssetId))
 				{
-					if (!PruneOptions.RemoveDependenciesWithoutPackagesKeepPrimaryAssetTypes.Contains(PrimaryAssetId.PrimaryAssetType))
+					if (!PruneOptions.RemoveDependenciesWithoutPackagesKeepPrimaryAssetTypes.Contains(
+						PrimaryAssetId.PrimaryAssetType))
 					{
 						bRemoveDependsNode = true;
 					}
@@ -616,7 +643,8 @@ void FAssetRegistryState::Prune(const FAssetRegistryPruneOptions& PruneOptions)
 		{
 			CachedDependsNodes.Remove(DependsNode->GetIdentifier());
 			NumDependsNodes--;
-			// if the depends nodes were preallocated in a block, we can't delete them one at a time, only the whole chunk in the destructor
+			// if the depends nodes were preallocated in a block, we can't delete them one at a time,
+			// only the whole chunk in the destructor
 			if (PreallocatedDependsNodeDataBuffers.Num() == 0)
 			{
 				delete DependsNode;
@@ -624,7 +652,10 @@ void FAssetRegistryState::Prune(const FAssetRegistryPruneOptions& PruneOptions)
 		}
 		else
 		{
-			DependsNode->RemoveLinks([&RemoveDependsNodes](const FDependsNode* ExistingDependsNode) { return RemoveDependsNodes.Contains(ExistingDependsNode); });
+			DependsNode->RemoveLinks([&RemoveDependsNodes](const FDependsNode* ExistingDependsNode)
+				{
+					return RemoveDependsNodes.Contains(ExistingDependsNode);
+				});
 		}
 	}
 
@@ -647,15 +678,20 @@ bool FAssetRegistryState::HasAssets(const FName PackagePath, bool bSkipARFiltere
 		return FoundAssetArray->ContainsByPredicate([this, bSkipARFilteredAssets](FAssetData* AssetData)
 		{
 			return AssetData && !IsPackageUnmountedAndFiltered(AssetData->PackageName)
-				&& (!bSkipARFilteredAssets || !UE::AssetRegistry::FFiltering::ShouldSkipAsset(AssetData->AssetClassPath, AssetData->PackageFlags));
+				&& (!bSkipARFilteredAssets
+					|| !UE::AssetRegistry::FFiltering::ShouldSkipAsset(
+						AssetData->AssetClassPath, AssetData->PackageFlags));
 		});
 	}
 	return false;
 }
 
-bool FAssetRegistryState::GetAssets(const FARCompiledFilter& Filter, const TSet<FName>& PackageNamesToSkip, TArray<FAssetData>& OutAssetData, bool bSkipARFilteredAssets) const
+bool FAssetRegistryState::GetAssets(const FARCompiledFilter& Filter, const TSet<FName>& PackageNamesToSkip,
+	TArray<FAssetData>& OutAssetData, bool bSkipARFilteredAssets) const
 {
-	const UE::AssetRegistry::EEnumerateAssetsFlags Flags = bSkipARFilteredAssets ? UE::AssetRegistry::EEnumerateAssetsFlags::None : UE::AssetRegistry::EEnumerateAssetsFlags::AllowUnfilteredArAssets;
+	const UE::AssetRegistry::EEnumerateAssetsFlags Flags = bSkipARFilteredAssets
+		? UE::AssetRegistry::EEnumerateAssetsFlags::None
+		: UE::AssetRegistry::EEnumerateAssetsFlags::AllowUnfilteredArAssets;
 	return EnumerateAssets(Filter, PackageNamesToSkip, [&OutAssetData](const FAssetData& AssetData)
 	{
 		OutAssetData.Emplace(AssetData);
@@ -695,7 +731,8 @@ bool DecideIntersectionMethod(int32 PreviousSize, int32 FilterResultsSize, int32
 	return FilterCost < ArrayCost;
 }
 
-void ArrayIntersection(TArray<const FAssetData*>& InOutResults, TConstArrayView<TArrayView<const FAssetData*>> Matches, int32 TotalMatches)
+void ArrayIntersection(TArray<const FAssetData*>& InOutResults,
+	TConstArrayView<TArrayView<const FAssetData*>> Matches, int32 TotalMatches)
 {
 	if (InOutResults.Num() < TotalMatches)
 	{
@@ -712,7 +749,8 @@ void ArrayIntersection(TArray<const FAssetData*>& InOutResults, TConstArrayView<
 			{
 				bool* Result = Exists.Find(Asset);
 				if (Result && 
-					!(*Result) // If there are duplicates of an Asset in multiple elements of Matches, only add the first one
+					!(*Result) // If there are duplicates of an Asset in multiple elements of Matches,
+					           // only add the first one
 				)
 				{
 					*Result = true;
@@ -758,7 +796,8 @@ void FilterAssets(TArray<const FAssetData*>& InOutResults, const TMap<KeyType, A
 		}
 	}
 
-	// Keys is a TSet and entries in the AccelerationMap do not overlap, so there should be no duplicates to remove in Matches
+	// Keys is a TSet and entries in the AccelerationMap do not overlap,
+	// so there should be no duplicates to remove in Matches
 	if (InOutResults.IsEmpty())
 	{
 		// No previous Results; set Results equal to the values found in AccelerationMap
@@ -844,7 +883,8 @@ bool AssetDataMatchesTag(const FAssetData* AssetData, const TPair<FName, TOption
 
 template<typename CallbackType, typename AccelerationMapType>
 void FilterAssets(TArray<const FAssetData*>& InOutResults, const AccelerationMapType& AccelerationMap,
-	const TMultiMap<FName, TOptional<FString>>& TagsAndValues, CallbackType&& FunctionToKeepAsset, int32 FilterComplexity)
+	const TMultiMap<FName, TOptional<FString>>& TagsAndValues,
+	CallbackType&& FunctionToKeepAsset, int32 FilterComplexity)
 {
 	TArray<TArray<const FAssetData*>, TInlineAllocator<10>> Matches;
 	Matches.Reserve(TagsAndValues.Num());
@@ -910,14 +950,17 @@ void FilterAssets(TArray<const FAssetData*>& InOutResults, const AccelerationMap
 bool FAssetRegistryState::EnumerateAssets(const FARCompiledFilter& Filter, const TSet<FName>& PackageNamesToSkip,
 	TFunctionRef<bool(const FAssetData&)> Callback, bool bSkipARFilteredAssets) const
 {
-	const UE::AssetRegistry::EEnumerateAssetsFlags Flags = bSkipARFilteredAssets ? UE::AssetRegistry::EEnumerateAssetsFlags::None : UE::AssetRegistry::EEnumerateAssetsFlags::AllowUnfilteredArAssets;
+	const UE::AssetRegistry::EEnumerateAssetsFlags Flags = bSkipARFilteredAssets
+		? UE::AssetRegistry::EEnumerateAssetsFlags::None
+		: UE::AssetRegistry::EEnumerateAssetsFlags::AllowUnfilteredArAssets;
 	return EnumerateAssets(Filter, PackageNamesToSkip, Callback, Flags);
 }
 
 bool FAssetRegistryState::EnumerateAssets(const FARCompiledFilter& Filter, const TSet<FName>& PackageNamesToSkip,
 	TFunctionRef<bool(const FAssetData&)> Callback) const
 {
-	return EnumerateAssets(Filter, PackageNamesToSkip, Callback, UE::AssetRegistry::EEnumerateAssetsFlags::AllowUnfilteredArAssets);
+	return EnumerateAssets(Filter, PackageNamesToSkip, Callback,
+		UE::AssetRegistry::EEnumerateAssetsFlags::AllowUnfilteredArAssets);
 }
 
 bool FAssetRegistryState::EnumerateAssets(const FARCompiledFilter& Filter, const TSet<FName>& PackageNamesToSkip,
@@ -950,8 +993,9 @@ bool FAssetRegistryState::EnumerateAssets(const FARCompiledFilter& Filter, const
 				return true;
 			}
 
-			return (!EnumHasAnyFlags(InEnumerateFlags, UE::AssetRegistry::EEnumerateAssetsFlags::AllowUnfilteredArAssets) &&
-				UE::AssetRegistry::FFiltering::ShouldSkipAsset(AssetData->AssetClassPath, AssetData->PackageFlags));
+			return (!EnumHasAnyFlags(InEnumerateFlags,
+				UE::AssetRegistry::EEnumerateAssetsFlags::AllowUnfilteredArAssets)
+				&& UE::AssetRegistry::FFiltering::ShouldSkipAsset(AssetData->AssetClassPath, AssetData->PackageFlags));
 		};
 
 
@@ -1058,9 +1102,12 @@ bool FAssetRegistryState::EnumerateAssets(const FARCompiledFilter& Filter, const
 	return true;
 }
 
-bool FAssetRegistryState::GetAllAssets(const TSet<FName>& PackageNamesToSkip, TArray<FAssetData>& OutAssetData, bool bSkipARFilteredAssets) const
+bool FAssetRegistryState::GetAllAssets(const TSet<FName>& PackageNamesToSkip, TArray<FAssetData>& OutAssetData,
+	bool bSkipARFilteredAssets) const
 {
-	const UE::AssetRegistry::EEnumerateAssetsFlags EnumerateFlags = bSkipARFilteredAssets ? UE::AssetRegistry::EEnumerateAssetsFlags::None : UE::AssetRegistry::EEnumerateAssetsFlags::AllowUnfilteredArAssets;
+	const UE::AssetRegistry::EEnumerateAssetsFlags EnumerateFlags = bSkipARFilteredAssets
+		? UE::AssetRegistry::EEnumerateAssetsFlags::None
+		: UE::AssetRegistry::EEnumerateAssetsFlags::AllowUnfilteredArAssets;
 	OutAssetData.Reserve(OutAssetData.Num() + CachedAssets.Num() - PackageNamesToSkip.Num());
 	return EnumerateAllAssets(PackageNamesToSkip, [&OutAssetData](const FAssetData& AssetData)
 	{
@@ -1070,15 +1117,20 @@ bool FAssetRegistryState::GetAllAssets(const TSet<FName>& PackageNamesToSkip, TA
 	EnumerateFlags);
 }
 
-bool FAssetRegistryState::EnumerateAllAssets(const TSet<FName>& PackageNamesToSkip, TFunctionRef<bool(const FAssetData&)> Callback, bool bSkipARFilteredAssets) const
+bool FAssetRegistryState::EnumerateAllAssets(const TSet<FName>& PackageNamesToSkip,
+	TFunctionRef<bool(const FAssetData&)> Callback, bool bSkipARFilteredAssets) const
 {
-	const UE::AssetRegistry::EEnumerateAssetsFlags EnumerateFlags = bSkipARFilteredAssets ? UE::AssetRegistry::EEnumerateAssetsFlags::None : UE::AssetRegistry::EEnumerateAssetsFlags::AllowUnfilteredArAssets;
+	const UE::AssetRegistry::EEnumerateAssetsFlags EnumerateFlags = bSkipARFilteredAssets
+		? UE::AssetRegistry::EEnumerateAssetsFlags::None
+		: UE::AssetRegistry::EEnumerateAssetsFlags::AllowUnfilteredArAssets;
 	return EnumerateAllAssets(PackageNamesToSkip, Callback, EnumerateFlags);
 }
 
-bool FAssetRegistryState::EnumerateAllAssets(const TSet<FName>& PackageNamesToSkip, TFunctionRef<bool(const FAssetData&)> Callback) const
+bool FAssetRegistryState::EnumerateAllAssets(const TSet<FName>& PackageNamesToSkip,
+	TFunctionRef<bool(const FAssetData&)> Callback) const
 {
-	return EnumerateAllAssets(PackageNamesToSkip, Callback, UE::AssetRegistry::EEnumerateAssetsFlags::AllowUnfilteredArAssets);
+	return EnumerateAllAssets(PackageNamesToSkip, Callback,
+		UE::AssetRegistry::EEnumerateAssetsFlags::AllowUnfilteredArAssets);
 }
 
 void FAssetRegistryState::EnumerateAllAssets(TFunctionRef<void(const FAssetData&)> Callback) const
@@ -1092,14 +1144,19 @@ void FAssetRegistryState::EnumerateAllAssets(TFunctionRef<void(const FAssetData&
 	}
 }
 
-bool FAssetRegistryState::EnumerateAllAssets(const TSet<FName>& PackageNamesToSkip, TFunctionRef<bool(const FAssetData&)> Callback, UE::AssetRegistry::EEnumerateAssetsFlags InEnumerateFlags) const
+bool FAssetRegistryState::EnumerateAllAssets(const TSet<FName>& PackageNamesToSkip,
+	TFunctionRef<bool(const FAssetData&)> Callback,
+	UE::AssetRegistry::EEnumerateAssetsFlags InEnumerateFlags) const
 {
 	for (const FAssetData* AssetData : CachedAssets)
 	{
 		if (AssetData &&
-			!PackageNamesToSkip.Contains(AssetData->PackageName) &&
-			(EnumHasAnyFlags(InEnumerateFlags, UE::AssetRegistry::EEnumerateAssetsFlags::AllowUnmountedPaths) || !IsPackageUnmountedAndFiltered(AssetData->PackageName)) &&
-			(EnumHasAnyFlags(InEnumerateFlags, UE::AssetRegistry::EEnumerateAssetsFlags::AllowUnfilteredArAssets) || !UE::AssetRegistry::FFiltering::ShouldSkipAsset(AssetData->AssetClassPath, AssetData->PackageFlags)))
+			!PackageNamesToSkip.Contains(AssetData->PackageName)
+			&& (EnumHasAnyFlags(InEnumerateFlags, UE::AssetRegistry::EEnumerateAssetsFlags::AllowUnmountedPaths) 
+				|| !IsPackageUnmountedAndFiltered(AssetData->PackageName))
+			&& (EnumHasAnyFlags(InEnumerateFlags, UE::AssetRegistry::EEnumerateAssetsFlags::AllowUnfilteredArAssets)
+				|| !UE::AssetRegistry::FFiltering::ShouldSkipAsset(AssetData->AssetClassPath,
+					AssetData->PackageFlags)))
 		{
 			if (!Callback(*AssetData))
 			{
@@ -1156,7 +1213,8 @@ FName FAssetRegistryState::GetFirstPackageByName(FStringView PackageName) const
 	if (LongPackageNames.Num() > 1)
 	{
 		LongPackageNames.Sort(FNameLexicalLess());
-		UE_LOG(LogAssetRegistry, Warning, TEXT("GetFirstPackageByName('%.*s') is returning '%s', but it also found '%s'%s."),
+		UE_LOG(LogAssetRegistry, Warning,
+			TEXT("GetFirstPackageByName('%.*s') is returning '%s', but it also found '%s'%s."),
 			PackageName.Len(), PackageName.GetData(), *LongPackageNames[0].ToString(), *LongPackageNames[1].ToString(),
 			(LongPackageNames.Num() > 2 ? *FString::Printf(TEXT(" and %d others"), LongPackageNames.Num() - 2) : TEXT("")));
 	}
@@ -1164,8 +1222,8 @@ FName FAssetRegistryState::GetFirstPackageByName(FStringView PackageName) const
 }
 
 bool FAssetRegistryState::GetDependencies(const FAssetIdentifier& AssetIdentifier,
-										  TArray<FAssetIdentifier>& OutDependencies,
-										  UE::AssetRegistry::EDependencyCategory Category, const UE::AssetRegistry::FDependencyQuery& Flags) const
+	TArray<FAssetIdentifier>& OutDependencies,
+	UE::AssetRegistry::EDependencyCategory Category, const UE::AssetRegistry::FDependencyQuery& Flags) const
 {
 	const FDependsNode* const* NodePtr = CachedDependsNodes.Find(AssetIdentifier);
 	const FDependsNode* Node = nullptr;
@@ -1186,8 +1244,8 @@ bool FAssetRegistryState::GetDependencies(const FAssetIdentifier& AssetIdentifie
 }
 
 bool FAssetRegistryState::GetDependencies(const FAssetIdentifier& AssetIdentifier,
-										  TArray<FAssetDependency>& OutDependencies,
-										  UE::AssetRegistry::EDependencyCategory Category, const UE::AssetRegistry::FDependencyQuery& Flags) const
+	TArray<FAssetDependency>& OutDependencies,
+	UE::AssetRegistry::EDependencyCategory Category, const UE::AssetRegistry::FDependencyQuery& Flags) const
 {
 	const FDependsNode* const* NodePtr = CachedDependsNodes.Find(AssetIdentifier);
 	const FDependsNode* Node = nullptr;
@@ -1208,8 +1266,8 @@ bool FAssetRegistryState::GetDependencies(const FAssetIdentifier& AssetIdentifie
 }
 
 bool FAssetRegistryState::GetReferencers(const FAssetIdentifier& AssetIdentifier,
-										 TArray<FAssetIdentifier>& OutReferencers,
-										 UE::AssetRegistry::EDependencyCategory Category, const UE::AssetRegistry::FDependencyQuery& Flags) const
+	TArray<FAssetIdentifier>& OutReferencers,
+	UE::AssetRegistry::EDependencyCategory Category, const UE::AssetRegistry::FDependencyQuery& Flags) const
 {
 	const FDependsNode* const* NodePtr = CachedDependsNodes.Find(AssetIdentifier);
 	const FDependsNode* Node = nullptr;
@@ -1239,8 +1297,8 @@ bool FAssetRegistryState::GetReferencers(const FAssetIdentifier& AssetIdentifier
 }
 
 bool FAssetRegistryState::GetReferencers(const FAssetIdentifier& AssetIdentifier,
-										 TArray<FAssetDependency>& OutReferencers,
-										 UE::AssetRegistry::EDependencyCategory Category, const UE::AssetRegistry::FDependencyQuery& Flags) const
+	TArray<FAssetDependency>& OutReferencers,
+	UE::AssetRegistry::EDependencyCategory Category, const UE::AssetRegistry::FDependencyQuery& Flags) const
 {
 	const FDependsNode* const* NodePtr = CachedDependsNodes.Find(AssetIdentifier);
 	const FDependsNode* Node = nullptr;
@@ -1304,7 +1362,8 @@ void FAssetRegistryState::SetDependencies(const FAssetIdentifier& AssetIdentifie
 {
 	for (const FAssetDependency& Dependency : Dependencies)
 	{
-		checkf(!(Dependency.Category & ~Category), TEXT("Input dependency has category %d which is outside of the requested categories %d."),
+		checkf(!(Dependency.Category & ~Category),
+			TEXT("Input dependency has category %d which is outside of the requested categories %d."),
 			(int32)Dependency.Category, (int32)Category);
 	}
 
@@ -1354,7 +1413,8 @@ void FAssetRegistryState::SetReferencers(const FAssetIdentifier& AssetIdentifier
 {
 	for (const FAssetDependency& Referencer : Referencers)
 	{
-		checkf(!(Referencer.Category & ~Category), TEXT("Input referencer has category %d which is outside of the requested categories %d."),
+		checkf(!(Referencer.Category & ~Category),
+			TEXT("Input referencer has category %d which is outside of the requested categories %d."),
 			(int32)Referencer.Category, (int32)Category);
 	}
 
@@ -1398,13 +1458,16 @@ bool FAssetRegistryState::Save(FArchive& OriginalAr, const FAssetRegistrySeriali
 		{
 			SortedAssetsByObjectPath.Add({ AssetData, AssetData->GetSoftObjectPath() });
 		}
-		Algo::Sort(SortedAssetsByObjectPath, [](const TPair<FAssetData*, FSoftObjectPath>& A, const TPair<FAssetData*, FSoftObjectPath>& B) {
+		Algo::Sort(SortedAssetsByObjectPath, [](const TPair<FAssetData*, FSoftObjectPath>& A,
+			const TPair<FAssetData*, FSoftObjectPath>& B)
+		{
 			return A.Value.LexicalLess(B.Value);
 		});
 
 		for (TPair<FAssetData*, FSoftObjectPath>& Asset : SortedAssetsByObjectPath)
 		{
-			// Hardcoding FAssetRegistryVersion::LatestVersion here so that branches can get optimized out in the forceinlined SerializeForCache
+			// Hardcoding FAssetRegistryVersion::LatestVersion here so that branches can get optimized out in
+			// the forceinlined SerializeForCache
 			Asset.Key->SerializeForCache(Ar);
 		}
 	}
@@ -1437,7 +1500,10 @@ bool FAssetRegistryState::Save(FArchive& OriginalAr, const FAssetRegistrySeriali
 				Dependencies.Add(Node);
 			}
 		}
-		Algo::Sort(Dependencies, [](FDependsNode* A, FDependsNode* B) { return A->GetIdentifier().LexicalLess(B->GetIdentifier()); });
+		Algo::Sort(Dependencies, [](FDependsNode* A, FDependsNode* B) 
+			{
+				return A->GetIdentifier().LexicalLess(B->GetIdentifier());
+			});
 		int32 NumDependencies = Dependencies.Num();
 
 		TMap<FDependsNode*, int32> DependsIndexMap;
@@ -1448,7 +1514,8 @@ bool FAssetRegistryState::Save(FArchive& OriginalAr, const FAssetRegistrySeriali
 			DependsIndexMap.Add(Node, Index++);
 		}
 
-		TUniqueFunction<int32(FDependsNode*, bool bAsReferencer)> GetSerializeIndexFromNode = [this, &RedirectCache, &DependsIndexMap](FDependsNode* InDependency, bool bAsReferencer)
+		TUniqueFunction<int32(FDependsNode*, bool bAsReferencer)> GetSerializeIndexFromNode =
+			[this, &RedirectCache, &DependsIndexMap](FDependsNode* InDependency, bool bAsReferencer)
 		{
 			if (!bAsReferencer)
 			{
@@ -1490,7 +1557,10 @@ bool FAssetRegistryState::Save(FArchive& OriginalAr, const FAssetRegistrySeriali
 		Ar << PackageDataCount;
 
 		TArray<TPair<FName, FAssetPackageData*>> SortedPackageData = CachedPackageData.Array();
-		Algo::Sort(SortedPackageData, [](TPair<FName, FAssetPackageData*>& A, TPair<FName, FAssetPackageData*>& B) { return A.Key.LexicalLess(B.Key); });
+		Algo::Sort(SortedPackageData, [](TPair<FName, FAssetPackageData*>& A, TPair<FName, FAssetPackageData*>& B)
+			{
+				return A.Key.LexicalLess(B.Key);
+			});
 		for (TPair<FName, FAssetPackageData*>& Pair : SortedPackageData)
 		{
 			Ar << Pair.Key;
@@ -1506,7 +1576,8 @@ bool FAssetRegistryState::Save(FArchive& OriginalAr, const FAssetRegistrySeriali
 	return !OriginalAr.IsError();
 }
 
-bool FAssetRegistryState::Load(FArchive& OriginalAr, const FAssetRegistryLoadOptions& Options, FAssetRegistryVersion::Type* OutVersion)
+bool FAssetRegistryState::Load(FArchive& OriginalAr, const FAssetRegistryLoadOptions& Options,
+	FAssetRegistryVersion::Type* OutVersion)
 {
 	LLM_SCOPE(ELLMTag::AssetRegistry);
 	FAssetRegistryHeader Header;
@@ -1516,7 +1587,8 @@ bool FAssetRegistryState::Load(FArchive& OriginalAr, const FAssetRegistryLoadOpt
 		*OutVersion = Header.Version;
 	}
 
-	FSoftObjectPathSerializationScope SerializationScope(NAME_None, NAME_None, ESoftObjectPathCollectType::NonPackage, ESoftObjectPathSerializeType::AlwaysSerialize);
+	FSoftObjectPathSerializationScope SerializationScope(NAME_None, NAME_None,
+		ESoftObjectPathCollectType::NonPackage, ESoftObjectPathSerializeType::AlwaysSerialize);
 
 	if (Header.Version < FAssetRegistryVersion::RemovedMD5Hash)
 	{
@@ -1547,7 +1619,8 @@ bool FAssetRegistryState::Load(FArchive& OriginalAr, const FAssetRegistryLoadOpt
 	return !OriginalAr.IsError();
 }
 
-/* static */ bool FAssetRegistryState::LoadFromDisk(const TCHAR* InPath, const FAssetRegistryLoadOptions& InOptions, FAssetRegistryState& OutState, FAssetRegistryVersion::Type* OutVersion)
+/* static */ bool FAssetRegistryState::LoadFromDisk(const TCHAR* InPath, const FAssetRegistryLoadOptions& InOptions,
+	FAssetRegistryState& OutState, FAssetRegistryVersion::Type* OutVersion)
 {
 	check(InPath);
 
@@ -1568,7 +1641,8 @@ bool FAssetRegistryState::Load(FArchive& OriginalAr, const FAssetRegistryLoadOpt
 }
 
 template<class Archive>
-void FAssetRegistryState::Load(Archive&& Ar, const FAssetRegistryHeader& Header, const FAssetRegistryLoadOptions& Options)
+void FAssetRegistryState::Load(Archive&& Ar, const FAssetRegistryHeader& Header,
+	const FAssetRegistryLoadOptions& Options)
 {
 	FAssetRegistryVersion::Type Version = Header.Version;
 
@@ -1580,8 +1654,9 @@ void FAssetRegistryState::Load(Archive&& Ar, const FAssetRegistryHeader& Header,
 	TArrayView<FAssetData> PreallocatedAssetDataBuffer(new FAssetData[LocalNumAssets], LocalNumAssets);
 	PreallocatedAssetDataBuffers.Add(PreallocatedAssetDataBuffer.GetData());
 
-	// Optimizing serialization of latest asset data format by moving version checking out of SerializeForCache function and falling back to versioned
-	// serialization should we attempt to load an older version of AR (usually commandlets)
+	// Optimizing serialization of latest asset data format by moving version checking out of SerializeForCache
+	// function and falling back to versioned serialization should we attempt to load an older version of AR
+	// (usually commandlets)
 	if (Version == FAssetRegistryVersion::LatestVersion)
 	{
 		for (FAssetData& NewAssetData : PreallocatedAssetDataBuffer)
@@ -1633,7 +1708,8 @@ void FAssetRegistryState::Load(Archive&& Ar, const FAssetRegistryHeader& Header,
 		TArrayView<FAssetPackageData> PreallocatedPackageDataBuffer;
 		if (Options.bLoadPackageData)
 		{
-			PreallocatedPackageDataBuffer = TArrayView<FAssetPackageData>(new FAssetPackageData[LocalNumPackageData], LocalNumPackageData);
+			PreallocatedPackageDataBuffer = TArrayView<FAssetPackageData>(
+				new FAssetPackageData[LocalNumPackageData], LocalNumPackageData);
 			PreallocatedPackageDataBuffers.Add(PreallocatedPackageDataBuffer.GetData());
 			CachedPackageData.Reserve(LocalNumPackageData);
 		}
@@ -1677,7 +1753,8 @@ void FAssetRegistryState::LoadDependencies(FArchive& Ar)
 	PreallocatedDependsNodeDataBuffers.Add(PreallocatedDependsNodeDataBuffer);
 	CachedDependsNodes.Reserve(LocalNumDependsNodes);
 	
-	TUniqueFunction<FDependsNode*(int32)> GetNodeFromSerializeIndex = [&PreallocatedDependsNodeDataBuffer, LocalNumDependsNodes](int32 Index) -> FDependsNode*
+	TUniqueFunction<FDependsNode*(int32)> GetNodeFromSerializeIndex =
+		[&PreallocatedDependsNodeDataBuffer, LocalNumDependsNodes](int32 Index) -> FDependsNode*
 	{
 		if (Index < 0 || LocalNumDependsNodes <= Index)
 		{
@@ -1695,7 +1772,8 @@ void FAssetRegistryState::LoadDependencies(FArchive& Ar)
 	}
 }
 
-void FAssetRegistryState::LoadDependencies_BeforeFlags(FArchive& Ar, bool bSerializeDependencies, FAssetRegistryVersion::Type Version)
+void FAssetRegistryState::LoadDependencies_BeforeFlags(FArchive& Ar, bool bSerializeDependencies,
+	FAssetRegistryVersion::Type Version)
 {
 	int32 LocalNumDependsNodes = 0;
 	Ar << LocalNumDependsNodes;
@@ -1708,7 +1786,8 @@ void FAssetRegistryState::LoadDependencies_BeforeFlags(FArchive& Ar, bool bSeria
 		PreallocatedDependsNodeDataBuffers.Add(PreallocatedDependsNodeDataBuffer);
 		CachedDependsNodes.Reserve(LocalNumDependsNodes);
 	}
-	TUniqueFunction<FDependsNode* (int32)> GetNodeFromSerializeIndex = [&PreallocatedDependsNodeDataBuffer, LocalNumDependsNodes](int32 Index)->FDependsNode *
+	TUniqueFunction<FDependsNode* (int32)> GetNodeFromSerializeIndex =
+		[&PreallocatedDependsNodeDataBuffer, LocalNumDependsNodes](int32 Index)->FDependsNode *
 	{
 		if (Index < 0 || LocalNumDependsNodes <= Index)
 		{
@@ -1735,7 +1814,8 @@ void FAssetRegistryState::LoadDependencies_BeforeFlags(FArchive& Ar, bool bSeria
 		}
 
 		// Call the DependsNode legacy serialization function
-		DependsNode->SerializeLoad_BeforeFlags(Ar, Version, PreallocatedDependsNodeDataBuffer, LocalNumDependsNodes, bSerializeDependencies, HardBits, SoftBits, HardManageBits, SoftManageBits);
+		DependsNode->SerializeLoad_BeforeFlags(Ar, Version, PreallocatedDependsNodeDataBuffer, LocalNumDependsNodes,
+			bSerializeDependencies, HardBits, SoftBits, HardManageBits, SoftManageBits);
 
 		// Register the DependsNode with its AssetIdentifier
 		if (bSerializeDependencies)
@@ -1782,7 +1862,11 @@ SIZE_T FAssetRegistryState::GetAllocatedSize(bool bLogDetailed) const
 		UE_LOG(LogAssetRegistry, Log, TEXT("Index Size: %" SIZE_T_FMT "k"), MapMemory / 1024);
 	}
 
-	SIZE_T AssetDataSize = 0, AssetBundlesSize = 0, NumAssetBundles = 0, NumSoftObjectPaths = 0, NumTopLevelAssetPaths = 0;
+	SIZE_T AssetDataSize = 0;
+	SIZE_T AssetBundlesSize = 0;
+	SIZE_T NumAssetBundles = 0;
+	SIZE_T NumSoftObjectPaths = 0;
+	SIZE_T NumTopLevelAssetPaths = 0;
 	FAssetDataTagMapSharedView::FMemoryCounter TagMemoryUsage;
 
 	for (const FAssetData* AssetData : CachedAssets)
@@ -1848,7 +1932,8 @@ SIZE_T FAssetRegistryState::GetAllocatedSize(bool bLogDetailed) const
 		PackageDataSize += PackageDataPair.Value->GetAllocatedSize();
 	}
 
-	SIZE_T TotalBytes = MapMemory + AssetDataSize + AssetBundlesSize + TagMemoryUsage.GetFixedSize()  + TagMemoryUsage.GetLooseSize() + DependNodesSize + DependenciesSize + PackageDataSize + MapArrayMemory;
+	SIZE_T TotalBytes = MapMemory + AssetDataSize + AssetBundlesSize + TagMemoryUsage.GetFixedSize()
+		+ TagMemoryUsage.GetLooseSize() + DependNodesSize + DependenciesSize + PackageDataSize + MapArrayMemory;
 
 	if (bLogDetailed)
 	{
@@ -1895,8 +1980,12 @@ FDependsNode* FAssetRegistryState::ResolveRedirector(FDependsNode* InDependency,
 				if (Asset->IsRedirector())
 				{
 					FDependsNode* ChainedRedirector = nullptr;
-					// This asset is a redirector, so we want to look at its dependencies and find the asset that it is redirecting to
-					CurrentDependency->IterateOverDependencies([&](FDependsNode* InDepends, UE::AssetRegistry::EDependencyCategory Category, UE::AssetRegistry::EDependencyProperty Property, bool bDuplicate) {
+					// This asset is a redirector, so we want to look at its dependencies and find the asset that
+					// it is redirecting to
+					CurrentDependency->IterateOverDependencies([&](FDependsNode* InDepends,
+						UE::AssetRegistry::EDependencyCategory Category,
+						UE::AssetRegistry::EDependencyProperty Property, bool bDuplicate)
+					{
 						if (bDuplicate)
 						{
 							return; // Already looked at this dependency node
@@ -1911,8 +2000,9 @@ FDependsNode* FAssetRegistryState::ResolveRedirector(FDependsNode* InDependency,
 						else if (CachedAssetsByPackageName.Contains(InDepends->GetPackageName()))
 						{
 							// This dependency isn't in the allowed list, but it is a valid asset in the registry.
-							// Because this is a redirector, this should mean that the redirector is pointing at ANOTHER
-							// redirector (or itself in some horrible situations) so we'll move to that node and try again
+							// Because this is a redirector, this should mean that the redirector is pointing at
+							// ANOTHER redirector (or itself in some horrible situations) so we'll move to that node
+							// and try again
 							ChainedRedirector = InDepends;
 						}
 					}, UE::AssetRegistry::EDependencyCategory::Package);
@@ -1958,7 +2048,8 @@ void ShrinkMultimap(TMap<KeyType, TArray<FAssetData*> >& Map)
 
 void FAssetRegistryState::SetAssetDatas(TArrayView<FAssetData> AssetDatas, const FAssetRegistryLoadOptions& Options)
 {
-	UE_CLOG(NumAssets != 0, LogAssetRegistry, Fatal, TEXT("Can only load into empty asset registry states. Load into temporary and append using InitializeFromExisting() instead."));
+	UE_CLOG(NumAssets != 0, LogAssetRegistry, Fatal,
+		TEXT("Can only load into empty asset registry states. Load into temporary and append using InitializeFromExisting() instead."));
 
 	NumAssets = AssetDatas.Num();
 	
@@ -1979,7 +2070,8 @@ void FAssetRegistryState::SetAssetDatas(TArrayView<FAssetData> AssetDatas, const
 		CachedAssetsByPackageName.Empty(AssetDatas.Num());
 		for (FAssetData& AssetData : AssetDatas)
 		{
-			TArray<FAssetData*, TInlineAllocator<1>>& PackageAssets = CachedAssetsByPackageName.FindOrAdd(AssetData.PackageName);
+			TArray<FAssetData*, TInlineAllocator<1>>& PackageAssets
+				= CachedAssetsByPackageName.FindOrAdd(AssetData.PackageName);
 			PackageAssets.Add(&AssetData);
 		}
 	};
@@ -2076,7 +2168,9 @@ void FAssetRegistryState::AddTagsToAssetData(const FSoftObjectPath& InObjectPath
 	FSetElementId Id = CachedAssets.FindId(FCachedAssetKey(InObjectPath));
 	if (!Id.IsValidId())
 	{
-		UE_LOG(LogAssetRegistry, Warning, TEXT("AddTagsToAssetData called with asset data that doesn't exist! Tags not added. ObjectPath: %s"), *InObjectPath.ToString());
+		UE_LOG(LogAssetRegistry, Warning,
+			TEXT("AddTagsToAssetData called with asset data that doesn't exist! Tags not added. ObjectPath: %s"),
+			*InObjectPath.ToString());
 		return;
 	}
 	FAssetData* AssetData = CachedAssets[Id];
@@ -2199,7 +2293,8 @@ void FAssetRegistryState::UpdateAssetData(FAssetData* AssetData, FAssetData&& Ne
 		check(NumRemoved <= 1);
 		if (NumRemoved == 0)
 		{
-			UE_LOG(LogAssetRegistry, Error, TEXT("UpdateAssetData called on AssetData %s that is not present in the AssetRegistry."),
+			UE_LOG(LogAssetRegistry, Error,
+				TEXT("UpdateAssetData called on AssetData %s that is not present in the AssetRegistry."),
 				*AssetData->GetObjectPathString());
 		}
 		NumAssets -= NumRemoved;
@@ -2209,8 +2304,10 @@ void FAssetRegistryState::UpdateAssetData(FAssetData* AssetData, FAssetData&& Ne
 	if (AssetData->PackageName != NewAssetData.PackageName)
 	{
 		bKeyFieldIsModified = true;
-		TArray<FAssetData*, TInlineAllocator<1>>& NewPackageAssets = CachedAssetsByPackageName.FindOrAdd(NewAssetData.PackageName);
-		TArray<FAssetData*, TInlineAllocator<1>>* OldPackageAssets = CachedAssetsByPackageName.Find(AssetData->PackageName);
+		TArray<FAssetData*, TInlineAllocator<1>>& NewPackageAssets
+			= CachedAssetsByPackageName.FindOrAdd(NewAssetData.PackageName);
+		TArray<FAssetData*, TInlineAllocator<1>>* OldPackageAssets
+			= CachedAssetsByPackageName.Find(AssetData->PackageName);
 
 		OldPackageAssets->Remove(AssetData);
 		NewPackageAssets.Add(AssetData);
@@ -2285,8 +2382,10 @@ void FAssetRegistryState::UpdateAssetData(FAssetData* AssetData, FAssetData&& Ne
 			!AssetData->HasSameChunkIDs(NewAssetData) ||
 			(AssetData->TaggedAssetBundles.IsValid() != NewAssetData.TaggedAssetBundles.IsValid() ||
 				(AssetData->TaggedAssetBundles.IsValid() &&
-					AssetData->TaggedAssetBundles.Get() != NewAssetData.TaggedAssetBundles.Get() && // First check whether the pointers are the same
-					*AssetData->TaggedAssetBundles != *NewAssetData.TaggedAssetBundles // If the pointers differ, check whether the contents differ
+					// First check whether the pointers are the same
+					AssetData->TaggedAssetBundles.Get() != NewAssetData.TaggedAssetBundles.Get()
+					// If the pointers differ, check whether the contents differ
+					&& *AssetData->TaggedAssetBundles != *NewAssetData.TaggedAssetBundles
 					));
 	}
 
@@ -2300,7 +2399,8 @@ void FAssetRegistryState::UpdateAssetData(FAssetData* AssetData, FAssetData&& Ne
 		CachedAssets.Add(AssetData, &bExisting);
 		if (bExisting)
 		{
-			UE_LOG(LogAssetRegistry, Error, TEXT("UpdateAssetData called with a change in ObjectPath from Old=\"%s\" to New=\"%s\", ")
+			UE_LOG(LogAssetRegistry, Error,
+				TEXT("UpdateAssetData called with a change in ObjectPath from Old=\"%s\" to New=\"%s\", ")
 				TEXT("but the new ObjectPath is already present with another AssetData. This will overwrite and leak the existing AssetData."),
 				*OldKey.ToString(), *NewKey.ToString());
 }
@@ -2335,10 +2435,12 @@ void FAssetRegistryState::RemoveAssetData(FAssetData* AssetData, bool bRemoveDep
 		bOutRemovedPackageData = false;
 		return;
 	}
-	RemoveAssetData(AssetData, FCachedAssetKey(AssetData), bRemoveDependencyData, bOutRemovedAssetData, bOutRemovedPackageData);
+	RemoveAssetData(AssetData, FCachedAssetKey(AssetData), bRemoveDependencyData, bOutRemovedAssetData,
+		bOutRemovedPackageData);
 	if (!bOutRemovedAssetData)
 	{
-		UE_LOG(LogAssetRegistry, Error, TEXT("RemoveAssetData called on AssetData %s that is not present in the AssetRegistry."),
+		UE_LOG(LogAssetRegistry, Error,
+			TEXT("RemoveAssetData called on AssetData %s that is not present in the AssetRegistry."),
 			*FCachedAssetKey(*AssetData).ToString());
 	}
 }
@@ -2369,7 +2471,8 @@ void FAssetRegistryState::RemoveAssetData(FAssetData* AssetData, const FCachedAs
 		return;
 	}
 
-	TArray<FAssetData*, TInlineAllocator<1>>* OldPackageAssets = CachedAssetsByPackageName.Find(AssetData->PackageName);
+	TArray<FAssetData*, TInlineAllocator<1>>* OldPackageAssets
+		= CachedAssetsByPackageName.Find(AssetData->PackageName);
 	TArray<FAssetData*>* OldPathAssets = CachedAssetsByPath.Find(AssetData->PackagePath);
 	TArray<FAssetData*>* OldClassAssets = CachedAssetsByClass.Find(AssetData->AssetClassPath);
 	OldPackageAssets->RemoveSingleSwap(AssetData);
@@ -2399,7 +2502,8 @@ void FAssetRegistryState::RemoveAssetData(FAssetData* AssetData, const FCachedAs
 		bOutRemovedPackageData = true;
 	}
 
-	// if the assets were preallocated in a block, we can't delete them one at a time, only the whole chunk in the destructor
+	// if the assets were preallocated in a block, we can't delete them one at a time,
+	// only the whole chunk in the destructor
 	if (PreallocatedAssetDataBuffers.Num() == 0)
 	{
 		delete AssetData;
@@ -2463,7 +2567,8 @@ bool FAssetRegistryState::RemoveDependsNode(const FAssetIdentifier& Identifier)
 			CachedDependsNodes.Remove(Identifier);
 			NumDependsNodes--;
 
-			// if the depends nodes were preallocated in a block, we can't delete them one at a time, only the whole chunk in the destructor
+			// if the depends nodes were preallocated in a block, we can't delete them one at a time,
+			// only the whole chunk in the destructor
 			if (PreallocatedDependsNodeDataBuffers.Num() == 0)
 			{
 				delete Node;
@@ -2503,11 +2608,13 @@ FAssetPackageData* FAssetRegistryState::GetAssetPackageData(FName PackageName)
 	return FoundData ? *FoundData : nullptr;
 }
 
-const FAssetPackageData* FAssetRegistryState::GetAssetPackageData(FName PackageName, FName& OutCorrectCasePackageName) const
+const FAssetPackageData* FAssetRegistryState::GetAssetPackageData(FName PackageName,
+	FName& OutCorrectCasePackageName) const
 {
 	// CachedPackageData is keyed using the Package Names whose casing matches the filesystem. In order to perform a
-	// single look up for the AssetPackageData while also returning the value of the key used to add to the map originally
-	// we create a KeyIterator which is currently the only means to get a TPair<Key,Value> from a TMap<Key,Value>
+	// single look up for the AssetPackageData while also returning the value of the key used to add to the map
+	// originally we create a KeyIterator which is currently the only means to get a TPair<Key,Value> from a
+	// TMap<Key,Value>
 	TMap<FName, FAssetPackageData*>::TConstKeyIterator It = CachedPackageData.CreateConstKeyIterator(PackageName);
 	FSetElementId Id = It.GetId();
 	if (!Id.IsValidId())
@@ -2543,7 +2650,8 @@ bool FAssetRegistryState::RemovePackageData(FName PackageName)
 			CachedPackageData.Remove(PackageName);
 			NumPackageData--;
 
-			// if the package data was preallocated in a block, we can't delete them one at a time, only the whole chunk in the destructor
+			// if the package data was preallocated in a block, we can't delete them one at a time,
+			// only the whole chunk in the destructor
 			if (PreallocatedPackageDataBuffers.Num() == 0)
 			{
 				delete Data;
@@ -2577,8 +2685,9 @@ bool FAssetRegistryState::IsPackageUnmountedAndFiltered(const FName PackageName)
 {
 	// TODO: This can be removed once UE-178174 is fixed, as there will no longer be unmounted content to enumerate
 #if WITH_EDITOR
-	// Note: We currently only perform this filtering in the editor; runtime use will have to perform its own filtering 
-	//       via FPackageName::IsValidPath so that it can choose to accept the additional cost of running that filter
+	// Note: We currently only perform this filtering in the editor; runtime use will have to perform its own
+	//       filtering via FPackageName::IsValidPath so that it can choose to accept the additional cost of running
+	//       that filter
 	return bCookedGlobalAssetRegistryState && GIsEditor && !FPackageName::IsValidPath(WriteToString<256>(PackageName));
 #else
 	return false;
@@ -2606,45 +2715,48 @@ bool IsFilterValid(const FARCompiledFilter& Filter)
 }
 
 #if ASSET_REGISTRY_STATE_DUMPING_ENABLED
-namespace UE
+namespace UE::AssetRegistry
 {
-namespace AssetRegistry
+
+void PropertiesToString(EDependencyProperty Properties, FStringBuilderBase& Builder, EDependencyCategory CategoryFilter)
 {
-	void PropertiesToString(EDependencyProperty Properties, FStringBuilderBase& Builder, EDependencyCategory CategoryFilter)
+	bool bFirst = true;
+	auto AppendPropertyName = [&Properties, &Builder, &bFirst](EDependencyProperty TestProperty,
+		const TCHAR* NameWith, const TCHAR* NameWithout)
 	{
-		bool bFirst = true;
-		auto AppendPropertyName = [&Properties, &Builder, &bFirst](EDependencyProperty TestProperty, const TCHAR* NameWith, const TCHAR* NameWithout)
+		if (!bFirst)
 		{
-			if (!bFirst)
-			{
-				Builder.Append(TEXT(","));
-			}
-			if (!!(Properties & TestProperty))
-			{
-				Builder.Append(NameWith);
-			}
-			else
-			{
-				Builder.Append(NameWithout);
-			}
-			bFirst = false;
-		};
-		if (!!(CategoryFilter & EDependencyCategory::Package))
+			Builder.Append(TEXT(","));
+		}
+		if (!!(Properties & TestProperty))
 		{
-			AppendPropertyName(EDependencyProperty::Hard, TEXT("Hard"), TEXT("Soft"));
-			AppendPropertyName(EDependencyProperty::Game, TEXT("Game"), TEXT("EditorOnly"));
-			AppendPropertyName(EDependencyProperty::Build, TEXT("Build"), TEXT("NotBuild"));
+			Builder.Append(NameWith);
 		}
-		if (!!(CategoryFilter & EDependencyCategory::Manage))
+		else
 		{
-			AppendPropertyName(EDependencyProperty::Direct, TEXT("Direct"), TEXT("Indirect"));
+			Builder.Append(NameWithout);
 		}
-		static_assert((EDependencyProperty::PackageMask | EDependencyProperty::SearchableNameMask | EDependencyProperty::ManageMask) == EDependencyProperty::AllMask,
-			"Need to handle new flags in this function");
+		bFirst = false;
+	};
+	if (!!(CategoryFilter & EDependencyCategory::Package))
+	{
+		AppendPropertyName(EDependencyProperty::Hard, TEXT("Hard"), TEXT("Soft"));
+		AppendPropertyName(EDependencyProperty::Game, TEXT("Game"), TEXT("EditorOnly"));
+		AppendPropertyName(EDependencyProperty::Build, TEXT("Build"), TEXT("NotBuild"));
 	}
-}
+	if (!!(CategoryFilter & EDependencyCategory::Manage))
+	{
+		AppendPropertyName(EDependencyProperty::Direct, TEXT("Direct"), TEXT("Indirect"));
+	}
+	static_assert((EDependencyProperty::PackageMask
+		| EDependencyProperty::SearchableNameMask
+		| EDependencyProperty::ManageMask)
+		== EDependencyProperty::AllMask,
+		"Need to handle new flags in this function");
 }
 
+} // namespace UE::AssetRegistry
+
 bool PrintAssetDataMapKeyIsLess(FName A, FName B)
 {
 	return A.Compare(B) < 0;
@@ -2670,7 +2782,8 @@ struct FPrintAssetDataMapKeyIsLess
 };
 
 template <typename MapType>
-static void PrintAssetDataMap(FString Name, const MapType& AssetMap, TStringBuilder<16>& PageBuffer, const TFunctionRef<void()>& AddLine,
+static void PrintAssetDataMap(FString Name, const MapType& AssetMap, TStringBuilder<16>& PageBuffer,
+	const TFunctionRef<void()>& AddLine,
 	TUniqueFunction<void(const typename MapType::KeyType& Key, const FAssetData& Data)>&& PrintValue = {})
 {
 	PageBuffer.Appendf(TEXT("--- Begin %s ---"), *Name);
@@ -2739,28 +2852,35 @@ void FAssetRegistryState::Dump(const TArray<FString>& Arguments, TArray<FString>
 	LinesPerPage = FMath::Max(0, LinesPerPage);
 	const int32 ExpectedNumPages = LinesPerPage > 0 ? (ExpectedNumLines / LinesPerPage) : 1;
 	const int32 PageEndSearchLength = FMath::Min(LinesPerPage, ExpectedNumLines) / 20;
-	const uint32 HashStartValue = MAX_uint32 - 49979693; // Pick a large starting value to bias against picking empty string
+	// Pick a large starting value to bias against picking empty string
+	const uint32 HashStartValue = MAX_uint32 - 49979693; 
 	const uint32 HashMultiplier = 67867967;
 	TStringBuilder<16> PageBuffer;
 	TStringBuilder<16> OverflowText;
 
 	OutPages.Reserve(ExpectedNumPages);
-	PageBuffer.AddUninitialized(FMath::Min(LinesPerPage, ExpectedNumLines) * EstimatedCharactersPerLine);// TODO: Add Reserve function to TStringBuilder
+	// TODO: Add Reserve function to TStringBuilder
+	PageBuffer.AddUninitialized(FMath::Min(LinesPerPage, ExpectedNumLines) * EstimatedCharactersPerLine);
 	PageBuffer.Reset();
 	OverflowText.AddUninitialized(PageEndSearchLength * EstimatedCharactersPerLine);
 	OverflowText.Reset();
 	int32 NumLinesInPage = 0;
 	const int32 LineTerminatorLen = TCString<TCHAR>::Strlen(LINE_TERMINATOR);
 
-	auto FinishPage = [&PageBuffer, &NumLinesInPage, HashStartValue, HashMultiplier, PageEndSearchLength, &OutPages, &OverflowText, LineTerminatorLen](bool bManualPageBreak)
+	auto FinishPage = [&PageBuffer, &NumLinesInPage, HashStartValue, HashMultiplier, PageEndSearchLength,
+		&OutPages, &OverflowText, LineTerminatorLen](bool bManualPageBreak)
 	{
 		int32 PageEndIndex = PageBuffer.Len();
 		const TCHAR* BufferEnd = PageBuffer.GetData() + PageEndIndex;
 		int32 NumOverflowLines = 0;
-		// We want to facilitate diffing dumps between two different versions that should be similar, but naively breaking up the dump into pages makes this difficult
-		// because after one missing or added line, every page from that point on will be offset and therefore different, making false positive differences
-		// To make pages after one missing or added line the same, we look for a good page ending based on the text of all the lines near the end of the current page
-		// By choosing specific-valued texts as page breaks, we will usually randomly get lucky and have the two diffs pick the same line for the end of the page
+		// We want to facilitate diffing dumps between two different versions that should be similar,
+		// but naively breaking up the dump into pages makes this difficult.
+		// because after one missing or added line, every page from that point on will be offset and
+		// therefore different, making false positive differences.
+		// To make pages after one missing or added line the same, we look for a good page ending based 
+		// on the text of all the lines near the end of the current page.
+		// By choosing specific-valued texts as page breaks, we will usually randomly get lucky and
+		// have the two diffs pick the same line for the end of the page.
 		if (!bManualPageBreak && NumLinesInPage > PageEndSearchLength)
 		{
 			const TCHAR* WinningLineEnd = BufferEnd;
@@ -2771,12 +2891,15 @@ void FAssetRegistryState::Dump(const TArray<FString>& Arguments, TArray<FString>
 			{
 				uint32 LineValue = HashStartValue;
 				const TCHAR* LineStart = LineEnd;
-				while (LineStart[-LineTerminatorLen] != LINE_TERMINATOR[0] || TCString<TCHAR>::Strncmp(LINE_TERMINATOR, LineStart - LineTerminatorLen, LineTerminatorLen) != 0)
+				while (LineStart[-LineTerminatorLen] != LINE_TERMINATOR[0]
+					|| TCString<TCHAR>::Strncmp(LINE_TERMINATOR,
+						LineStart - LineTerminatorLen, LineTerminatorLen) != 0)
 				{
 					--LineStart;
 					LineValue = LineValue * HashMultiplier + static_cast<uint32>(TChar<TCHAR>::ToLower(*LineStart));
 				}
-				if (SearchIndex == 0 || LineValue < WinningLineValue) // We arbitrarily choose the smallest hash as the winning value
+				// We arbitrarily choose the smallest hash as the winning value
+				if (SearchIndex == 0 || LineValue < WinningLineValue)
 				{
 					WinningLineValue = LineValue;
 					WinningLineEnd = LineEnd;
@@ -2849,7 +2972,8 @@ void FAssetRegistryState::Dump(const TArray<FString>& Arguments, TArray<FString>
 			Keys.Emplace(FCachedAssetKey(*AssetData));
 		}
 		Keys.Sort([](const FCachedAssetKey& A, const FCachedAssetKey& B) {
-			return WriteToString<1024>(A).ToView().Compare(WriteToString<1024>(B).ToView(), ESearchCase::IgnoreCase) < 0;
+			return WriteToString<1024>(A).ToView().Compare(WriteToString<1024>(B).ToView(),
+				ESearchCase::IgnoreCase) < 0;
 		});
 
 		for (const FCachedAssetKey& Key : Keys)
@@ -2881,7 +3005,7 @@ void FAssetRegistryState::Dump(const TArray<FString>& Arguments, TArray<FString>
 		PrintAssetDataMap(TEXT("CachedAssetsByClass"), CachedAssetsByClass, PageBuffer, AddLine);
 	}
 
-	// Only print this if it's requested specifically - '-all' will print tags-per-asset rather than assets-per-tag 
+	// Only print this if it's requested specifically - '-all' will print tags-per-asset rather than assets-per-tag.
 	if (Arguments.Contains(TEXT("Tag")))
 	{
 		AddPageBreak();
@@ -2974,16 +3098,43 @@ void FAssetRegistryState::Dump(const TArray<FString>& Arguments, TArray<FString>
 		PageBuffer.Append(TEXT("--- Begin CachedDependsNodes ---"));
 		AddLine();
 
-		auto SortByAssetID = [](const FDependsNode& A, const FDependsNode& B) { return A.GetIdentifier().ToString() < B.GetIdentifier().ToString(); };
+		auto SortByAssetID = [](const FDependsNode& A, const FDependsNode& B)
+			{
+				return A.GetIdentifier().ToString() < B.GetIdentifier().ToString();
+			};
 		TArray<FDependsNode*> Nodes;
 		CachedDependsNodes.GenerateValueArray(Nodes);
 		Nodes.Sort(SortByAssetID);
 
 		if (Arguments.Contains(TEXT("LegacyDependencies"))) // LegacyDependencies are not show by all; they have to be directly requested
 		{
-			EDependencyCategory CategoryTypes[] = { EDependencyCategory::Package, EDependencyCategory::Package,EDependencyCategory::SearchableName,EDependencyCategory::Manage, EDependencyCategory::Manage, EDependencyCategory::None };
-			EDependencyQuery CategoryQueries[] = { EDependencyQuery::Hard, EDependencyQuery::Soft, EDependencyQuery::NoRequirements, EDependencyQuery::Direct, EDependencyQuery::Indirect, EDependencyQuery::NoRequirements };
-			const TCHAR* CategoryNames[] = { TEXT("Hard"), TEXT("Soft"), TEXT("SearchableName"), TEXT("HardManage"), TEXT("SoftManage"), TEXT("References") };
+			EDependencyCategory CategoryTypes[] =
+			{
+				EDependencyCategory::Package,
+				EDependencyCategory::Package,
+				EDependencyCategory::SearchableName,
+				EDependencyCategory::Manage,
+				EDependencyCategory::Manage,
+				EDependencyCategory::None
+			};
+			EDependencyQuery CategoryQueries[] =
+			{
+				EDependencyQuery::Hard,
+				EDependencyQuery::Soft,
+				EDependencyQuery::NoRequirements,
+				EDependencyQuery::Direct,
+				EDependencyQuery::Indirect,
+				EDependencyQuery::NoRequirements
+			};
+			const TCHAR* CategoryNames[] =
+			{
+				TEXT("Hard"),
+				TEXT("Soft"),
+				TEXT("SearchableName"),
+				TEXT("HardManage"),
+				TEXT("SoftManage"),
+				TEXT("References")
+			};
 			const int NumCategories = UE_ARRAY_COUNT(CategoryTypes);
 			check(NumCategories == UE_ARRAY_COUNT(CategoryNames) && NumCategories == UE_ARRAY_COUNT(CategoryQueries));
 
@@ -3024,8 +3175,20 @@ void FAssetRegistryState::Dump(const TArray<FString>& Arguments, TArray<FString>
 		}
 		else
 		{
-			EDependencyCategory CategoryTypes[] = { EDependencyCategory::Package, EDependencyCategory::SearchableName,EDependencyCategory::Manage, EDependencyCategory::None };
-			const TCHAR* CategoryNames[] = { TEXT("Package"), TEXT("SearchableName"), TEXT("Manage"), TEXT("References") };
+			EDependencyCategory CategoryTypes[] =
+			{
+				EDependencyCategory::Package,
+				EDependencyCategory::SearchableName,
+				EDependencyCategory::Manage,
+				EDependencyCategory::None
+			};
+			const TCHAR* CategoryNames[] =
+			{
+				TEXT("Package"),
+				TEXT("SearchableName"),
+				TEXT("Manage"),
+				TEXT("References")
+			};
 			const int NumCategories = UE_ARRAY_COUNT(CategoryTypes);
 			check(NumCategories == UE_ARRAY_COUNT(CategoryNames));
 
@@ -3133,7 +3296,8 @@ void FAssetRegistryState::Dump(const TArray<FString>& Arguments, TArray<FString>
 				++Counter;
 				for (const FAssetBundleEntry& Entry : AssetData->TaggedAssetBundles->Bundles)
 				{
-					PageBuffer << TEXT("  Owner: ") << FCachedAssetKey(AssetData) << TEXT(" BundleName: ") << Entry.BundleName;
+					PageBuffer << TEXT("  Owner: ")
+						<< FCachedAssetKey(AssetData) << TEXT(" BundleName: ") << Entry.BundleName;
 					AddLine();
 
 					for (const FTopLevelAssetPath& Path : Entry.AssetPaths)
@@ -3171,7 +3335,8 @@ void FAssetRegistryState::Dump(const TArray<FString>& Arguments, TArray<FString>
 
 #include "Misc/AutomationTest.h"
 
-IMPLEMENT_SIMPLE_AUTOMATION_TEST(FAssetRegistryAssetPathStringsTest, "System.AssetRegistry.AssetPathStrings", EAutomationTestFlags_ApplicationContextMask | EAutomationTestFlags::EngineFilter);
+IMPLEMENT_SIMPLE_AUTOMATION_TEST(FAssetRegistryAssetPathStringsTest, "System.AssetRegistry.AssetPathStrings",
+	EAutomationTestFlags_ApplicationContextMask | EAutomationTestFlags::EngineFilter);
 
 // Tests that we can produce correct paths for objects represented by FCachedAssetKey and FAssetData
 // E.g	PackageName.AssetName
@@ -3187,35 +3352,50 @@ bool FAssetRegistryAssetPathStringsTest::RunTest(const FString& Parameters)
 	const FName SubSubObjectOuter = TEXT("/Path/To/PackageName.OuterName:SubOuterName");
 	const FName AssetName = TEXT("AssetName");
 
-	const FString TopLevelPathString = FString::Printf(TEXT("%s.%s"), *TopLevelOuter.ToString(), *AssetName.ToString());
-	const FString DirectSubObjectPathString = FString::Printf(TEXT("%s:%s"), *DirectSubObjectOuter.ToString(), *AssetName.ToString());
-	const FString SubSubObjectPathString = FString::Printf(TEXT("%s.%s"), *SubSubObjectOuter.ToString(), *AssetName.ToString());
+	const FString TopLevelPathString = FString::Printf(TEXT("%s.%s"),
+		*TopLevelOuter.ToString(), *AssetName.ToString());
+	const FString DirectSubObjectPathString = FString::Printf(TEXT("%s:%s"),
+		*DirectSubObjectOuter.ToString(), *AssetName.ToString());
+	const FString SubSubObjectPathString = FString::Printf(TEXT("%s.%s"),
+		*SubSubObjectOuter.ToString(), *AssetName.ToString());
 
 	const FSoftObjectPath TopLevelPath(TopLevelPathString);
 	const FSoftObjectPath DirectSubObjectPath(DirectSubObjectPathString);
 	const FSoftObjectPath SubSubObjectPath(SubSubObjectPathString);
 
-	TestEqual(TEXT("SoftObjectPath::ToString() correct for top-level asset"),		TopLevelPath.ToString(), TopLevelPathString);
-	TestEqual(TEXT("SoftObjectPath::ToString() correct for subobject asset"),		DirectSubObjectPath.ToString(), DirectSubObjectPathString);
-	TestEqual(TEXT("SoftObjectPath::ToString() correct for sub-subobject asset"),	SubSubObjectPath.ToString(), SubSubObjectPathString);
+	TestEqual(TEXT("SoftObjectPath::ToString() correct for top-level asset"),
+		TopLevelPath.ToString(), TopLevelPathString);
+	TestEqual(TEXT("SoftObjectPath::ToString() correct for subobject asset"),
+		DirectSubObjectPath.ToString(), DirectSubObjectPathString);
+	TestEqual(TEXT("SoftObjectPath::ToString() correct for sub-subobject asset"),
+		SubSubObjectPath.ToString(), SubSubObjectPathString);
 
 	// Construct FCachedAssetKey from FSoftObjectPath of various lengths + check they have the right components
 	const FCachedAssetKey TopLevelAssetKey(TopLevelPath);
-	TestEqual(TEXT("FCachedAssetKey::OuterPath correct for top-level asset"), TopLevelAssetKey.OuterPath.ToString(), TopLevelOuter.ToString());
-	TestEqual(TEXT("FCachedAssetKey::ObjectName correct for top-level asset"), TopLevelAssetKey.ObjectName.ToString(), AssetName.ToString());
+	TestEqual(TEXT("FCachedAssetKey::OuterPath correct for top-level asset"),
+		TopLevelAssetKey.OuterPath.ToString(), TopLevelOuter.ToString());
+	TestEqual(TEXT("FCachedAssetKey::ObjectName correct for top-level asset"),
+		TopLevelAssetKey.ObjectName.ToString(), AssetName.ToString());
 
 	const FCachedAssetKey DirectSubObjectKey(DirectSubObjectPath);
-	TestEqual(TEXT("FCachedAssetKey::OuterPath correct for subobject asset"), DirectSubObjectKey.OuterPath.ToString(), DirectSubObjectOuter.ToString());
-	TestEqual(TEXT("FCachedAssetKey::ObjectName correct for subobject asset"), DirectSubObjectKey.ObjectName.ToString(), AssetName.ToString());
+	TestEqual(TEXT("FCachedAssetKey::OuterPath correct for subobject asset"),
+		DirectSubObjectKey.OuterPath.ToString(), DirectSubObjectOuter.ToString());
+	TestEqual(TEXT("FCachedAssetKey::ObjectName correct for subobject asset"),
+		DirectSubObjectKey.ObjectName.ToString(), AssetName.ToString());
 
 	const FCachedAssetKey SubSubObjectKey(SubSubObjectPath);
-	TestEqual(TEXT("FCachedAssetKey::OuterPath correct for sub-subobject asset"), SubSubObjectKey.OuterPath.ToString(), SubSubObjectOuter.ToString());
-	TestEqual(TEXT("FCachedAssetKey::ObjectName correct for sub-subobject asset"), SubSubObjectKey.ObjectName.ToString(), AssetName.ToString());
+	TestEqual(TEXT("FCachedAssetKey::OuterPath correct for sub-subobject asset"),
+		SubSubObjectKey.OuterPath.ToString(), SubSubObjectOuter.ToString());
+	TestEqual(TEXT("FCachedAssetKey::ObjectName correct for sub-subobject asset"),
+		SubSubObjectKey.ObjectName.ToString(), AssetName.ToString());
 
 	// Construct FCachedAssetKey from FSoftObjectPath of various lengths + check they give the right strings from AppendString
-	TestEqual(TEXT("FCachedAssetKey::ToString() correct for top-level asset"),		TopLevelAssetKey.ToString(), TopLevelPathString);
-	TestEqual(TEXT("FCachedAssetKey::ToString() correct for subobject asset"),		DirectSubObjectKey.ToString(), DirectSubObjectPathString);
-	TestEqual(TEXT("FCachedAssetKey::ToString() correct for sub-subobject asset"),	SubSubObjectKey.ToString(), SubSubObjectPathString);
+	TestEqual(TEXT("FCachedAssetKey::ToString() correct for top-level asset"),
+		TopLevelAssetKey.ToString(), TopLevelPathString);
+	TestEqual(TEXT("FCachedAssetKey::ToString() correct for subobject asset"),
+		DirectSubObjectKey.ToString(), DirectSubObjectPathString);
+	TestEqual(TEXT("FCachedAssetKey::ToString() correct for sub-subobject asset"),
+		SubSubObjectKey.ToString(), SubSubObjectPathString);
 
 	auto PathToAssetData = [](const FString& Path)
 	{
@@ -3235,9 +3415,12 @@ bool FAssetRegistryAssetPathStringsTest::RunTest(const FString& Parameters)
 	};
 
 	// Test FAssetData::AppendPath for asset data with variable length of OptionalOuterPath
-	TestEqual(TEXT("FAssetData::AppendPath() correct for top-level asset"),			TopLevelAssetData.GetObjectPathString(), TopLevelPathString);
-	TestEqual(TEXT("FAssetData::AppendPath() correct for subobject asset"),			DirectSubObjectAssetData.GetObjectPathString(), DirectSubObjectPathString);
-	TestEqual(TEXT("FAssetData::AppendPath() correct for sub-subobject asset"),		SubSubObjectAssetData.GetObjectPathString(), SubSubObjectPathString);
+	TestEqual(TEXT("FAssetData::AppendPath() correct for top-level asset"),
+		TopLevelAssetData.GetObjectPathString(), TopLevelPathString);
+	TestEqual(TEXT("FAssetData::AppendPath() correct for subobject asset"),
+		DirectSubObjectAssetData.GetObjectPathString(), DirectSubObjectPathString);
+	TestEqual(TEXT("FAssetData::AppendPath() correct for sub-subobject asset"),
+		SubSubObjectAssetData.GetObjectPathString(), SubSubObjectPathString);
 	
 	return true;
 }
diff --git a/Engine/Source/Runtime/AssetRegistry/Private/DependsNode.cpp b/Engine/Source/Runtime/AssetRegistry/Private/DependsNode.cpp
--- a/Engine/Source/Runtime/AssetRegistry/Private/DependsNode.cpp
+++ b/Engine/Source/Runtime/AssetRegistry/Private/DependsNode.cpp
@@ -29,9 +29,15 @@ void FDependsNode::PrintReferencers() const
 }
 
 template <uint32 FlagWidth>
-void IterateDependencyList(const FDependsNode::FIterateDependenciesCallback& InCallback, UE::AssetRegistry::EDependencyCategory SearchCategory, const UE::AssetRegistry::FDependencyQuery& SearchFlags,
-	UE::AssetRegistry::EDependencyCategory ListCategory, UE::AssetRegistry::EDependencyProperty CategoryMask, const TArray<FDependsNode*>& Dependencies, const TBitArray<>* FlagBits,
-	UE::AssetRegistry::EDependencyProperty(*ByteToProperties)(uint8), bool IsSorted)
+void IterateDependencyList(const FDependsNode::FIterateDependenciesCallback& InCallback,
+	UE::AssetRegistry::EDependencyCategory SearchCategory,
+	const UE::AssetRegistry::FDependencyQuery& SearchFlags,
+	UE::AssetRegistry::EDependencyCategory ListCategory,
+	UE::AssetRegistry::EDependencyProperty CategoryMask,
+	const TArray<FDependsNode*>& Dependencies,
+	const TBitArray<>* FlagBits,
+	UE::AssetRegistry::EDependencyProperty(*ByteToProperties)(uint8),
+	bool IsSorted)
 {
 	using namespace UE::AssetRegistry;
 	if (!(SearchCategory & ListCategory))
@@ -60,7 +66,8 @@ void IterateDependencyList(const FDependsNode::FIterateDependenciesCallback& InC
 			for (uint32 DependencyFlagBits : DependsNodeFlagsSet)
 			{
 				EDependencyProperty DependencyProperties = ByteToProperties(static_cast<uint8>(DependencyFlagBits));
-				if (((DependencyProperties & RequiredProperties) == RequiredProperties) && ((DependencyProperties & ExcludedProperties) == EDependencyProperty::None))
+				if (EnumHasAllFlags(DependencyProperties, RequiredProperties)
+					&& !EnumHasAnyFlags(DependencyProperties, ExcludedProperties))
 				{
 					InCallback(DependsNode, ListCategory, DependencyProperties, bDuplicate);
 				}
@@ -71,9 +78,16 @@ void IterateDependencyList(const FDependsNode::FIterateDependenciesCallback& InC
 }
 
 template <uint32 FlagWidth>
-void IterateDependencyList(const FDependsNode::FIterateDependenciesCallback& InCallback, const FDependsNode* SearchNode, UE::AssetRegistry::EDependencyCategory SearchCategory, const UE::AssetRegistry::FDependencyQuery& SearchFlags,
-	UE::AssetRegistry::EDependencyCategory ListCategory, UE::AssetRegistry::EDependencyProperty CategoryMask, const TArray<FDependsNode*>& Dependencies, const TBitArray<>* FlagBits,
-	UE::AssetRegistry::EDependencyProperty(*ByteToProperties)(uint8), bool IsSorted)
+void IterateDependencyList(const FDependsNode::FIterateDependenciesCallback& InCallback,
+	const FDependsNode* SearchNode,
+	UE::AssetRegistry::EDependencyCategory SearchCategory,
+	const UE::AssetRegistry::FDependencyQuery& SearchFlags,
+	UE::AssetRegistry::EDependencyCategory ListCategory,
+	UE::AssetRegistry::EDependencyProperty CategoryMask,
+	const TArray<FDependsNode*>& Dependencies,
+	const TBitArray<>* FlagBits,
+	UE::AssetRegistry::EDependencyProperty(*ByteToProperties)(uint8),
+	bool IsSorted)
 {
 	using namespace UE::AssetRegistry;
 	if (!(SearchCategory & ListCategory))
@@ -99,7 +113,8 @@ void IterateDependencyList(const FDependsNode::FIterateDependenciesCallback& InC
 			for (uint32 DependencyFlagBits : DependsNodeFlagsSet)
 			{
 				EDependencyProperty DependencyProperties = ByteToProperties(static_cast<uint8>(DependencyFlagBits));
-				if (((DependencyProperties & RequiredProperties) == RequiredProperties) && ((DependencyProperties & ExcludedProperties) == EDependencyProperty::None))
+				if (EnumHasAllFlags(DependencyProperties, RequiredProperties)
+					&& !EnumHasAnyFlags(DependencyProperties, ExcludedProperties))
 				{
 					InCallback(DependsNode, ListCategory, DependencyProperties, bDuplicate);
 				}
@@ -128,25 +143,48 @@ void IterateDependencyList(const FDependsNode::FIterateDependenciesCallback& InC
 	}
 }
 
-void FDependsNode::IterateOverDependencies(const FIterateDependenciesCallback& InCallback, UE::AssetRegistry::EDependencyCategory Category, const UE::AssetRegistry::FDependencyQuery& Flags) const
+void FDependsNode::IterateOverDependencies(const FIterateDependenciesCallback& InCallback,
+	UE::AssetRegistry::EDependencyCategory Category,
+	const UE::AssetRegistry::FDependencyQuery& Flags) const
 {
 	using namespace UE::AssetRegistry;
-	IterateDependencyList<PackageFlagWidth>(InCallback, Category, Flags, EDependencyCategory::Package, EDependencyProperty::PackageMask, PackageDependencies, &PackageFlags, ByteToPackageProperties, PackageIsSorted);
-	IterateDependencyList<SearchableNameFlagWidth>(InCallback, Category, Flags, EDependencyCategory::SearchableName, EDependencyProperty::SearchableNameMask, NameDependencies, nullptr, nullptr, SearchableNameIsSorted);
-	IterateDependencyList<ManageFlagWidth>(InCallback, Category, Flags, EDependencyCategory::Manage, EDependencyProperty::ManageMask, ManageDependencies, &ManageFlags, ByteToManageProperties, ManageIsSorted);
+	IterateDependencyList<PackageFlagWidth>(InCallback, Category, Flags, EDependencyCategory::Package,
+		EDependencyProperty::PackageMask, PackageDependencies, &PackageFlags,
+		ByteToPackageProperties, PackageIsSorted);
+	IterateDependencyList<SearchableNameFlagWidth>(InCallback, Category, Flags, EDependencyCategory::SearchableName,
+		EDependencyProperty::SearchableNameMask, NameDependencies, nullptr,
+		nullptr, SearchableNameIsSorted);
+	IterateDependencyList<ManageFlagWidth>(InCallback, Category, Flags, EDependencyCategory::Manage,
+		EDependencyProperty::ManageMask, ManageDependencies, &ManageFlags,
+		ByteToManageProperties, ManageIsSorted);
 }
 
-void FDependsNode::IterateOverDependencies(const FIterateDependenciesCallback& InCallback, const FDependsNode* DependsNode, UE::AssetRegistry::EDependencyCategory Category, const UE::AssetRegistry::FDependencyQuery& Flags) const
+void FDependsNode::IterateOverDependencies(const FIterateDependenciesCallback& InCallback,
+	const FDependsNode* DependsNode,
+	UE::AssetRegistry::EDependencyCategory Category,
+	const UE::AssetRegistry::FDependencyQuery& Flags) const
 {
 	using namespace UE::AssetRegistry;
-	IterateDependencyList<PackageFlagWidth>(InCallback, DependsNode, Category, Flags, EDependencyCategory::Package, EDependencyProperty::PackageMask, PackageDependencies, &PackageFlags, ByteToPackageProperties, PackageIsSorted);
-	IterateDependencyList<SearchableNameFlagWidth>(InCallback, DependsNode, Category, Flags, EDependencyCategory::SearchableName, EDependencyProperty::SearchableNameMask, NameDependencies, nullptr, nullptr, SearchableNameIsSorted);
-	IterateDependencyList<ManageFlagWidth>(InCallback, DependsNode, Category, Flags, EDependencyCategory::Manage, EDependencyProperty::ManageMask, ManageDependencies, &ManageFlags, ByteToManageProperties, ManageIsSorted);
+	IterateDependencyList<PackageFlagWidth>(InCallback, DependsNode, Category, Flags, EDependencyCategory::Package,
+		EDependencyProperty::PackageMask, PackageDependencies, &PackageFlags,
+		ByteToPackageProperties, PackageIsSorted);
+	IterateDependencyList<SearchableNameFlagWidth>(InCallback, DependsNode, Category, Flags,
+		EDependencyCategory::SearchableName,
+		EDependencyProperty::SearchableNameMask, NameDependencies, nullptr, 
+		nullptr, SearchableNameIsSorted);
+	IterateDependencyList<ManageFlagWidth>(InCallback, DependsNode, Category, Flags, EDependencyCategory::Manage,
+		EDependencyProperty::ManageMask, ManageDependencies, &ManageFlags,
+		ByteToManageProperties, ManageIsSorted);
 }
 
-void FDependsNode::GetDependencies(TArray<FDependsNode*>& OutDependencies, UE::AssetRegistry::EDependencyCategory Category, const UE::AssetRegistry::FDependencyQuery& Flags) const
+void FDependsNode::GetDependencies(TArray<FDependsNode*>& OutDependencies,
+	UE::AssetRegistry::EDependencyCategory Category,
+	const UE::AssetRegistry::FDependencyQuery& Flags) const
 {
-	IterateOverDependencies([&OutDependencies](FDependsNode* InDependency, UE::AssetRegistry::EDependencyCategory InCategory, UE::AssetRegistry::EDependencyProperty InProperties, bool bDuplicate)
+	IterateOverDependencies([&OutDependencies](FDependsNode* InDependency,
+		UE::AssetRegistry::EDependencyCategory InCategory,
+		UE::AssetRegistry::EDependencyProperty InProperties,
+		bool bDuplicate)
 	{
 		if (!bDuplicate)
 		{
@@ -156,9 +194,13 @@ void FDependsNode::GetDependencies(TArray<FDependsNode*>& OutDependencies, UE::A
 	Category, Flags);
 }
 
-void FDependsNode::GetDependencies(TArray<FAssetIdentifier>& OutDependencies, UE::AssetRegistry::EDependencyCategory Category, const UE::AssetRegistry::FDependencyQuery& Flags) const
+void FDependsNode::GetDependencies(TArray<FAssetIdentifier>& OutDependencies,
+	UE::AssetRegistry::EDependencyCategory Category, const UE::AssetRegistry::FDependencyQuery& Flags) const
 {
-	IterateOverDependencies([&OutDependencies](const FDependsNode* InDependency, UE::AssetRegistry::EDependencyCategory InCategory, UE::AssetRegistry::EDependencyProperty InProperties, bool bDuplicate)
+	IterateOverDependencies([&OutDependencies](const FDependsNode* InDependency,
+		UE::AssetRegistry::EDependencyCategory InCategory,
+		UE::AssetRegistry::EDependencyProperty InProperties,
+		bool bDuplicate)
 	{
 		if (!bDuplicate)
 		{
@@ -168,24 +210,34 @@ void FDependsNode::GetDependencies(TArray<FAssetIdentifier>& OutDependencies, UE
 	Category, Flags);
 }
 
-void FDependsNode::GetDependencies(TArray<FAssetDependency>& OutDependencies, UE::AssetRegistry::EDependencyCategory Category, const UE::AssetRegistry::FDependencyQuery& Flags) const
+void FDependsNode::GetDependencies(TArray<FAssetDependency>& OutDependencies,
+	UE::AssetRegistry::EDependencyCategory Category, const UE::AssetRegistry::FDependencyQuery& Flags) const
 {
-	IterateOverDependencies([&OutDependencies](const FDependsNode* InDependency, UE::AssetRegistry::EDependencyCategory InCategory, UE::AssetRegistry::EDependencyProperty InProperties, bool bDuplicate)
+	IterateOverDependencies([&OutDependencies](const FDependsNode* InDependency,
+		UE::AssetRegistry::EDependencyCategory InCategory,
+		UE::AssetRegistry::EDependencyProperty InProperties,
+		bool bDuplicate)
 	{
 		OutDependencies.Add(FAssetDependency{ InDependency->GetIdentifier(), InCategory, InProperties });
 	},
 	Category, Flags);
 }
 
-void FDependsNode::GetReferencers(TArray<FDependsNode*>& OutReferencers, UE::AssetRegistry::EDependencyCategory Category, const UE::AssetRegistry::FDependencyQuery& Flags) const
+void FDependsNode::GetReferencers(TArray<FDependsNode*>& OutReferencers,
+	UE::AssetRegistry::EDependencyCategory Category, const UE::AssetRegistry::FDependencyQuery& Flags) const
 {
 	for (FDependsNode* Referencer : Referencers)
 	{
 		bool bShouldAdd = false;
 		// If type specified, filter
-		if (Category != UE::AssetRegistry::EDependencyCategory::All || Flags.Required != UE::AssetRegistry::EDependencyProperty::None || Flags.Excluded != UE::AssetRegistry::EDependencyProperty::None)
+		if (Category != UE::AssetRegistry::EDependencyCategory::All
+			|| Flags.Required != UE::AssetRegistry::EDependencyProperty::None
+			|| Flags.Excluded != UE::AssetRegistry::EDependencyProperty::None)
 		{
-			Referencer->IterateOverDependencies([&bShouldAdd](const FDependsNode* InDependency, UE::AssetRegistry::EDependencyCategory InCategory, UE::AssetRegistry::EDependencyProperty InProperties, bool bDuplicate)
+			Referencer->IterateOverDependencies([&bShouldAdd](const FDependsNode* InDependency,
+				UE::AssetRegistry::EDependencyCategory InCategory,
+				UE::AssetRegistry::EDependencyProperty InProperties,
+				bool bDuplicate)
 				{
 					bShouldAdd = true;
 				}, this, Category, Flags);
@@ -202,11 +254,15 @@ void FDependsNode::GetReferencers(TArray<FDependsNode*>& OutReferencers, UE::Ass
 	}
 }
 
-void FDependsNode::GetReferencers(TArray<FAssetDependency>& OutReferencers, UE::AssetRegistry::EDependencyCategory Category, const UE::AssetRegistry::FDependencyQuery& Flags) const
+void FDependsNode::GetReferencers(TArray<FAssetDependency>& OutReferencers,
+	UE::AssetRegistry::EDependencyCategory Category, const UE::AssetRegistry::FDependencyQuery& Flags) const
 {
 	for (FDependsNode* Referencer : Referencers)
 	{
-		Referencer->IterateOverDependencies([&OutReferencers, Referencer](const FDependsNode* InDependency, UE::AssetRegistry::EDependencyCategory InCategory, UE::AssetRegistry::EDependencyProperty InProperties, bool bDuplicate)
+		Referencer->IterateOverDependencies([&OutReferencers, Referencer](const FDependsNode* InDependency,
+			UE::AssetRegistry::EDependencyCategory InCategory,
+			UE::AssetRegistry::EDependencyProperty InProperties,
+			bool bDuplicate)
 			{
 				OutReferencers.Add(FAssetDependency{ Referencer->GetIdentifier(), InCategory, InProperties });
 			}, this, Category, Flags);
@@ -264,27 +320,35 @@ void AddDependency(FDependsNode* InDependency, UE::AssetRegistry::EDependencyPro
 	}
 }
 
-void FDependsNode::AddDependency(FDependsNode* InDependency, UE::AssetRegistry::EDependencyCategory Category, UE::AssetRegistry::EDependencyProperty Properties)
+void FDependsNode::AddDependency(FDependsNode* InDependency,
+	UE::AssetRegistry::EDependencyCategory Category, UE::AssetRegistry::EDependencyProperty Properties)
 {
 	using namespace UE::AssetRegistry;
-	if (!!(Category & UE::AssetRegistry::EDependencyCategory::Package))
+	if (EnumHasAnyFlags(Category, UE::AssetRegistry::EDependencyCategory::Package))
 	{
-		::AddDependency<PackageFlagWidth>(InDependency, Properties, EDependencyProperty::PackageMask, PackageDependencies, &PackageFlags, PackagePropertiesToByte, PackageIsSorted);
-		check((Category & ~EDependencyCategory::Package) == EDependencyCategory::None); // It is illegal to try to add a dependency as more than one category at a time
+		::AddDependency<PackageFlagWidth>(InDependency, Properties, EDependencyProperty::PackageMask,
+			PackageDependencies, &PackageFlags, PackagePropertiesToByte, PackageIsSorted);
+		// It is illegal to try to add a dependency as more than one category at a time
+		check((Category & ~EDependencyCategory::Package) == EDependencyCategory::None);
 	}
-	else if (!!(Category & EDependencyCategory::SearchableName))
+	else if (EnumHasAnyFlags(Category, EDependencyCategory::SearchableName))
 	{
-		::AddDependency<SearchableNameFlagWidth>(InDependency, Properties, EDependencyProperty::SearchableNameMask, NameDependencies, nullptr, nullptr, SearchableNameIsSorted);
-		check((Category & ~EDependencyCategory::SearchableName) == EDependencyCategory::None); // It is illegal to try to add a dependency as more than one category at a time
+		::AddDependency<SearchableNameFlagWidth>(InDependency, Properties, EDependencyProperty::SearchableNameMask,
+			NameDependencies, nullptr, nullptr, SearchableNameIsSorted);
+		// It is illegal to try to add a dependency as more than one category at a time
+		check((Category & ~EDependencyCategory::SearchableName) == EDependencyCategory::None);
 	}
-	else if (!!(Category & UE::AssetRegistry::EDependencyCategory::Manage))
+	else if (EnumHasAnyFlags(Category, UE::AssetRegistry::EDependencyCategory::Manage))
 	{
-		::AddDependency<ManageFlagWidth>(InDependency, Properties, EDependencyProperty::ManageMask, ManageDependencies, &ManageFlags, ManagePropertiesToByte, ManageIsSorted);
-		check((Category & ~EDependencyCategory::Manage) == EDependencyCategory::None); // It is illegal to try to add a dependency as more than one category at a time
+		::AddDependency<ManageFlagWidth>(InDependency, Properties, EDependencyProperty::ManageMask,
+			ManageDependencies, &ManageFlags, ManagePropertiesToByte, ManageIsSorted);
+		// It is illegal to try to add a dependency as more than one category at a time
+		check((Category & ~EDependencyCategory::Manage) == EDependencyCategory::None);
 	}
 	else
 	{
-		check(false); // It is illegal to try to add a dependency without a category
+		// It is illegal to try to add a dependency without a category
+		check(false);
 	}
 }
 
@@ -302,11 +366,15 @@ void FDependsNode::AddPackageDependencySet(FDependsNode* InDependency, const FPa
 void FDependsNode::AddReferencer(FDependsNode* InReferencer)
 {
 	using namespace UE::AssetRegistry;
-	::AddDependency<0>(InReferencer, EDependencyProperty::None, EDependencyProperty::None, Referencers, nullptr, nullptr, ReferencersIsSorted);
+	::AddDependency<0>(InReferencer, EDependencyProperty::None, EDependencyProperty::None,
+		Referencers, nullptr, nullptr, ReferencersIsSorted);
 }
 
 template <uint32 FlagWidth>
-void RemoveDependency(FDependsNode* InDependency, TArray<FDependsNode*>& Dependencies, TBitArray<>* FlagBits, bool IsSorted)
+void RemoveDependency(FDependsNode* InDependency,
+	TArray<FDependsNode*>& Dependencies,
+	TBitArray<>* FlagBits,
+	bool IsSorted)
 {
 	check(FlagWidth == 0 || FlagBits != nullptr);
 	if (IsSorted)
@@ -373,11 +441,17 @@ void FDependsNode::RefreshReferencers()
 {
 	if (IsReferencersSorted())
 	{
-		Referencers.RemoveAll([this](FDependsNode* Referencer) { return !Referencer->ContainsDependency(this); });
+		Referencers.RemoveAll([this](FDependsNode* Referencer)
+			{
+				return !Referencer->ContainsDependency(this);
+			});
 	}
 	else
 	{
-		Referencers.RemoveAllSwap([this](FDependsNode* Referencer) { return !Referencer->ContainsDependency(this); });
+		Referencers.RemoveAllSwap([this](FDependsNode* Referencer)
+			{
+				return !Referencer->ContainsDependency(this);
+			});
 	}
 }
 
@@ -428,7 +502,10 @@ void FDependsNode::RemoveManageReferencesToNode()
 }
 
 template <uint32 FlagWidth>
-void RemoveAll(const TUniqueFunction<bool(const FDependsNode*)>& ShouldRemove, TArray<FDependsNode*>& Dependencies, TBitArray<>* FlagBits, bool IsSorted)
+void RemoveAll(const TUniqueFunction<bool(const FDependsNode*)>& ShouldRemove,
+	TArray<FDependsNode*>& Dependencies,
+	TBitArray<>* FlagBits,
+	bool IsSorted)
 {
 	check(FlagWidth == 0 || FlagBits != nullptr);
 	if (IsSorted)
@@ -439,7 +516,8 @@ void RemoveAll(const TUniqueFunction<bool(const FDependsNode*)>& ShouldRemove, T
 		}
 		else
 		{
-			// This block is the same functionality as TArray::RemoveAll, but it needs to handle removing the corresponding FlagBits
+			// This block is the same functionality as TArray::RemoveAll,
+			// but it needs to handle removing the corresponding FlagBits
 			const int32 OriginalNum = Dependencies.Num();
 			if (!OriginalNum)
 			{
@@ -450,27 +528,30 @@ void RemoveAll(const TUniqueFunction<bool(const FDependsNode*)>& ShouldRemove, T
 			int32 WriteIndex = 0;
 			int32 ReadIndex = 0;
 			FDependsNode** DependencyData = Dependencies.GetData();
-			bool Keep = !ShouldRemove(DependencyData[ReadIndex]); // use a ! to guarantee it can't be anything other than zero or one
+			// When calling ShouldRemove, use a ! to guarantee it can't be anything other than zero or one
+			bool bKeep = !ShouldRemove(DependencyData[ReadIndex]);
 			do
 			{
 				int32 RunStartIndex = ReadIndex++;
-				while (ReadIndex < OriginalNum && Keep == !ShouldRemove(DependencyData[ReadIndex]))
+				while (ReadIndex < OriginalNum && bKeep == !ShouldRemove(DependencyData[ReadIndex]))
 				{
 					ReadIndex++;
 				}
 				int32 RunLength = ReadIndex - RunStartIndex;
 				checkSlow(RunLength > 0);
-				if (Keep)
+				if (bKeep)
 				{
 					// this was a keep run, we need to move it
 					if (WriteIndex != RunStartIndex)
 					{
-						FMemory::Memmove(&DependencyData[WriteIndex], &DependencyData[RunStartIndex], sizeof(DependencyData[0]) * RunLength);
-						FlagBits->SetRangeFromRange(WriteIndex * FlagSetWidth, FlagSetWidth * RunLength, FlagBits->GetData(), RunStartIndex * FlagSetWidth);
+						FMemory::Memmove(&DependencyData[WriteIndex], &DependencyData[RunStartIndex],
+							sizeof(DependencyData[0]) * RunLength);
+						FlagBits->SetRangeFromRange(WriteIndex * FlagSetWidth,
+							FlagSetWidth * RunLength, FlagBits->GetData(), RunStartIndex * FlagSetWidth);
 					}
 					WriteIndex += RunLength;
 				}
-				Keep = !Keep;
+				bKeep = !bKeep;
 			} while (ReadIndex < OriginalNum);
 
 			Dependencies.SetNum(WriteIndex);
@@ -511,7 +592,8 @@ void FDependsNode::RemoveLinks(const TUniqueFunction<bool(const FDependsNode*)>&
 }
 
 
-bool FDependsNode::ContainsDependency(FDependsNode* InDependency, UE::AssetRegistry::EDependencyCategory Category) const
+bool FDependsNode::ContainsDependency(FDependsNode* InDependency,
+	UE::AssetRegistry::EDependencyCategory Category) const
 {
 	using namespace UE::AssetRegistry;
 	auto ListContains = [InDependency](const TArray<FDependsNode*>& List, bool IsSorted)
@@ -557,7 +639,10 @@ void FDependsNode::PrintDependenciesRecursive(const FString& Indent, TSet<const
 		UE_LOG(LogAssetRegistry, Log, TEXT("%s%s"), *Indent, *Identifier.ToString());
 		VisitedNodes.Add(this);
 
-		IterateOverDependencies([&Indent, &VisitedNodes](FDependsNode* InDependency, UE::AssetRegistry::EDependencyCategory InCategory, UE::AssetRegistry::EDependencyProperty InProperties, bool bDuplicate)
+		IterateOverDependencies([&Indent, &VisitedNodes](FDependsNode* InDependency,
+			UE::AssetRegistry::EDependencyCategory InCategory,
+			UE::AssetRegistry::EDependencyProperty InProperties,
+			bool bDuplicate)
 		{
 			if (!bDuplicate)
 			{
@@ -594,11 +679,17 @@ int32 FDependsNode::GetConnectionCount() const
 	return PackageDependencies.Num() + NameDependencies.Num() + ManageDependencies.Num() + Referencers.Num();
 }
 
-void FDependsNode::SerializeSave(FArchive& Ar, const TUniqueFunction<int32(FDependsNode*, bool)>& GetSerializeIndexFromNode, FSaveScratch& Scratch, const FAssetRegistrySerializationOptions& Options) const
+void FDependsNode::SerializeSave(FArchive& Ar,
+	const TUniqueFunction<int32(FDependsNode*, bool)>& GetSerializeIndexFromNode,
+	FSaveScratch& Scratch,
+	const FAssetRegistrySerializationOptions& Options) const
 {
 	Ar << const_cast<FAssetIdentifier&>(Identifier);
 
-	auto WriteDependencies = [&Ar, &GetSerializeIndexFromNode, &Scratch](const TArray<FDependsNode*>& InDependencies, const TBitArray<>* InFlagBits, int FlagSetWidth, bool bAsReferencer)
+	auto WriteDependencies = [&Ar, &GetSerializeIndexFromNode, &Scratch](const TArray<FDependsNode*>& InDependencies,
+		const TBitArray<>* InFlagBits,
+		int FlagSetWidth,
+		bool bAsReferencer)
 	{
 		TArray<FSaveScratch::FSortInfo>& SortInfos = Scratch.SortInfos;
 		TArray<int32>& OutDependencies = Scratch.OutDependencies;
@@ -616,7 +707,10 @@ void FDependsNode::SerializeSave(FArchive& Ar, const TUniqueFunction<int32(FDepe
 		}
 		// Sort the serialized dependencies to make the output deterministic.
 		// OutDependencies and OutFlagBits (if present) are associated arrays and have to be sorted together
-		Algo::Sort(SortInfos, [](const FSaveScratch::FSortInfo& A, const FSaveScratch::FSortInfo& B) { return A.SerializeIndex < B.SerializeIndex; });
+		Algo::Sort(SortInfos, [](const FSaveScratch::FSortInfo& A, const FSaveScratch::FSortInfo& B)
+			{
+				return A.SerializeIndex < B.SerializeIndex;
+			});
 
 		int32 NumOutDependencies = SortInfos.Num();
 		OutDependencies.Reset(NumOutDependencies);
@@ -633,7 +727,8 @@ void FDependsNode::SerializeSave(FArchive& Ar, const TUniqueFunction<int32(FDepe
 				OutFlagBits.AddRange(*InFlagBits, FlagSetWidth, SortInfo.ListIndex * FlagSetWidth);
 			}
 
-			// We don't use BitArray::operator<< because we want to avoid the reallocation that operator<< does on load and we want to avoid saving BitArray.Num when it can be derived from NumDependencies
+			// We don't use BitArray::operator<< because we want to avoid the reallocation that operator<< does on load
+			// and we want to avoid saving BitArray.Num when it can be derived from NumDependencies.
 			int32 NumFlagBits = FlagSetWidth * NumOutDependencies;
 			check(OutFlagBits.Num() == NumFlagBits);
 			int32 NumFlagWords = FBitSet::CalculateNumWords(NumFlagBits);
@@ -642,16 +737,22 @@ void FDependsNode::SerializeSave(FArchive& Ar, const TUniqueFunction<int32(FDepe
 	};
 
 	WriteDependencies(PackageDependencies, &PackageFlags, PackageFlagSetWidth, false);
-	WriteDependencies(Options.bSerializeSearchableNameDependencies ? NameDependencies : FDependsNodeList(), nullptr, 0, false);
-	WriteDependencies(Options.bSerializeManageDependencies ? ManageDependencies : FDependsNodeList(), Options.bSerializeManageDependencies ? &ManageFlags : nullptr, ManageFlagSetWidth, false);
+	WriteDependencies(Options.bSerializeSearchableNameDependencies ? NameDependencies : FDependsNodeList(),
+		nullptr, 0, false);
+	WriteDependencies(Options.bSerializeManageDependencies ? ManageDependencies : FDependsNodeList(),
+		Options.bSerializeManageDependencies ? &ManageFlags : nullptr, ManageFlagSetWidth, false);
 	WriteDependencies(Referencers, nullptr, 0, true);
 }
 
-void FDependsNode::SerializeLoad(FArchive& Ar, const TUniqueFunction<FDependsNode* (int32)>& GetNodeFromSerializeIndex, FLoadScratch& Scratch)
+void FDependsNode::SerializeLoad(FArchive& Ar,
+	const TUniqueFunction<FDependsNode* (int32)>& GetNodeFromSerializeIndex,
+	FLoadScratch& Scratch)
 {
 	Ar << Identifier;
 
-	auto ReadDependencies = [&Ar, &GetNodeFromSerializeIndex, &Scratch](TArray<FDependsNode*>& OutDependencies, TBitArray<>* OutFlagBits, int FlagSetWidth)
+	auto ReadDependencies = [&Ar, &GetNodeFromSerializeIndex, &Scratch](TArray<FDependsNode*>& OutDependencies,
+		TBitArray<>* OutFlagBits,
+		int FlagSetWidth)
 	{
 		TArray<int32>& InDependencies = Scratch.InDependencies;
 		TArray<uint32>& InFlagBits = Scratch.InFlagBits;
@@ -664,7 +765,8 @@ void FDependsNode::SerializeLoad(FArchive& Ar, const TUniqueFunction<FDependsNod
 		int32 NumDependencies = InDependencies.Num();
 		if (OutFlagBits)
 		{
-			// We don't use BitArray::operator<< because we want to avoid the reallocation that operator<< does on load and we want to avoid saving BitArray.Num when it can be derived from NumDependencies
+			// We don't use BitArray::operator<< because we want to avoid the reallocation that operator<< does on load
+			// and we want to avoid saving BitArray.Num when it can be derived from NumDependencies.
 			NumFlagBits = FlagSetWidth * NumDependencies;
 			const int32 NumFlagWords = FBitSet::CalculateNumWords(NumFlagBits);
 			InFlagBits.SetNumUninitialized(NumFlagWords);
@@ -689,7 +791,10 @@ void FDependsNode::SerializeLoad(FArchive& Ar, const TUniqueFunction<FDependsNod
 			SortIndexes.Add(Index);
 		}
 
-		Algo::Sort(SortIndexes, [&PointerDependencies](int32 A, int32 B) { return PointerDependencies[A] < PointerDependencies[B]; });
+		Algo::Sort(SortIndexes, [&PointerDependencies](int32 A, int32 B)
+			{
+				return PointerDependencies[A] < PointerDependencies[B];
+			});
 
 		OutDependencies.Empty(NumDependencies);
 		for (int32 SortIndex : SortIndexes)
@@ -703,7 +808,8 @@ void FDependsNode::SerializeLoad(FArchive& Ar, const TUniqueFunction<FDependsNod
 			for (int32 WriteIndex = 0; WriteIndex < NumDependencies; ++WriteIndex)
 			{
 				int32 ReadIndex = SortIndexes[WriteIndex];
-				OutFlagBits->SetRangeFromRange(WriteIndex * FlagSetWidth, FlagSetWidth, InFlagBitsData, ReadIndex * FlagSetWidth);
+				OutFlagBits->SetRangeFromRange(WriteIndex * FlagSetWidth, FlagSetWidth,
+					InFlagBitsData, ReadIndex * FlagSetWidth);
 			}
 		}
 	};
@@ -716,8 +822,15 @@ void FDependsNode::SerializeLoad(FArchive& Ar, const TUniqueFunction<FDependsNod
 	SetIsDependenciesInitialized(true);
 }
 
-void FDependsNode::SerializeLoad_BeforeFlags(FArchive& Ar, FAssetRegistryVersion::Type Version, FDependsNode* PreallocatedDependsNodeDataBuffer, int32 NumDependsNodes, bool bSerializeDependencies,
-	uint32 HardBits, uint32 SoftBits, uint32 HardManageBits, uint32 SoftManageBits)
+void FDependsNode::SerializeLoad_BeforeFlags(FArchive& Ar,
+	FAssetRegistryVersion::Type Version,
+	FDependsNode* PreallocatedDependsNodeDataBuffer,
+	int32 NumDependsNodes,
+	bool bSerializeDependencies,
+	uint32 HardBits,
+	uint32 SoftBits,
+	uint32 HardManageBits,
+	uint32 SoftManageBits)
 {
 	Ar << Identifier;
 
@@ -742,7 +855,13 @@ void FDependsNode::SerializeLoad_BeforeFlags(FArchive& Ar, FAssetRegistryVersion
 	ManageDependencies.Empty(bSerializeDependencies ? NumSoftManage + NumHardManage : 0);
 	Referencers.Empty(NumReferencers);
 
-	auto SerializeNodeArray = [&Ar, PreallocatedDependsNodeDataBuffer, NumDependsNodes](int32 Num, TArray<FDependsNode*>& OutNodes, TBitArray<>* OutFlagBits, uint32 FlagSetWidth, uint32 FlagSetBits, bool bShouldOverwriteFlag, bool bAllowWrite)
+	auto SerializeNodeArray = [&Ar, PreallocatedDependsNodeDataBuffer, NumDependsNodes](int32 Num,
+		TArray<FDependsNode*>& OutNodes,
+		TBitArray<>* OutFlagBits,
+		uint32 FlagSetWidth,
+		uint32 FlagSetBits,
+		bool bShouldOverwriteFlag,
+		bool bAllowWrite)
 	{
 		for (int32 DependencyIndex = 0; DependencyIndex < Num; ++DependencyIndex)
 		{
@@ -780,23 +899,31 @@ void FDependsNode::SerializeLoad_BeforeFlags(FArchive& Ar, FAssetRegistryVersion
 	SerializeNodeArray(NumHard, PackageDependencies, &PackageFlags, PackageFlagSetWidth, HardBits, true, true);
 	SerializeNodeArray(NumSoft, PackageDependencies, &PackageFlags, PackageFlagSetWidth, SoftBits, false, true);
 	SerializeNodeArray(NumName, NameDependencies, nullptr, 0, 0, true, bSerializeDependencies);
-	SerializeNodeArray(NumSoftManage, ManageDependencies, &ManageFlags, ManageFlagSetWidth, SoftManageBits, true, bSerializeDependencies);
-	SerializeNodeArray(NumHardManage, ManageDependencies, &ManageFlags, ManageFlagSetWidth, HardManageBits, true, bSerializeDependencies);
+	SerializeNodeArray(NumSoftManage, ManageDependencies, &ManageFlags, ManageFlagSetWidth, SoftManageBits,
+		true, bSerializeDependencies);
+	SerializeNodeArray(NumHardManage, ManageDependencies, &ManageFlags, ManageFlagSetWidth, HardManageBits,
+		true, bSerializeDependencies);
 	SerializeNodeArray(NumReferencers, Referencers, nullptr, 0, 0, true, true);
 
 	SetIsDependenciesInitialized(true);
 }
 
-void FDependsNode::GetPropertySetBits_BeforeFlags(uint32& HardBits, uint32& SoftBits, uint32& HardManageBits, uint32& SoftManageBits)
+void FDependsNode::GetPropertySetBits_BeforeFlags(uint32& HardBits,
+	uint32& SoftBits,
+	uint32& HardManageBits,
+	uint32& SoftManageBits)
 {
 	{
 		FDependsNode::FPackageFlagSet FlagSet;
-		FlagSet.Add(PackagePropertiesToByte(UE::AssetRegistry::EDependencyProperty::Hard | UE::AssetRegistry::EDependencyProperty::Game | UE::AssetRegistry::EDependencyProperty::Build));
+		FlagSet.Add(PackagePropertiesToByte(
+			UE::AssetRegistry::EDependencyProperty::Hard
+			| UE::AssetRegistry::EDependencyProperty::Game | UE::AssetRegistry::EDependencyProperty::Build));
 		FlagSet.Save(&HardBits);
 	}
 	{
 		FDependsNode::FPackageFlagSet FlagSet;
-		FlagSet.Add(PackagePropertiesToByte(UE::AssetRegistry::EDependencyProperty::Game | UE::AssetRegistry::EDependencyProperty::Build));
+		FlagSet.Add(PackagePropertiesToByte(
+			UE::AssetRegistry::EDependencyProperty::Game | UE::AssetRegistry::EDependencyProperty::Build));
 		FlagSet.Save(&SoftBits);
 	}
 	HardManageBits = 0x1;
@@ -852,7 +979,10 @@ void SortDependencyList(TArray<FDependsNode*>& Dependencies, TBitArray<>* Flags)
 		}
 
 		// Sort the index array
-		Algo::Sort(Order, [&DependencyData](int32 A, int32 B) { return DependencyData[A] < DependencyData[B]; });
+		Algo::Sort(Order, [&DependencyData](int32 A, int32 B)
+			{
+				return DependencyData[A] < DependencyData[B];
+			});
 
 		// Remove duplicate in the dependency array, which are now adjacent, and merge their corresponding flags
 		typedef TPropertyCombinationSet<FlagWidth> FCombinationSet;
diff --git a/Engine/Source/Runtime/AssetRegistry/Private/DependsNode.h b/Engine/Source/Runtime/AssetRegistry/Private/DependsNode.h
--- a/Engine/Source/Runtime/AssetRegistry/Private/DependsNode.h
+++ b/Engine/Source/Runtime/AssetRegistry/Private/DependsNode.h
@@ -33,13 +33,23 @@ public:
 	/** Prints the referencers to this node to the log */
 	void PrintReferencers() const;
 	/** Gets the list of dependencies for this node */
-	void GetDependencies(TArray<FDependsNode*>& OutDependencies, UE::AssetRegistry::EDependencyCategory Category = UE::AssetRegistry::EDependencyCategory::All, const UE::AssetRegistry::FDependencyQuery& Flags = UE::AssetRegistry::FDependencyQuery()) const;
+	void GetDependencies(TArray<FDependsNode*>& OutDependencies,
+		UE::AssetRegistry::EDependencyCategory Category = UE::AssetRegistry::EDependencyCategory::All,
+		const UE::AssetRegistry::FDependencyQuery& Flags = UE::AssetRegistry::FDependencyQuery()) const;
 	/** Gets the list of dependency names for this node */
-	void GetDependencies(TArray<FAssetIdentifier>& OutDependencies, UE::AssetRegistry::EDependencyCategory Category = UE::AssetRegistry::EDependencyCategory::All, const UE::AssetRegistry::FDependencyQuery& Flags = UE::AssetRegistry::FDependencyQuery()) const;
-	void GetDependencies(TArray<FAssetDependency>& OutDependencies, UE::AssetRegistry::EDependencyCategory Category = UE::AssetRegistry::EDependencyCategory::All, const UE::AssetRegistry::FDependencyQuery& Flags = UE::AssetRegistry::FDependencyQuery()) const;
+	void GetDependencies(TArray<FAssetIdentifier>& OutDependencies,
+		UE::AssetRegistry::EDependencyCategory Category = UE::AssetRegistry::EDependencyCategory::All,
+		const UE::AssetRegistry::FDependencyQuery& Flags = UE::AssetRegistry::FDependencyQuery()) const;
+	void GetDependencies(TArray<FAssetDependency>& OutDependencies,
+		UE::AssetRegistry::EDependencyCategory Category = UE::AssetRegistry::EDependencyCategory::All,
+		const UE::AssetRegistry::FDependencyQuery& Flags = UE::AssetRegistry::FDependencyQuery()) const;
 	/** Gets the list of referencers to this node */
-	void GetReferencers(TArray<FDependsNode*>& OutReferencers, UE::AssetRegistry::EDependencyCategory Category = UE::AssetRegistry::EDependencyCategory::All, const UE::AssetRegistry::FDependencyQuery& Flags = UE::AssetRegistry::FDependencyQuery()) const;
-	void GetReferencers(TArray<FAssetDependency>& OutReferencers, UE::AssetRegistry::EDependencyCategory Category = UE::AssetRegistry::EDependencyCategory::All, const UE::AssetRegistry::FDependencyQuery& Flags = UE::AssetRegistry::FDependencyQuery()) const;
+	void GetReferencers(TArray<FDependsNode*>& OutReferencers,
+		UE::AssetRegistry::EDependencyCategory Category = UE::AssetRegistry::EDependencyCategory::All,
+		const UE::AssetRegistry::FDependencyQuery& Flags = UE::AssetRegistry::FDependencyQuery()) const;
+	void GetReferencers(TArray<FAssetDependency>& OutReferencers,
+		UE::AssetRegistry::EDependencyCategory Category = UE::AssetRegistry::EDependencyCategory::All,
+		const UE::AssetRegistry::FDependencyQuery& Flags = UE::AssetRegistry::FDependencyQuery()) const;
 	/** Gets the name of the package that this node represents */
 	FName GetPackageName() const { return Identifier.PackageName; }
 	/** Sets the name of the package that this node represents */
@@ -49,89 +59,64 @@ public:
 	/** Sets the entire identifier */
 	void SetIdentifier(const FAssetIdentifier& InIdentifier) { Identifier = InIdentifier; }
 	/** Add a dependency to this node */
-	void AddDependency(FDependsNode* InDependency, UE::AssetRegistry::EDependencyCategory InDependencyType, UE::AssetRegistry::EDependencyProperty InProperties);
+	void AddDependency(FDependsNode* InDependency,
+		UE::AssetRegistry::EDependencyCategory InDependencyType, UE::AssetRegistry::EDependencyProperty InProperties);
 	void GetPackageReferencers(TArray<TPair<FAssetIdentifier, FPackageFlagSet>>& OutReferencers);
 	void AddPackageDependencySet(FDependsNode* InDependency, const FPackageFlagSet& PropertyCombinationSet);
 	/** Add a referencer to this node */
 	void AddReferencer(FDependsNode* InReferencer);
 	/** Remove a dependency from this node */
-	void RemoveDependency(FDependsNode* InDependency, UE::AssetRegistry::EDependencyCategory Category = UE::AssetRegistry::EDependencyCategory::All);
+	void RemoveDependency(FDependsNode* InDependency,
+		UE::AssetRegistry::EDependencyCategory Category = UE::AssetRegistry::EDependencyCategory::All);
 	/** Remove a referencer from this node */
 	void RemoveReferencer(FDependsNode* InReferencer);
 	/** Removes any referencers that no longer have this node as a dependency */
 	void RefreshReferencers();
-	bool ContainsDependency(FDependsNode* InDependency, UE::AssetRegistry::EDependencyCategory Category = UE::AssetRegistry::EDependencyCategory::All) const;
+	bool ContainsDependency(FDependsNode* InDependency,
+		UE::AssetRegistry::EDependencyCategory Category = UE::AssetRegistry::EDependencyCategory::All) const;
 	/** Clear all dependency records from this node */
 	void ClearDependencies(UE::AssetRegistry::EDependencyCategory Category = UE::AssetRegistry::EDependencyCategory::All);
 	void ClearReferencers();
-	/** Removes Manage dependencies on this node and clean up referencers array. Manage references are the only ones safe to remove at runtime */
+	/**
+	 * Removes Manage dependencies on this node and clean up referencers array.
+	 * Manage references are the only ones safe to remove at runtime.
+	 */
 	void RemoveManageReferencesToNode();
 	/** Remove all nodes from referencers and dependencies for which ShouldRemove returns true */
 	void RemoveLinks(const TUniqueFunction<bool(const FDependsNode*)>& ShouldRemove);
 	/** Returns number of connections this node has, both references and dependencies */
 	int32 GetConnectionCount() const;
 	/** Returns amount of memory used by the arrays */
-	SIZE_T GetAllocatedSize(void) const
-	{
-		return PackageDependencies.GetAllocatedSize() + PackageFlags.GetAllocatedSize() + NameDependencies.GetAllocatedSize() + ManageDependencies.GetAllocatedSize() + ManageFlags.GetAllocatedSize() + Referencers.GetAllocatedSize();
-	}
+	SIZE_T GetAllocatedSize(void) const;
 
-	typedef TUniqueFunction<void(FDependsNode * Dependency, UE::AssetRegistry::EDependencyCategory Category, UE::AssetRegistry::EDependencyProperty Properties, bool bDuplicate)> FIterateDependenciesCallback;
-	/** Iterate over all the dependencies of this node, optionally filtered by the target node, category and query, and call the supplied lambda parameter on the record */
-	void IterateOverDependencies(const FIterateDependenciesCallback& InCallback, UE::AssetRegistry::EDependencyCategory Category = UE::AssetRegistry::EDependencyCategory::All, const UE::AssetRegistry::FDependencyQuery& Flags = UE::AssetRegistry::FDependencyQuery()) const;
-	void IterateOverDependencies(const FIterateDependenciesCallback& InCallback, const FDependsNode* DependsNode, UE::AssetRegistry::EDependencyCategory Category = UE::AssetRegistry::EDependencyCategory::All, const UE::AssetRegistry::FDependencyQuery& Flags = UE::AssetRegistry::FDependencyQuery()) const;
+	typedef TUniqueFunction<void(FDependsNode * Dependency,
+		UE::AssetRegistry::EDependencyCategory Category,
+		UE::AssetRegistry::EDependencyProperty Properties,
+		bool bDuplicate)> FIterateDependenciesCallback;
+	/**
+	 * Iterate over all the dependencies of this node, optionally filtered by the target node, category and query,
+	 * and call the supplied lambda parameter on the record.
+	 */
+	void IterateOverDependencies(const FIterateDependenciesCallback& InCallback,
+		UE::AssetRegistry::EDependencyCategory Category = UE::AssetRegistry::EDependencyCategory::All,
+		const UE::AssetRegistry::FDependencyQuery& Flags = UE::AssetRegistry::FDependencyQuery()) const;
+	void IterateOverDependencies(const FIterateDependenciesCallback& InCallback,
+		const FDependsNode* DependsNode,
+		UE::AssetRegistry::EDependencyCategory Category = UE::AssetRegistry::EDependencyCategory::All,
+		const UE::AssetRegistry::FDependencyQuery& Flags = UE::AssetRegistry::FDependencyQuery()) const;
 
 	/** Iterate over all the referencers of this node and call the supplied lambda parameter on the referencer */
 	template <class T>
-	void IterateOverReferencers(const T& InCallback) const
-	{
-		for (FDependsNode* Referencer : Referencers)
-		{
-			InCallback(Referencer);
-		}
-	}
-
-	void Reserve(int32 InNumPackageDependencies, int32 InNumNameDependencies, int32 InNumManageDependencies, int32 InNumReferencers)
-	{
-		PackageDependencies.Reserve(InNumPackageDependencies);
-		PackageFlags.Reserve(InNumPackageDependencies * PackageFlagSetWidth);
-		NameDependencies.Reserve(InNumNameDependencies);
-		ManageDependencies.Reserve(InNumManageDependencies);
-		ManageFlags.Reserve(InNumManageDependencies * ManageFlagSetWidth);
-		Referencers.Reserve(InNumReferencers);
-	}
-
-	void Reserve(const FDependsNode* Other)
-	{
-		Reserve(Other->PackageDependencies.Num(), Other->NameDependencies.Num(), Other->ManageDependencies.Num(), Other->Referencers.Num());
-	}
+	void IterateOverReferencers(const T& InCallback) const;
 
-	inline static uint8 PackagePropertiesToByte(UE::AssetRegistry::EDependencyProperty Properties)
-	{
-		return (0x01 * (static_cast<uint8>(Properties & UE::AssetRegistry::EDependencyProperty::Hard) != 0))
-			| (0x02 * (static_cast<uint8>(Properties & UE::AssetRegistry::EDependencyProperty::Game) != 0))
-			| (0x04 * (static_cast<uint8>(Properties & UE::AssetRegistry::EDependencyProperty::Build) != 0));
-	}
+	void Reserve(int32 InNumPackageDependencies, int32 InNumNameDependencies,
+		int32 InNumManageDependencies, int32 InNumReferencers);
+	void Reserve(const FDependsNode* Other);
 
-	inline static UE::AssetRegistry::EDependencyProperty ByteToPackageProperties(uint8 Bits)
-	{
-		return static_cast<UE::AssetRegistry::EDependencyProperty>(
-			(static_cast<uint32>(UE::AssetRegistry::EDependencyProperty::Hard) * ((Bits & 0x01) != 0))
-			| (static_cast<uint32>(UE::AssetRegistry::EDependencyProperty::Game) * ((Bits & 0x02) != 0))
-			| (static_cast<uint32>(UE::AssetRegistry::EDependencyProperty::Build) * ((Bits & 0x04) != 0))
-			);
-	}
-	inline static uint8 ManagePropertiesToByte(UE::AssetRegistry::EDependencyProperty Properties)
-	{
-		return (0x01 * (static_cast<uint8>(Properties & UE::AssetRegistry::EDependencyProperty::Direct) != 0));
-	}
-
-	inline static UE::AssetRegistry::EDependencyProperty ByteToManageProperties(uint8 Bits)
-	{
-		return static_cast<UE::AssetRegistry::EDependencyProperty>(
-			(static_cast<uint32>(UE::AssetRegistry::EDependencyProperty::Direct) * ((Bits & 0x01) != 0))
-			);
-	}
+	static uint8 PackagePropertiesToByte(UE::AssetRegistry::EDependencyProperty Properties);
+	static UE::AssetRegistry::EDependencyProperty ByteToPackageProperties(uint8 Bits);
+	static uint8 ManagePropertiesToByte(UE::AssetRegistry::EDependencyProperty Properties);
+	static UE::AssetRegistry::EDependencyProperty ByteToManageProperties(uint8 Bits);
 
 	struct FSaveScratch
 	{
@@ -144,7 +129,8 @@ public:
 		TArray<int32> OutDependencies;
 		TBitArray<> OutFlagBits;
 	};
-	void SerializeSave(FArchive& Ar, const TUniqueFunction<int32(FDependsNode*, bool)>& GetSerializeIndexFromNode, FSaveScratch& Scratch, const FAssetRegistrySerializationOptions& Options) const;
+	void SerializeSave(FArchive& Ar, const TUniqueFunction<int32(FDependsNode*, bool)>& GetSerializeIndexFromNode,
+		FSaveScratch& Scratch, const FAssetRegistrySerializationOptions& Options) const;
 	struct FLoadScratch
 	{
 		TArray<int32> InDependencies;
@@ -152,11 +138,14 @@ public:
 		TArray<FDependsNode*> PointerDependencies;
 		TArray<int32> SortIndexes;
 	};
-	void SerializeLoad(FArchive& Ar, const TUniqueFunction<FDependsNode*(int32)>& GetNodeFromSerializeIndex, FLoadScratch& Scratch);
+	void SerializeLoad(FArchive& Ar,
+		const TUniqueFunction<FDependsNode*(int32)>& GetNodeFromSerializeIndex, FLoadScratch& Scratch);
 
-	void SerializeLoad_BeforeFlags(FArchive& Ar, FAssetRegistryVersion::Type Version, FDependsNode* PreallocatedDependsNodeDataBuffer, int32 NumDependsNodes, bool bSerializeDependencies,
+	void SerializeLoad_BeforeFlags(FArchive& Ar, FAssetRegistryVersion::Type Version,
+		FDependsNode* PreallocatedDependsNodeDataBuffer, int32 NumDependsNodes, bool bSerializeDependencies,
 		uint32 HardBits, uint32 SoftBits, uint32 HardManageBits, uint32 SoftManageBits);
-	static void GetPropertySetBits_BeforeFlags(uint32& HardBits, uint32& SoftBits, uint32& HardManageBits, uint32& SoftManageBits);
+	static void GetPropertySetBits_BeforeFlags(uint32& HardBits, uint32& SoftBits,
+		uint32& HardManageBits, uint32& SoftManageBits);
 
 	bool IsDependencyListSorted(UE::AssetRegistry::EDependencyCategory Category) const;
 	void SetIsDependencyListSorted(UE::AssetRegistry::EDependencyCategory Category, bool bValue);
@@ -167,19 +156,18 @@ public:
 
 private:
 
-	/** Recursively prints dependencies of the node starting with the specified indent. VisitedNodes should be an empty set at first which is populated recursively. */
+	/**
+	 * Recursively prints dependencies of the node starting with the specified indent.
+	 * VisitedNodes should be an empty set at first which is populated recursively.
+	 */
 	void PrintDependenciesRecursive(const FString& Indent, TSet<const FDependsNode*>& VisitedNodes) const;
-	/** Recursively prints referencers to the node starting with the specified indent. VisitedNodes should be an empty set at first which is populated recursively. */
+	/**
+	 * Recursively prints referencers to the node starting with the specified indent.
+	 * VisitedNodes should be an empty set at first which is populated recursively.
+	 */
 	void PrintReferencersRecursive(const FString& Indent, TSet<const FDependsNode*>& VisitedNodes) const;
 
-	void Construct()
-	{
-		PackageIsSorted = 1;
-		SearchableNameIsSorted = 1;
-		ManageIsSorted = 1;
-		ReferencersIsSorted = 1;
-		DependenciesInitialized = 0;
-	}
+	void Construct();
 
 	/** The name of the package/object this node represents */
 	FAssetIdentifier Identifier;
@@ -195,3 +183,78 @@ private:
 	uint32 ReferencersIsSorted : 1;
 	uint32 DependenciesInitialized : 1;
 };
+
+
+///////////////////////////////////////////////////////
+// Inline implementations
+///////////////////////////////////////////////////////
+
+
+inline SIZE_T FDependsNode::GetAllocatedSize(void) const
+{
+	return PackageDependencies.GetAllocatedSize() + PackageFlags.GetAllocatedSize()
+		+ NameDependencies.GetAllocatedSize() + ManageDependencies.GetAllocatedSize() + ManageFlags.GetAllocatedSize()
+		+ Referencers.GetAllocatedSize();
+}
+
+template <class T>
+inline void FDependsNode::IterateOverReferencers(const T& InCallback) const
+{
+	for (FDependsNode* Referencer : Referencers)
+	{
+		InCallback(Referencer);
+	}
+}
+
+inline void FDependsNode::Reserve(int32 InNumPackageDependencies, int32 InNumNameDependencies,
+	int32 InNumManageDependencies, int32 InNumReferencers)
+{
+	PackageDependencies.Reserve(InNumPackageDependencies);
+	PackageFlags.Reserve(InNumPackageDependencies * PackageFlagSetWidth);
+	NameDependencies.Reserve(InNumNameDependencies);
+	ManageDependencies.Reserve(InNumManageDependencies);
+	ManageFlags.Reserve(InNumManageDependencies * ManageFlagSetWidth);
+	Referencers.Reserve(InNumReferencers);
+}
+
+inline void FDependsNode::Reserve(const FDependsNode* Other)
+{
+	Reserve(Other->PackageDependencies.Num(), Other->NameDependencies.Num(),
+		Other->ManageDependencies.Num(), Other->Referencers.Num());
+}
+
+inline uint8 FDependsNode::PackagePropertiesToByte(UE::AssetRegistry::EDependencyProperty Properties)
+{
+	return (0x01 * (static_cast<uint8>(Properties & UE::AssetRegistry::EDependencyProperty::Hard) != 0))
+		| (0x02 * (static_cast<uint8>(Properties & UE::AssetRegistry::EDependencyProperty::Game) != 0))
+		| (0x04 * (static_cast<uint8>(Properties & UE::AssetRegistry::EDependencyProperty::Build) != 0));
+}
+
+inline UE::AssetRegistry::EDependencyProperty FDependsNode::ByteToPackageProperties(uint8 Bits)
+{
+	return static_cast<UE::AssetRegistry::EDependencyProperty>(
+		(static_cast<uint32>(UE::AssetRegistry::EDependencyProperty::Hard) * ((Bits & 0x01) != 0))
+		| (static_cast<uint32>(UE::AssetRegistry::EDependencyProperty::Game) * ((Bits & 0x02) != 0))
+		| (static_cast<uint32>(UE::AssetRegistry::EDependencyProperty::Build) * ((Bits & 0x04) != 0))
+		);
+}
+inline uint8 FDependsNode::ManagePropertiesToByte(UE::AssetRegistry::EDependencyProperty Properties)
+{
+	return (0x01 * (static_cast<uint8>(Properties & UE::AssetRegistry::EDependencyProperty::Direct) != 0));
+}
+
+inline UE::AssetRegistry::EDependencyProperty FDependsNode::ByteToManageProperties(uint8 Bits)
+{
+	return static_cast<UE::AssetRegistry::EDependencyProperty>(
+		(static_cast<uint32>(UE::AssetRegistry::EDependencyProperty::Direct) * ((Bits & 0x01) != 0))
+		);
+}
+
+inline void FDependsNode::Construct()
+{
+	PackageIsSorted = 1;
+	SearchableNameIsSorted = 1;
+	ManageIsSorted = 1;
+	ReferencersIsSorted = 1;
+	DependenciesInitialized = 0;
+}
diff --git a/Engine/Source/Runtime/AssetRegistry/Public/AssetRegistry/AssetRegistryState.h b/Engine/Source/Runtime/AssetRegistry/Public/AssetRegistry/AssetRegistryState.h
--- a/Engine/Source/Runtime/AssetRegistry/Public/AssetRegistry/AssetRegistryState.h
+++ b/Engine/Source/Runtime/AssetRegistry/Public/AssetRegistry/AssetRegistryState.h
@@ -38,21 +38,19 @@ namespace UE::AssetRegistry
 
 struct FAssetRegistryHeader;
 
-/** Load/Save options used to modify how the cache is serialized. These are read out of the AssetRegistry section of Engine.ini and can be changed per platform. */
+/**
+ * Load/Save options used to modify how the cache is serialized.
+ * These are read out of the AssetRegistry section of Engine.ini and can be changed per platform.
+ */
 struct FAssetRegistrySerializationOptions
 {
-	FAssetRegistrySerializationOptions(UE::AssetRegistry::ESerializationTarget Target = UE::AssetRegistry::ESerializationTarget::ForGame)
-	{
-		if (Target == UE::AssetRegistry::ESerializationTarget::ForDevelopment)
-		{
-			InitForDevelopment();
-		}
-	}
+	FAssetRegistrySerializationOptions(
+		UE::AssetRegistry::ESerializationTarget Target = UE::AssetRegistry::ESerializationTarget::ForGame);
 
 	/** True rather to load/save registry at all */
 	bool bSerializeAssetRegistry = false;
 
-	/** True rather to load/save dependency info. If true this will handle hard and soft package references */
+	/** True rather to load/save dependency info. If true this will handle hard and soft package references. */
 	bool bSerializeDependencies = false;
 
 	/** True rather to load/save dependency info for Name references,  */
@@ -67,53 +65,50 @@ struct FAssetRegistrySerializationOptions
 	/** True if CookFilterlistTagsByClass is an allow list. False if it is a deny list. */
 	bool bUseAssetRegistryTagsAllowListInsteadOfDenyList = false;
 
-	/** True if we want to only write out asset data if it has valid tags. This saves memory by not saving data for things like textures */
+	/**
+	 * True if we want to only write out asset data if it has valid tags.
+	 * This saves memory by not saving data for things like textures.
+	 */
 	bool bFilterAssetDataWithNoTags = false;
 
-	/** True if we also want to filter out dependency data for assets that have no tags. Only filters if bFilterAssetDataWithNoTags is also true */
+	/**
+	 * True if we also want to filter out dependency data for assets that have no tags.
+	 * Only filters if bFilterAssetDataWithNoTags is also true.
+	 */
 	bool bFilterDependenciesWithNoTags = false;
 
 	/** Filter out searchable names from dependency data */
 	bool bFilterSearchableNames = false;
 
 	/**
-	 * Keep tags intended for the cooker's output DevelopmentAssetRegistry. this flag defaults to false and is set to
+	 * Keep tags intended for the cooker's output DevelopmentAssetRegistry. This flag defaults to false and is set to
 	 * true only by the cooker.
 	 */
 	bool bKeepDevelopmentAssetRegistryTags = false;
 
-	/** The map of class pathname to tag set of tags that are allowed in cooked builds. This is either an allow list or deny list depending on bUseAssetRegistryTagsAllowListInsteadOfDenyList */
+	/**
+	 * The map of class pathname to tag set of tags that are allowed in cooked builds.
+	 * This is either an allow list or deny list depending on bUseAssetRegistryTagsAllowListInsteadOfDenyList.
+	 */
 	TMap<FTopLevelAssetPath, TSet<FName>> CookFilterlistTagsByClass;
 
-	/** Tag keys whose values should be stored as FName in cooked builds */
+	/** Tag keys whose values should be stored as FName in cooked builds. */
 	TSet<FName> CookTagsAsName;
 
-	/** Tag keys whose values should be stored as FRegistryExportPath in cooked builds */
+	/** Tag keys whose values should be stored as FRegistryExportPath in cooked builds. */
 	TSet<FName> CookTagsAsPath;
 
 	/** Disable all filters */
-	void DisableFilters()
-	{
-		bFilterAssetDataWithNoTags = false;
-		bFilterDependenciesWithNoTags = false;
-		bFilterSearchableNames = false;
-	}
+	void DisableFilters();
 
 private:
-	void InitForDevelopment()
-	{
-		bSerializeAssetRegistry = bSerializeDependencies = bSerializeSearchableNameDependencies = bSerializeManageDependencies = bSerializePackageData = true;
-		DisableFilters();
-	}
+	void InitForDevelopment();
 };
 
 struct FAssetRegistryLoadOptions
 {
 	FAssetRegistryLoadOptions() = default;
-	explicit FAssetRegistryLoadOptions(const FAssetRegistrySerializationOptions& Options)
-		: bLoadDependencies(Options.bSerializeDependencies)
-		, bLoadPackageData(Options.bSerializePackageData)
-	{}
+	explicit FAssetRegistryLoadOptions(const FAssetRegistrySerializationOptions& Options);
 
 	bool bLoadDependencies = true;
 	bool bLoadPackageData = true;
@@ -136,69 +131,75 @@ struct FAssetRegistryPruneOptions
 
 namespace UE::AssetRegistry::Private
 {
-	/* 
-	* Key type for TSet<FAssetData*> in the asset registry.
-	* Top level assets are searched for by their asset path as two names (e.g. '/Path/ToPackageName' + 'AssetName')
-	* Other assets (e.g. external actors) are searched for by their full path with the whole outer chain as a single name. 
-	* (e.g. '/Path/To/Package.TopLevel:Subobject' + 'DeeperSubobject')
-	*/
-	struct FCachedAssetKey
-	{
-		explicit FCachedAssetKey(const FAssetData* InAssetData);
-		explicit FCachedAssetKey(const FAssetData& InAssetData);
-		explicit FCachedAssetKey(FTopLevelAssetPath InAssetPath);
-		explicit FCachedAssetKey(const FSoftObjectPath& InObjectPath);
 
-		FString ToString() const;
-		int32 Compare(const FCachedAssetKey& Other) const;	// Order asset keys with fast non-lexical comparison
-		void AppendString(FStringBuilderBase& Builder) const;
+/* 
+* Key type for TSet<FAssetData*> in the asset registry.
+* Top level assets are searched for by their asset path as two names (e.g. '/Path/ToPackageName' + 'AssetName')
+* Other assets (e.g. external actors) are searched for by their full path with the whole outer chain as a single name. 
+* (e.g. '/Path/To/Package.TopLevel:Subobject' + 'DeeperSubobject')
+*/
+struct FCachedAssetKey
+{
+	explicit FCachedAssetKey(const FAssetData* InAssetData);
+	explicit FCachedAssetKey(const FAssetData& InAssetData);
+	explicit FCachedAssetKey(FTopLevelAssetPath InAssetPath);
+	explicit FCachedAssetKey(const FSoftObjectPath& InObjectPath);
 
-		FName OuterPath = NAME_None;
-		FName ObjectName = NAME_None;
-	};
+	FString ToString() const;
+	int32 Compare(const FCachedAssetKey& Other) const;	// Order asset keys with fast non-lexical comparison
+	void AppendString(FStringBuilderBase& Builder) const;
 
-	inline FStringBuilderBase& operator<<(FStringBuilderBase& Builder, const FCachedAssetKey& Key);
-	inline bool operator==(const FCachedAssetKey& A, const FCachedAssetKey& B);
-	inline bool operator!=(const FCachedAssetKey& A, const FCachedAssetKey& B);
-	inline uint32 GetTypeHash(const FCachedAssetKey& A);
+	FName OuterPath = NAME_None;
+	FName ObjectName = NAME_None;
+};
 
-	/* 
-	* Policy type for TSet<FAssetData*> to use FCachedAssetKey for hashing/equality.
-	* This allows is to store just FAssetData* in the map without storing an extra copy of the key fields to save memory.
-	*/
-	struct FCachedAssetKeyFuncs
-	{
-		using KeyInitType = FCachedAssetKey;
-		using ElementInitType = void; // TSet doesn't actually use this type 
+inline FStringBuilderBase& operator<<(FStringBuilderBase& Builder, const FCachedAssetKey& Key);
+inline bool operator==(const FCachedAssetKey& A, const FCachedAssetKey& B);
+inline bool operator!=(const FCachedAssetKey& A, const FCachedAssetKey& B);
+inline uint32 GetTypeHash(const FCachedAssetKey& A);
 
-		enum { bAllowDuplicateKeys = false };
+/* 
+* Policy type for TSet<FAssetData*> to use FCachedAssetKey for hashing/equality.
+* This allows is to store just FAssetData* in the map without storing an extra copy of the key fields to save memory.
+*/
+struct FCachedAssetKeyFuncs
+{
+	using KeyInitType = FCachedAssetKey;
+	using ElementInitType = void; // TSet doesn't actually use this type 
 
-		static FORCEINLINE KeyInitType GetSetKey(const FAssetData* Element)
-		{
-			return FCachedAssetKey(*Element);
-		}
+	enum { bAllowDuplicateKeys = false };
 
-		static FORCEINLINE bool Matches(KeyInitType A, KeyInitType B)
-		{
-			return A == B;
-		}
+	static FORCEINLINE KeyInitType GetSetKey(const FAssetData* Element)
+	{
+		return FCachedAssetKey(*Element);
+	}
 
-		static FORCEINLINE uint32 GetKeyHash(KeyInitType Key)
-		{
-			return GetTypeHash(Key);
-		}
-	};
+	static FORCEINLINE bool Matches(KeyInitType A, KeyInitType B)
+	{
+		return A == B;
+	}
 
-	using FAssetDataMap = TSet<FAssetData*, FCachedAssetKeyFuncs>;
-	using FConstAssetDataMap = TSet<const FAssetData*, FCachedAssetKeyFuncs>;
-}
+	static FORCEINLINE uint32 GetKeyHash(KeyInitType Key)
+	{
+		return GetTypeHash(Key);
+	}
+};
+
+using FAssetDataMap = TSet<FAssetData*, FCachedAssetKeyFuncs>;
+using FConstAssetDataMap = TSet<const FAssetData*, FCachedAssetKeyFuncs>;
 
-/** The state of an asset registry, this is used internally by IAssetRegistry to represent the disk cache, and is also accessed directly to save/load cooked caches */
+} // namespace UE::AssetRegistry::Private
+
+/**
+ * The state of an asset registry, this is used internally by IAssetRegistry to represent the disk cache,
+ * and is also accessed directly to save/load cooked caches.
+ */
 class FAssetRegistryState
 {
 	using FCachedAssetKey = UE::AssetRegistry::Private::FCachedAssetKey;
 public:
-	// These types are an implementation detail and they and the functions which take/return them are subject to change without deprecation warnings.
+	// These types are an implementation detail and they and the functions which take/return them are subject to change
+	// without deprecation warnings.
 	using FAssetDataMap = UE::AssetRegistry::Private::FAssetDataMap;
 	using FConstAssetDataMap = UE::AssetRegistry::Private::FConstAssetDataMap;
 
@@ -236,9 +237,11 @@ public:
 	 * @param Filter filter to apply to the assets in the AssetRegistry
 	 * @param PackageNamesToSkip explicit list of packages to skip, because they were already added
 	 * @param OutAssetData the list of assets in this path
-	 * @param bSkipARFilteredAssets If true, skip assets that are skipped by UE::AssetRegistry::FFiltering (false by default)
+	 * @param bSkipARFilteredAssets If true, skip assets that are skipped by UE::AssetRegistry::FFiltering
+	          (false by default)
 	 */
-	ASSETREGISTRY_API bool GetAssets(const FARCompiledFilter& Filter, const TSet<FName>& PackageNamesToSkip, TArray<FAssetData>& OutAssetData, bool bSkipARFilteredAssets = false) const;
+	ASSETREGISTRY_API bool GetAssets(const FARCompiledFilter& Filter, const TSet<FName>& PackageNamesToSkip,
+		TArray<FAssetData>& OutAssetData, bool bSkipARFilteredAssets = false) const;
 
 	/**
 	 * Enumerate asset data for all assets that match the filter.
@@ -251,8 +254,10 @@ public:
 	 * @param bARFiltering Whether to apply filtering from UE::AssetRegistry::FFiltering (false by default)
 	 */
 	UE_DEPRECATED(5.5, "Use EnumerateAssets with InEnumerateFlags instead.")
-	ASSETREGISTRY_API bool EnumerateAssets(const FARCompiledFilter& Filter, const TSet<FName>& PackageNamesToSkip, TFunctionRef<bool(const FAssetData&)> Callback, bool bARFiltering) const;
-	ASSETREGISTRY_API bool EnumerateAssets(const FARCompiledFilter& Filter, const TSet<FName>& PackageNamesToSkip, TFunctionRef<bool(const FAssetData&)> Callback) const;
+	ASSETREGISTRY_API bool EnumerateAssets(const FARCompiledFilter& Filter, const TSet<FName>& PackageNamesToSkip,
+		TFunctionRef<bool(const FAssetData&)> Callback, bool bARFiltering) const;
+	ASSETREGISTRY_API bool EnumerateAssets(const FARCompiledFilter& Filter, const TSet<FName>& PackageNamesToSkip,
+		TFunctionRef<bool(const FAssetData&)> Callback) const;
 
 	/**
 	 * Enumerate asset data for all assets that match the filter.
@@ -265,7 +270,9 @@ public:
 	 * @param InEnumerateFlags flags to control enumeration and filtering.
 	 *        @see EEnumerateAssetsFlags.
 	 */
-	ASSETREGISTRY_API bool EnumerateAssets(const FARCompiledFilter& Filter, const TSet<FName>& PackageNamesToSkip, TFunctionRef<bool(const FAssetData&)> Callback, UE::AssetRegistry::EEnumerateAssetsFlags InEnumerateFlags) const;
+	ASSETREGISTRY_API bool EnumerateAssets(const FARCompiledFilter& Filter, const TSet<FName>& PackageNamesToSkip, 
+		TFunctionRef<bool(const FAssetData&)> Callback,
+		UE::AssetRegistry::EEnumerateAssetsFlags InEnumerateFlags) const;
 
 	/**
 	 * Gets asset data for all assets in the registry state.
@@ -274,7 +281,8 @@ public:
 	 * @param OutAssetData the list of assets
 	 * @param bARFiltering Whether to apply filtering from UE::AssetRegistry::FFiltering (false by default)
 	 */
-	ASSETREGISTRY_API bool GetAllAssets(const TSet<FName>& PackageNamesToSkip, TArray<FAssetData>& OutAssetData, bool bARFiltering = false) const;
+	ASSETREGISTRY_API bool GetAllAssets(const TSet<FName>& PackageNamesToSkip, TArray<FAssetData>& OutAssetData,
+		bool bARFiltering = false) const;
 
 	/**
 	 * Enumerates asset data for all assets in the registry state.
@@ -284,9 +292,11 @@ public:
 	 * @param bARFiltering Whether to apply filtering from UE::AssetRegistry::FFiltering (false by default)
 	 */
 	UE_DEPRECATED(5.5, "Use EnumerateAllAssets with InEnumerateFlags instead.")
-	ASSETREGISTRY_API bool EnumerateAllAssets(const TSet<FName>& PackageNamesToSkip, TFunctionRef<bool(const FAssetData&)> Callback, bool bARFiltering) const;
+	ASSETREGISTRY_API bool EnumerateAllAssets(const TSet<FName>& PackageNamesToSkip,
+		TFunctionRef<bool(const FAssetData&)> Callback, bool bARFiltering) const;
 	ASSETREGISTRY_API void EnumerateAllAssets(TFunctionRef<void(const FAssetData&)> Callback) const;
-	ASSETREGISTRY_API bool EnumerateAllAssets(const TSet<FName>& PackageNamesToSkip, TFunctionRef<bool(const FAssetData&)> Callback) const;
+	ASSETREGISTRY_API bool EnumerateAllAssets(const TSet<FName>& PackageNamesToSkip,
+		TFunctionRef<bool(const FAssetData&)> Callback) const;
 
 	/**
 	 * Enumerates asset data for all assets in the registry state.
@@ -296,7 +306,9 @@ public:
 	 * @param InEnumerateFlags flags to control enumeration and filtering.
 	 *        @see EEnumerateAssetsFlags.
 	 */
-	ASSETREGISTRY_API bool EnumerateAllAssets(const TSet<FName>& PackageNamesToSkip, TFunctionRef<bool(const FAssetData&)> Callback, UE::AssetRegistry::EEnumerateAssetsFlags InEnumerateFlags) const;
+	ASSETREGISTRY_API bool EnumerateAllAssets(const TSet<FName>& PackageNamesToSkip,
+		TFunctionRef<bool(const FAssetData&)> Callback,
+		UE::AssetRegistry::EEnumerateAssetsFlags InEnumerateFlags) const;
 	
 
 	/**
@@ -326,26 +338,46 @@ public:
 	ASSETREGISTRY_API FName GetFirstPackageByName(FStringView PackageName) const;
 
 	/**
-	 * Appends a list of packages and searchable names that are referenced by the supplied package or name. (On disk references ONLY)
+	 * Appends a list of packages and searchable names that are referenced by the supplied package or name.
+	 * (On disk references ONLY)
 	 *
 	 * @param AssetIdentifier	the name of the package/name for which to gather dependencies
 	 * @param OutDependencies	a list of things that are referenced by AssetIdentifier
-	 * @param Category	which category(ies) of dependencies to include in the output list. Dependencies matching ANY of the OR'd categories will be returned.
-	 * @param Flags	which flags are required present or not present on the dependencies. Dependencies matching ALL required and NONE excluded bits will be returned. For each potentially returned dependency, flags not applicable to their category are ignored.
+	 * @param Category	which category(ies) of dependencies to include in the output list.
+	 *        Dependencies matching ANY of the OR'd categories will be returned.
+	 * @param Flags	which flags are required present or not present on the dependencies.
+	 *        Dependencies matching ALL required and NONE excluded bits will be returned.
+	 *        For each potentially returned dependency, flags not applicable to their category are ignored.
 	 */
-	ASSETREGISTRY_API bool GetDependencies(const FAssetIdentifier& AssetIdentifier, TArray<FAssetIdentifier>& OutDependencies, UE::AssetRegistry::EDependencyCategory Category = UE::AssetRegistry::EDependencyCategory::All, const UE::AssetRegistry::FDependencyQuery& Flags = UE::AssetRegistry::FDependencyQuery()) const;
-	ASSETREGISTRY_API bool GetDependencies(const FAssetIdentifier& AssetIdentifier, TArray<FAssetDependency>& OutDependencies, UE::AssetRegistry::EDependencyCategory Category = UE::AssetRegistry::EDependencyCategory::All, const UE::AssetRegistry::FDependencyQuery& Flags = UE::AssetRegistry::FDependencyQuery()) const;
+	ASSETREGISTRY_API bool GetDependencies(const FAssetIdentifier& AssetIdentifier,
+		TArray<FAssetIdentifier>& OutDependencies,
+		UE::AssetRegistry::EDependencyCategory Category = UE::AssetRegistry::EDependencyCategory::All,
+		const UE::AssetRegistry::FDependencyQuery& Flags = UE::AssetRegistry::FDependencyQuery()) const;
+	ASSETREGISTRY_API bool GetDependencies(const FAssetIdentifier& AssetIdentifier,
+		TArray<FAssetDependency>& OutDependencies,
+		UE::AssetRegistry::EDependencyCategory Category = UE::AssetRegistry::EDependencyCategory::All,
+		const UE::AssetRegistry::FDependencyQuery& Flags = UE::AssetRegistry::FDependencyQuery()) const;
 
 	/**
-	 * Appends a list of packages and searchable names that reference the supplied package or name. (On disk references ONLY)
+	 * Appends a list of packages and searchable names that reference the supplied package or name.
+	 * (On disk references ONLY)
 	 *
 	 * @param AssetIdentifier	the name of the package/name for which to gather dependencies
 	 * @param OutReferencers	a list of things that reference AssetIdentifier
-	 * @param Category	which category(ies) of dependencies to include in the output list. Dependencies matching ANY of the OR'd categories will be returned.
-	 * @param Flags	which flags are required present or not present on the dependencies. Dependencies matching ALL required and NONE excluded bits will be returned. For each potentially returned dependency, flags not applicable to their category are ignored.
+	 * @param Category	which category(ies) of dependencies to include in the output list.
+	 *        Dependencies matching ANY of the OR'd categories will be returned.
+	 * @param Flags	which flags are required present or not present on the dependencies.
+	 *        Dependencies matching ALL required and NONE excluded bits will be returned.
+	 *        For each potentially returned dependency, flags not applicable to their category are ignored.
 	 */
-	ASSETREGISTRY_API bool GetReferencers(const FAssetIdentifier& AssetIdentifier, TArray<FAssetIdentifier>& OutReferencers, UE::AssetRegistry::EDependencyCategory Category = UE::AssetRegistry::EDependencyCategory::All, const UE::AssetRegistry::FDependencyQuery& Flags = UE::AssetRegistry::FDependencyQuery()) const;
-	ASSETREGISTRY_API bool GetReferencers(const FAssetIdentifier& AssetIdentifier, TArray<FAssetDependency>& OutReferencers, UE::AssetRegistry::EDependencyCategory Category = UE::AssetRegistry::EDependencyCategory::All, const UE::AssetRegistry::FDependencyQuery& Flags = UE::AssetRegistry::FDependencyQuery()) const;
+	ASSETREGISTRY_API bool GetReferencers(const FAssetIdentifier& AssetIdentifier,
+		TArray<FAssetIdentifier>& OutReferencers,
+		UE::AssetRegistry::EDependencyCategory Category = UE::AssetRegistry::EDependencyCategory::All,
+		const UE::AssetRegistry::FDependencyQuery& Flags = UE::AssetRegistry::FDependencyQuery()) const;
+	ASSETREGISTRY_API bool GetReferencers(const FAssetIdentifier& AssetIdentifier,
+		TArray<FAssetDependency>& OutReferencers,
+		UE::AssetRegistry::EDependencyCategory Category = UE::AssetRegistry::EDependencyCategory::All,
+		const UE::AssetRegistry::FDependencyQuery& Flags = UE::AssetRegistry::FDependencyQuery()) const;
 
 	/**
 	 * Gets the asset data for the specified object path
@@ -354,10 +386,7 @@ public:
 	 * @return the assets data, null if not found
 	 */
 	UE_DEPRECATED(5.1, "Asset path FNames have been deprecated, use FSoftObjectPath instead.")
-	const FAssetData* GetAssetByObjectPath(const FName ObjectPath) const
-	{
-		return GetAssetByObjectPath(FSoftObjectPath(ObjectPath.ToString()));
-	}
+	const FAssetData* GetAssetByObjectPath(const FName ObjectPath) const;
 
 	/**
 	 * Gets the asset data for the specified object path
@@ -365,17 +394,7 @@ public:
 	 * @param ObjectPath the path of the object to be looked up
 	 * @return the assets data, null if not found
 	 */
-	const FAssetData* GetAssetByObjectPath(const FSoftObjectPath& ObjectPath) const 
-	{
-		FCachedAssetKey Key(ObjectPath);
-		FAssetData* const* FoundAsset = CachedAssets.Find(Key);
-		if (FoundAsset)
-		{
-			return *FoundAsset;
-		}
-
-		return nullptr;
-	}
+	const FAssetData* GetAssetByObjectPath(const FSoftObjectPath& ObjectPath) const;
 
 	/**
 	 * Gets the asset data for the specified package name
@@ -383,15 +402,7 @@ public:
 	 * @param PackageName the path of the package to be looked up
 	 * @return an array of AssetData*, empty if nothing found
 	 */
-	TArrayView<FAssetData const* const> GetAssetsByPackageName(const FName PackageName) const
-	{
-		if (const TArray<FAssetData*, TInlineAllocator<1>>* FoundAssetArray = CachedAssetsByPackageName.Find(PackageName))
-		{
-			return MakeArrayView(*FoundAssetArray);
-		}
-
-		return TArrayView<FAssetData* const>();
-	}
+	TArrayView<FAssetData const* const> GetAssetsByPackageName(const FName PackageName) const;
 
 	/**
 	 * Gets the asset data for the specified asset class
@@ -408,39 +419,17 @@ public:
 	 * @param ClassPathName the class path name of the assets to look for
 	 * @return An array of AssetData*, empty if nothing found
 	 */
-	const TArray<const FAssetData*>& GetAssetsByClassPathName(const FTopLevelAssetPath ClassPathName) const
-	{
-		static TArray<const FAssetData*> InvalidArray;
-		const TArray<FAssetData*>* FoundAssetArray = CachedAssetsByClass.Find(ClassPathName);
-		if (FoundAssetArray)
-		{
-			return reinterpret_cast<const TArray<const FAssetData*>&>(*FoundAssetArray);
-		}
-
-		return InvalidArray;
-	}
+	const TArray<const FAssetData*>& GetAssetsByClassPathName(const FTopLevelAssetPath ClassPathName) const;
 
 	/**
 	 * Enumerates all tags of any asset in the AssetRegistry
 	 *
 	 * @param Callback the function for each tag
 	 */
-	void EnumerateTags(TFunctionRef<bool(FName TagName)> Callback) const
-	{
-		for (const TPair<FName, TSet<FAssetData*>>& Pair : CachedAssetsByTag)
-		{
-			if (!Callback(Pair.Key))
-			{
-				break;
-			}
-		}
-	}
+	void EnumerateTags(TFunctionRef<bool(FName TagName)> Callback) const;
 
 	/** Return whether the given TagName occurs in the tags of any asset in the AssetRegistry */
-	bool ContainsTag(FName TagName) const
-	{
-		return CachedAssetsByTag.Contains(TagName);
-	}
+	bool ContainsTag(FName TagName) const;
 
 	/**
 	 * Gets the asset data for the specified asset tag
@@ -449,12 +438,7 @@ public:
 	 * @return An array of AssetData*, empty if nothing found
 	 */
 	UE_DEPRECATED(5.5, "GetAssetsByTagName has been deprecated. Please use EnumerateAssetsByTagName")
-	const TArray<const FAssetData*>& GetAssetsByTagName(const FName TagName) const
-	{
-		ensureMsgf(false, TEXT("GetAssetsByTagName has been deprecated. Please use EnumerateAssetsByTagName"));
-		static TArray<const FAssetData*> InvalidArray;
-		return InvalidArray;
-	}
+	const TArray<const FAssetData*>& GetAssetsByTagName(const FName TagName) const;
 
 	/**
 	 * Enumerates the asset data for the specified asset tag
@@ -462,88 +446,38 @@ public:
 	 * @param TagName the tag name to search for
 	 * @param Callback the function called for each asset data
 	 */
-	void EnumerateAssetsByTagName(const FName TagName, TFunctionRef<bool(const FAssetData* AssetData)> Callback) const
-	{
-		const TSet<FAssetData*>* FoundAssets = CachedAssetsByTag.Find(TagName);
-		if (FoundAssets)
-		{
-			for (const FAssetData* AssetData : *FoundAssets)
-			{
-				if (!Callback(AssetData))
-				{
-					break;
-				}
-			}
-		}
-	}
+	void EnumerateAssetsByTagName(const FName TagName, TFunctionRef<bool(const FAssetData* AssetData)> Callback) const;
 
 	/**
-	 * Enumerates all tags of any asset in the AssetRegistry including function that can be called to enumerate assets for each tag
+	 * Enumerates all tags of any asset in the AssetRegistry including function that can be called to enumerate
+	 * assets for each tag
 	 *
 	 * @param Callback the function for each tag pair
 	 */
-	void EnumerateTagToAssetDatas(TFunctionRef<bool(FName TagName, IAssetRegistry::FEnumerateAssetDatasFunc EnumerateAssets)> Callback) const
-	{
-		for (const TPair<FName, TSet<FAssetData*>>& Pair : CachedAssetsByTag)
-		{
-			const bool bKeepEnumerating = Callback(Pair.Key, [&Pair](IAssetRegistry::FAssetDataFunc AssetCallback)
-				{
-					for (const FAssetData* AssetData : Pair.Value)
-					{
-						if (!AssetCallback(AssetData))
-						{
-							return false;
-						}
-					}
-
-					return true;
-				});
-
-			if (!bKeepEnumerating)
-			{
-				break;
-			}
-		}
-	}
+	void EnumerateTagToAssetDatas(
+		TFunctionRef<bool(FName TagName, IAssetRegistry::FEnumerateAssetDatasFunc EnumerateAssets)> Callback) const;
 
 	/** Returns const version of internal ObjectPath->AssetData map for fast iteration */
-	const FConstAssetDataMap& GetAssetDataMap() const
-	{
-		return reinterpret_cast<const FConstAssetDataMap&>(CachedAssets);
-	}
+	const FConstAssetDataMap& GetAssetDataMap() const;
 
 	/** Returns const version of internal Tag->AssetDatas map for fast iteration */
 	UE_DEPRECATED(5.5, "GetTagToAssetDatasMap has been deprecated. Please use EnumerateTags or EnumerateTagToAssetDatas")
-	const TMap<FName, const TArray<const FAssetData*>> GetTagToAssetDatasMap() const
-	{
-		ensureMsgf(false, TEXT("GetTagToAssetDatasMap has been deprecated. Please use EnumerateTags or EnumerateTagToAssetDatas"));
-		static TMap<FName, const TArray<const FAssetData*>> InvalidMap;
-		return InvalidMap;
-	}
+	const TMap<FName, const TArray<const FAssetData*>> GetTagToAssetDatasMap() const;
 
 	/** Returns const version of internal PackageName->PackageData map for fast iteration */
-	const TMap<FName, const FAssetPackageData*>& GetAssetPackageDataMap() const
-	{
-		return reinterpret_cast<const TMap<FName, const FAssetPackageData*>&>(CachedPackageData);
-	}
+	const TMap<FName, const FAssetPackageData*>& GetAssetPackageDataMap() const;
 
 	/** Get the set of primary assets contained in this state */
 	ASSETREGISTRY_API void GetPrimaryAssetsIds(TSet<FPrimaryAssetId>& OutPrimaryAssets) const;
 
 	/** Returns pointer to the asset package data */
 	ASSETREGISTRY_API const FAssetPackageData* GetAssetPackageData(FName PackageName) const;
-	ASSETREGISTRY_API const FAssetPackageData* GetAssetPackageData(FName PackageName, FName& OutCorrectCasePackageName) const;
+	ASSETREGISTRY_API const FAssetPackageData* GetAssetPackageData(FName PackageName,
+		FName& OutCorrectCasePackageName) const;
 	ASSETREGISTRY_API FAssetPackageData* GetAssetPackageData(FName PackageName);
 
 	/** Returns all package names */
-	void GetPackageNames(TArray<FName>& OutPackageNames) const
-	{
-		OutPackageNames.Reserve(CachedAssetsByPackageName.Num());
-		for (auto It = CachedAssetsByPackageName.CreateConstIterator(); It; ++It)
-		{
-			OutPackageNames.Add(It.Key());
-		}
-	}
+	void GetPackageNames(TArray<FName>& OutPackageNames) const;
 
 	/** Finds an existing package data, or creates a new one to modify */
 	ASSETREGISTRY_API FAssetPackageData* CreateOrGetAssetPackageData(FName PackageName);
@@ -562,8 +496,10 @@ public:
 	ASSETREGISTRY_API void UpdateAssetData(FAssetData&& NewAssetData, bool bCreateIfNotExists = false);
 
 	/** Updates an existing asset data with the new value and updates lookup maps */
-	ASSETREGISTRY_API void UpdateAssetData(FAssetData* AssetData, const FAssetData& NewAssetData, bool* bOutModified = nullptr);
-	ASSETREGISTRY_API void UpdateAssetData(FAssetData* AssetData, FAssetData&& NewAssetData, bool* bOutModified = nullptr);
+	ASSETREGISTRY_API void UpdateAssetData(FAssetData* AssetData, const FAssetData& NewAssetData,
+		bool* bOutModified = nullptr);
+	ASSETREGISTRY_API void UpdateAssetData(FAssetData* AssetData, FAssetData&& NewAssetData,
+		bool* bOutModified = nullptr);
 
 	/**
 	 * Updates all asset data package flags in the specified package
@@ -584,66 +520,84 @@ public:
 	 * Clear all dependencies of the given category from the given AssetIdentifier (e.g. package).
 	 * Also clears the referencer link from each of the dependencies.
 	 */
-	ASSETREGISTRY_API void ClearDependencies(const FAssetIdentifier& AssetIdentifier, UE::AssetRegistry::EDependencyCategory Category);
+	ASSETREGISTRY_API void ClearDependencies(const FAssetIdentifier& AssetIdentifier,
+		UE::AssetRegistry::EDependencyCategory Category);
 	/**
 	 * Add the given dependencies to the given AssetIdentifier (e.g. package).
 	 * Also adds a referencer link on each of the dependencies.
 	 */
-	ASSETREGISTRY_API void AddDependencies(const FAssetIdentifier& AssetIdentifier, TConstArrayView<FAssetDependency> Dependencies);
+	ASSETREGISTRY_API void AddDependencies(const FAssetIdentifier& AssetIdentifier,
+		TConstArrayView<FAssetDependency> Dependencies);
 	/**
 	 * Clears existing dependencies of the given Category(s) and assigns the input Dependencies. Gives an error if 
 	 * any elements of Dependencies are outside of the Category(s).
 	 */
-	ASSETREGISTRY_API void SetDependencies(const FAssetIdentifier& AssetIdentifier, TConstArrayView<FAssetDependency> Dependencies,
+	ASSETREGISTRY_API void SetDependencies(const FAssetIdentifier& AssetIdentifier,
+		TConstArrayView<FAssetDependency> Dependencies,
 		UE::AssetRegistry::EDependencyCategory Category = UE::AssetRegistry::EDependencyCategory::All);
 	/**
 	 * Clear all referencers of the given category from the given AssetIdentifier (e.g. package).
 	 * Also clears the dependency link from each of the referencers.
 	 */
-	ASSETREGISTRY_API void ClearReferencers(const FAssetIdentifier& AssetIdentifier, UE::AssetRegistry::EDependencyCategory Category);
+	ASSETREGISTRY_API void ClearReferencers(const FAssetIdentifier& AssetIdentifier,
+		UE::AssetRegistry::EDependencyCategory Category);
 	/**
 	 * Add a dependency on the given AssetIdentifier (e.g. package) from each of the Referencers.
 	 * Also adds a referencer link to each referencer on the AssetIdentifer's node.
 	 */
-	ASSETREGISTRY_API void AddReferencers(const FAssetIdentifier& AssetIdentifier, TConstArrayView<FAssetDependency> Referencers);
+	ASSETREGISTRY_API void AddReferencers(const FAssetIdentifier& AssetIdentifier,
+		TConstArrayView<FAssetDependency> Referencers);
 	/**
 	 * Clears existing referencers of the given Category(s) and assigns the input Referencers. Gives an error if
 	 * any elements of Referencers are outside of the Category(s).
 	 */
-	ASSETREGISTRY_API void SetReferencers(const FAssetIdentifier& AssetIdentifier, TConstArrayView<FAssetDependency> Referencers,
+	ASSETREGISTRY_API void SetReferencers(const FAssetIdentifier& AssetIdentifier,
+		TConstArrayView<FAssetDependency> Referencers,
 		UE::AssetRegistry::EDependencyCategory Category = UE::AssetRegistry::EDependencyCategory::All);
 
 	/** Resets to default state */
 	ASSETREGISTRY_API void Reset();
 
 	/** Initializes cache from existing set of asset data and depends nodes */
-	ASSETREGISTRY_API void InitializeFromExisting(const FAssetDataMap& AssetDataMap, const TMap<FAssetIdentifier, FDependsNode*>& DependsNodeMap, const TMap<FName, FAssetPackageData*>& AssetPackageDataMap, const FAssetRegistrySerializationOptions& Options, EInitializationMode InitializationMode = EInitializationMode::Rebuild);
-	void InitializeFromExisting(const FAssetRegistryState& Existing, const FAssetRegistrySerializationOptions& Options, EInitializationMode InitializationMode = EInitializationMode::Rebuild)
-	{
-		InitializeFromExisting(Existing.CachedAssets, Existing.CachedDependsNodes, Existing.CachedPackageData, Options, InitializationMode);
-	}
+	ASSETREGISTRY_API void InitializeFromExisting(const FAssetDataMap& AssetDataMap,
+		const TMap<FAssetIdentifier, FDependsNode*>& DependsNodeMap,
+		const TMap<FName, FAssetPackageData*>& AssetPackageDataMap,
+		const FAssetRegistrySerializationOptions& Options,
+		EInitializationMode InitializationMode = EInitializationMode::Rebuild);
+	void InitializeFromExisting(const FAssetRegistryState& Existing, const FAssetRegistrySerializationOptions& Options,
+		EInitializationMode InitializationMode = EInitializationMode::Rebuild);
 
 	/** 
 	 * Prunes an asset cache, this removes asset data, nodes, and package data that isn't needed. 
-	 * @param RequiredPackages If set, only these packages will be maintained. If empty it will keep all unless filtered by other parameters
+	 * @param RequiredPackages If set, only these packages will be maintained. If empty it will keep all
+	 *        unless filtered by other parameters
 	 * @param RemovePackages These packages will be removed from the current set
-	 * @param ChunksToKeep The list of chunks that are allowed to remain. Any assets in other chunks are pruned. If empty, all assets are kept regardless of chunk
+	 * @param ChunksToKeep The list of chunks that are allowed to remain. Any assets in other chunks are pruned.
+	 *        If empty, all assets are kept regardless of chunk
 	 * @param Options Serialization options to read filter info from
 	 */
-	ASSETREGISTRY_API void PruneAssetData(const TSet<FName>& RequiredPackages, const TSet<FName>& RemovePackages, const TSet<int32> ChunksToKeep, const FAssetRegistrySerializationOptions& Options);
-	ASSETREGISTRY_API void PruneAssetData(const TSet<FName>& RequiredPackages, const TSet<FName>& RemovePackages, const FAssetRegistrySerializationOptions& Options);
+	ASSETREGISTRY_API void PruneAssetData(const TSet<FName>& RequiredPackages, const TSet<FName>& RemovePackages,
+		const TSet<int32> ChunksToKeep, const FAssetRegistrySerializationOptions& Options);
+	ASSETREGISTRY_API void PruneAssetData(const TSet<FName>& RequiredPackages, const TSet<FName>& RemovePackages,
+		const FAssetRegistrySerializationOptions& Options);
 	ASSETREGISTRY_API void Prune(const FAssetRegistryPruneOptions& PruneOptions);
 
 	
 	/**
-	 * Initializes a cache from an existing using a set of filters. This is more efficient than calling InitalizeFromExisting and then PruneAssetData.
+	 * Initializes a cache from an existing using a set of filters.
+	 * This is more efficient than calling InitalizeFromExisting and then PruneAssetData.
+	 *
 	 * @param ExistingState State to use initialize from
-	 * @param RequiredPackages If set, only these packages will be maintained. If empty it will keep all unless filtered by other parameters
+	 * @param RequiredPackages If set, only these packages will be maintained.
+	 *        If empty it will keep all unless filtered by other parameters.
 	 * @param RemovePackages These packages will be removed from the current set
-	 * @param ChunksToKeep The list of chunks that are allowed to remain. Any assets in other chunks are pruned. If empty, all assets are kept regardless of chunk
+	 * @param ChunksToKeep The list of chunks that are allowed to remain. Any assets in other chunks are pruned.
+	 *        If empty, all assets are kept regardless of chunk
 	 * @param Options Serialization options to read filter info from
 	 */
-	ASSETREGISTRY_API void InitializeFromExistingAndPrune(const FAssetRegistryState& ExistingState, const TSet<FName>& RequiredPackages, const TSet<FName>& RemovePackages, const TSet<int32> ChunksToKeep, const FAssetRegistrySerializationOptions& Options);
+	ASSETREGISTRY_API void InitializeFromExistingAndPrune(const FAssetRegistryState& ExistingState,
+		const TSet<FName>& RequiredPackages, const TSet<FName>& RemovePackages, const TSet<int32> ChunksToKeep,
+		const FAssetRegistrySerializationOptions& Options);
 
 	/** Edit every AssetData's Tags to remove Tags that are filtered out by the filtering rules in Options */
 	ASSETREGISTRY_API void FilterTags(const FAssetRegistrySerializationOptions& Options);
@@ -654,14 +608,17 @@ public:
 
 	/** Save without editor-only data */
 	ASSETREGISTRY_API bool Save(FArchive& Ar, const FAssetRegistrySerializationOptions& Options);
-	ASSETREGISTRY_API bool Load(FArchive& Ar, const FAssetRegistryLoadOptions& Options = FAssetRegistryLoadOptions(), FAssetRegistryVersion::Type* OutVersion = nullptr);
+	ASSETREGISTRY_API bool Load(FArchive& Ar,
+		const FAssetRegistryLoadOptions& Options = FAssetRegistryLoadOptions(),
+		FAssetRegistryVersion::Type* OutVersion = nullptr);
 
 	/** 
 	* Example Usage:
-	*	FAssetRegistryState AssetRegistry;
-	*	bool bSucceeded = FAssetRegistryState::LoadFromDisk(TEXT("Path/To/AR"), FAssetRegistryLoadOptions(), AssetRegistry);
+	*	FAssetRegistryState AR;
+	*	bool bSucceeded = FAssetRegistryState::LoadFromDisk(TEXT("Path/To/AR"), FAssetRegistryLoadOptions(), AR);
 	*/
-	static ASSETREGISTRY_API bool LoadFromDisk(const TCHAR* InPath, const FAssetRegistryLoadOptions& InOptions, FAssetRegistryState& OutState, FAssetRegistryVersion::Type* OutVersion = nullptr);
+	static ASSETREGISTRY_API bool LoadFromDisk(const TCHAR* InPath, const FAssetRegistryLoadOptions& InOptions,
+		FAssetRegistryState& OutState, FAssetRegistryVersion::Type* OutVersion = nullptr);
 
 	/** Returns memory size of entire registry, optionally logging sizes */
 	ASSETREGISTRY_API SIZE_T GetAllocatedSize(bool bLogDetailed = false) const;
@@ -678,12 +635,18 @@ public:
 #if ASSET_REGISTRY_STATE_DUMPING_ENABLED
 	/**
 	 * Writes out the state in textual form. Use arguments to control which segments to emit.
-	 * @param Arguments List of segments to emit. Possible values: 'ObjectPath', 'PackageName', 'Path', 'Class', 'Tag', 'Dependencies' and 'PackageData'
-	 * @param OutPages Textual representation will be written to this array; each entry will have LinesPerPage lines of the full dump.
-	 * @param LinesPerPage - how many lines should be combined into each string element of OutPages, for e.g. breaking up the dump into separate files.
-	 *        To facilitate diffing between similar-but-different registries, the actual number of lines per page will be slightly less than LinesPerPage; we introduce partially deterministic pagebreaks near the end of each page.
+	 * @param Arguments List of segments to emit.
+	 *        Possible values: 'ObjectPath', 'PackageName', 'Path', 'Class', 'Tag', 'Dependencies' and 'PackageData'.
+	 * @param OutPages Textual representation will be written to this array;
+	 *        each entry will have LinesPerPage lines of the full dump.
+	 * @param LinesPerPage - how many lines should be combined into each string element of OutPages,
+	 *        for e.g. breaking up the dump into separate files.
+	 *        To facilitate diffing between similar-but-different registries, the actual number of lines per page will
+	 *        be slightly less than LinesPerPage; we introduce partially deterministic pagebreaks near the end of
+	 *        each page.
 	 */
-	ASSETREGISTRY_API void Dump(const TArray<FString>& Arguments, TArray<FString>& OutPages, int32 LinesPerPage=1) const;
+	ASSETREGISTRY_API void Dump(const TArray<FString>& Arguments, TArray<FString>& OutPages,
+		int32 LinesPerPage=1) const;
 #endif
 
 private:
@@ -694,7 +657,8 @@ private:
 	void SetAssetDatas(TArrayView<FAssetData> AssetDatas, const FAssetRegistryLoadOptions& Options);
 
 	/** Find the first non-redirector dependency node starting from InDependency. */
-	FDependsNode* ResolveRedirector(FDependsNode* InDependency, const FAssetDataMap& InAllowedAssets, TMap<FDependsNode*, FDependsNode*>& InCache);
+	FDependsNode* ResolveRedirector(FDependsNode* InDependency, const FAssetDataMap& InAllowedAssets,
+		TMap<FDependsNode*, FDependsNode*>& InCache);
 
 	/** Finds an existing node for the given package and returns it, or returns null if one isn't found */
 	FDependsNode* FindDependsNode(const FAssetIdentifier& Identifier) const;
@@ -706,7 +670,8 @@ private:
 	bool RemoveDependsNode(const FAssetIdentifier& Identifier);
 
 	/** Filter a set of tags and output a copy of the filtered set. */
-	static void FilterTags(const FAssetDataTagMapSharedView& InTagsAndValues, FAssetDataTagMap& OutTagsAndValues, const TSet<FName>* ClassSpecificFilterList, const FAssetRegistrySerializationOptions & Options);
+	static void FilterTags(const FAssetDataTagMapSharedView& InTagsAndValues, FAssetDataTagMap& OutTagsAndValues,
+		const TSet<FName>* ClassSpecificFilterList, const FAssetRegistrySerializationOptions & Options);
 
 	void LoadDependencies(FArchive& Ar);
 	void LoadDependencies_BeforeFlags(FArchive& Ar, bool bSerializeDependencies, FAssetRegistryVersion::Type Version);
@@ -719,12 +684,17 @@ private:
 		bool& bOutRemovedAssetData, bool& bOutRemovedPackageData);
 
 	/**
-	 * Returns true if the given package should be filtered from the results because the package belongs to an unmounted content path.
-	 * This can only happen when loading a cooked asset registry (@see bCookedGlobalAssetRegistryState), as it may contain state for plugins that are not currently loaded.
+	 * Returns true if the given package should be filtered from the results because the package belongs
+	 * to an unmounted content path.
+	 * This can only happen when loading a cooked asset registry (@see bCookedGlobalAssetRegistryState),
+	 * as it may contain state for plugins that are not currently loaded.
 	 */
 	bool IsPackageUnmountedAndFiltered(const FName PackageName) const;
 
-	/** Set of asset data for assets saved to disk. Searched via path name types, implicitly converted to FCachedAssetKey. */
+	/**
+	 * Set of asset data for assets saved to disk. Searched via path name types,
+	 * implicitly converted to FCachedAssetKey.
+	 */
 	FAssetDataMap CachedAssets;
 
 	/** The map of package names to asset data for assets saved to disk */
@@ -745,12 +715,18 @@ private:
 	/** A map of Package Names to Package Data */
 	TMap<FName, FAssetPackageData*> CachedPackageData;
 
-	/** When loading a registry from disk, we can allocate all the FAssetData objects in one chunk, to save on 10s of thousands of heap allocations */
+	/**
+	 * When loading a registry from disk, we can allocate all the FAssetData objects in one chunk,
+	 * to save on 10s of thousands of heap allocations.
+	 */
 	TArray<FAssetData*> PreallocatedAssetDataBuffers;
 	TArray<FDependsNode*> PreallocatedDependsNodeDataBuffers;
 	TArray<FAssetPackageData*> PreallocatedPackageDataBuffers;
 
-	/** Counters for asset/depends data memory allocation to ensure that every FAssetData and FDependsNode created is deleted */
+	/**
+	 * Counters for asset/depends data memory allocation to ensure that every FAssetData and
+	 * FDependsNode created is deleted.
+	 */
 	int32 NumAssets = 0;
 	int32 NumDependsNodes = 0;
 	int32 NumPackageData = 0;
@@ -764,113 +740,292 @@ private:
 
 namespace UE::AssetRegistry::Private
 {
-	FORCEINLINE uint32 HashCombineQuick(uint32 A, uint32 B)
-	{
-		return A ^ (B + 0x9e3779b9 + (A << 6) + (A >> 2));
-	}
 
-	inline FCachedAssetKey::FCachedAssetKey(const FAssetData* InAssetData)
+FORCEINLINE uint32 HashCombineQuick(uint32 A, uint32 B)
+{
+	return A ^ (B + 0x9e3779b9 + (A << 6) + (A >> 2));
+}
+
+inline FCachedAssetKey::FCachedAssetKey(const FAssetData* InAssetData)
+{
+	if (!InAssetData)
 	{
-		if (!InAssetData)
-		{
-			return;
-		}
+		return;
+	}
 
 #if WITH_EDITORONLY_DATA
-		if (!InAssetData->GetOptionalOuterPathName().IsNone())
-		{
-			OuterPath = InAssetData->GetOptionalOuterPathName();
-		}
-		else
+	if (!InAssetData->GetOptionalOuterPathName().IsNone())
+	{
+		OuterPath = InAssetData->GetOptionalOuterPathName();
+	}
+	else
 #endif
-		{
-			OuterPath = InAssetData->PackageName;
-		}
-		ObjectName = InAssetData->AssetName;
+	{
+		OuterPath = InAssetData->PackageName;
 	}
+	ObjectName = InAssetData->AssetName;
+}
+
+inline FCachedAssetKey::FCachedAssetKey(const FAssetData& InAssetData)
+	: FCachedAssetKey(&InAssetData)
+{
+}
 
-	inline FCachedAssetKey::FCachedAssetKey(const FAssetData& InAssetData)
-		: FCachedAssetKey(&InAssetData)
+inline FCachedAssetKey::FCachedAssetKey(FTopLevelAssetPath InAssetPath)
+	: OuterPath(InAssetPath.GetPackageName())
+	, ObjectName(InAssetPath.GetAssetName())
+{
+}
+
+inline FCachedAssetKey::FCachedAssetKey(const FSoftObjectPath& InObjectPath)
+{
+	if (InObjectPath.GetAssetFName().IsNone())
 	{
+		// Packages themselves never appear in the asset registry
+		return;
 	}
-
-	inline FCachedAssetKey::FCachedAssetKey(FTopLevelAssetPath InAssetPath)
-		: OuterPath(InAssetPath.GetPackageName())
-		, ObjectName(InAssetPath.GetAssetName())
+	else if (InObjectPath.GetSubPathString().IsEmpty())
 	{
+		// If InObjectPath represents a top-level asset we can just take the existing FNames.
+		OuterPath = InObjectPath.GetLongPackageFName();
+		ObjectName = InObjectPath.GetAssetFName();
+	}
+	else
+	{
+		// If InObjectPath represents a subobject we need to split the path into the path of the outer and the name of the innermost object.
+		TStringBuilder<FName::StringBufferSize> Builder;
+		InObjectPath.ToString(Builder);
+
+		const FAssetPathParts Parts = SplitIntoOuterPathAndAssetName(Builder);
+
+		// This should be impossible as at bare minimum concatenating the package name and asset name should add a separator
+		check(!Parts.OuterPath.IsEmpty() && !Parts.InnermostName.IsEmpty()); 
+
+		// Don't create FNames for this query struct. If the AssetData exists to find, the FName will already exist due to OptionalOuterPath on FAssetData.
+		OuterPath = FName(Parts.OuterPath, FNAME_Find); 
+		ObjectName = FName(Parts.InnermostName);
 	}
+}
+inline FString FCachedAssetKey::ToString() const
+{
+	TStringBuilder<FName::StringBufferSize> Builder;
+	AppendString(Builder);
+	return FString(Builder);
+}
 
-	inline FCachedAssetKey::FCachedAssetKey(const FSoftObjectPath& InObjectPath)
+inline int32 FCachedAssetKey::Compare(const FCachedAssetKey& Other) const
+{
+	if (OuterPath == Other.OuterPath)
 	{
-		if (InObjectPath.GetAssetFName().IsNone())
-		{
-			// Packages themselves never appear in the asset registry
-			return;
-		}
-		else if (InObjectPath.GetSubPathString().IsEmpty())
-		{
-			// If InObjectPath represents a top-level asset we can just take the existing FNames.
-			OuterPath = InObjectPath.GetLongPackageFName();
-			ObjectName = InObjectPath.GetAssetFName();
-		}
-		else
-		{
-			// If InObjectPath represents a subobject we need to split the path into the path of the outer and the name of the innermost object.
-			TStringBuilder<FName::StringBufferSize> Builder;
-			InObjectPath.ToString(Builder);
+		return ObjectName.CompareIndexes(Other.ObjectName);
+	}
+	else
+	{
+		return OuterPath.CompareIndexes(Other.OuterPath);
+	}
+}
 
-			const FAssetPathParts Parts = SplitIntoOuterPathAndAssetName(Builder);
+inline void FCachedAssetKey::AppendString(FStringBuilderBase& Builder) const
+{
+	ConcatenateOuterPathAndObjectName(Builder, OuterPath, ObjectName);
+}
+
+inline FStringBuilderBase& operator<<(FStringBuilderBase& Builder, const FCachedAssetKey& Key)
+{
+	Key.AppendString(Builder);
+	return Builder;
+}
 
-			// This should be impossible as at bare minimum concatenating the package name and asset name should add a separator
-			check(!Parts.OuterPath.IsEmpty() && !Parts.InnermostName.IsEmpty()); 
+inline bool operator==(const FCachedAssetKey& A, const FCachedAssetKey& B)
+{
+	return A.OuterPath == B.OuterPath && A.ObjectName == B.ObjectName;
+}
 
-			// Don't create FNames for this query struct. If the AssetData exists to find, the FName will already exist due to OptionalOuterPath on FAssetData.
-			OuterPath = FName(Parts.OuterPath, FNAME_Find); 
-			ObjectName = FName(Parts.InnermostName);
-		}
+inline bool operator!=(const FCachedAssetKey& A, const FCachedAssetKey& B)
+{
+	return A.OuterPath != B.OuterPath || A.ObjectName != B.ObjectName;
+}
+
+inline uint32 GetTypeHash(const FCachedAssetKey& A)
+{
+	return HashCombineQuick(GetTypeHash(A.OuterPath), GetTypeHash(A.ObjectName));
+}
+
+} // namespace UE::AssetRegistry::Private
+
+
+///////////////////////////////////////////////////////
+// Inline implementations
+///////////////////////////////////////////////////////
+
+
+inline FAssetRegistrySerializationOptions::FAssetRegistrySerializationOptions(
+	UE::AssetRegistry::ESerializationTarget Target)
+{
+	if (Target == UE::AssetRegistry::ESerializationTarget::ForDevelopment)
+	{
+		InitForDevelopment();
 	}
-	inline FString FCachedAssetKey::ToString() const
+}
+
+inline void FAssetRegistrySerializationOptions::DisableFilters()
+{
+	bFilterAssetDataWithNoTags = false;
+	bFilterDependenciesWithNoTags = false;
+	bFilterSearchableNames = false;
+}
+
+inline void FAssetRegistrySerializationOptions::InitForDevelopment()
+{
+	bSerializeAssetRegistry = true;
+	bSerializeDependencies = true;
+	bSerializeSearchableNameDependencies = true;
+	bSerializeManageDependencies = true;
+	bSerializePackageData = true;
+	DisableFilters();
+}
+
+inline FAssetRegistryLoadOptions::FAssetRegistryLoadOptions(const FAssetRegistrySerializationOptions& Options)
+	: bLoadDependencies(Options.bSerializeDependencies)
+	, bLoadPackageData(Options.bSerializePackageData)
+{
+}
+
+inline const FAssetData* FAssetRegistryState::GetAssetByObjectPath(const FName ObjectPath) const
+{
+	return GetAssetByObjectPath(FSoftObjectPath(ObjectPath.ToString()));
+}
+
+inline const FAssetData* FAssetRegistryState::GetAssetByObjectPath(const FSoftObjectPath& ObjectPath) const
+{
+	FCachedAssetKey Key(ObjectPath);
+	FAssetData* const* FoundAsset = CachedAssets.Find(Key);
+	if (FoundAsset)
 	{
-		TStringBuilder<FName::StringBufferSize> Builder;
-		AppendString(Builder);
-		return FString(Builder);
+		return *FoundAsset;
 	}
 
-	inline int32 FCachedAssetKey::Compare(const FCachedAssetKey& Other) const
+	return nullptr;
+}
+
+inline TArrayView<FAssetData const* const> FAssetRegistryState::GetAssetsByPackageName(const FName PackageName) const
+{
+	if (const TArray<FAssetData*, TInlineAllocator<1>>*FoundAssetArray = CachedAssetsByPackageName.Find(PackageName))
 	{
-		if (OuterPath == Other.OuterPath)
-		{
-			return ObjectName.CompareIndexes(Other.ObjectName);
-		}
-		else
-		{
-			return OuterPath.CompareIndexes(Other.OuterPath);
-		}
+		return MakeArrayView(*FoundAssetArray);
 	}
 
-	inline void FCachedAssetKey::AppendString(FStringBuilderBase& Builder) const
+	return TArrayView<FAssetData* const>();
+}
+
+inline const TArray<const FAssetData*>& FAssetRegistryState::GetAssetsByClassPathName(
+	const FTopLevelAssetPath ClassPathName) const
+{
+	static TArray<const FAssetData*> InvalidArray;
+	const TArray<FAssetData*>* FoundAssetArray = CachedAssetsByClass.Find(ClassPathName);
+	if (FoundAssetArray)
 	{
-		ConcatenateOuterPathAndObjectName(Builder, OuterPath, ObjectName);
+		return reinterpret_cast<const TArray<const FAssetData*>&>(*FoundAssetArray);
 	}
 
-	inline FStringBuilderBase& operator<<(FStringBuilderBase& Builder, const FCachedAssetKey& Key)
+	return InvalidArray;
+}
+
+inline void FAssetRegistryState::EnumerateTags(TFunctionRef<bool(FName TagName)> Callback) const
+{
+	for (const TPair<FName, TSet<FAssetData*>>& Pair : CachedAssetsByTag)
 	{
-		Key.AppendString(Builder);
-		return Builder;
+		if (!Callback(Pair.Key))
+		{
+			break;
+		}
 	}
+}
 
-	inline bool operator==(const FCachedAssetKey& A, const FCachedAssetKey& B)
+inline bool FAssetRegistryState::ContainsTag(FName TagName) const
+{
+	return CachedAssetsByTag.Contains(TagName);
+}
+
+inline const TArray<const FAssetData*>& FAssetRegistryState::GetAssetsByTagName(const FName TagName) const
+{
+	ensureMsgf(false, TEXT("GetAssetsByTagName has been deprecated. Please use EnumerateAssetsByTagName"));
+	static TArray<const FAssetData*> InvalidArray;
+	return InvalidArray;
+}
+
+inline void FAssetRegistryState::EnumerateAssetsByTagName(const FName TagName,
+	TFunctionRef<bool(const FAssetData* AssetData)> Callback) const
+{
+	const TSet<FAssetData*>* FoundAssets = CachedAssetsByTag.Find(TagName);
+	if (FoundAssets)
 	{
-		return A.OuterPath == B.OuterPath && A.ObjectName == B.ObjectName;
+		for (const FAssetData* AssetData : *FoundAssets)
+		{
+			if (!Callback(AssetData))
+			{
+				break;
+			}
+		}
 	}
+}
 
-	inline bool operator!=(const FCachedAssetKey& A, const FCachedAssetKey& B)
+inline void FAssetRegistryState::EnumerateTagToAssetDatas(
+	TFunctionRef<bool(FName TagName, IAssetRegistry::FEnumerateAssetDatasFunc EnumerateAssets)> Callback) const
+{
+	for (const TPair<FName, TSet<FAssetData*>>& Pair : CachedAssetsByTag)
 	{
-		return A.OuterPath != B.OuterPath || A.ObjectName != B.ObjectName;
+		const bool bKeepEnumerating = Callback(Pair.Key, [&Pair](IAssetRegistry::FAssetDataFunc AssetCallback)
+			{
+				for (const FAssetData* AssetData : Pair.Value)
+				{
+					if (!AssetCallback(AssetData))
+					{
+						return false;
+					}
+				}
+
+				return true;
+			});
+
+		if (!bKeepEnumerating)
+		{
+			break;
+		}
 	}
+}
+
+inline const FAssetRegistryState::FConstAssetDataMap& FAssetRegistryState::GetAssetDataMap() const
+{
+	return reinterpret_cast<const FConstAssetDataMap&>(CachedAssets);
+}
+
+inline const TMap<FName, const TArray<const FAssetData*>> FAssetRegistryState::GetTagToAssetDatasMap() const
+{
+	ensureMsgf(false,
+		TEXT("GetTagToAssetDatasMap has been deprecated. Please use EnumerateTags or EnumerateTagToAssetDatas"));
+	static TMap<FName, const TArray<const FAssetData*>> InvalidMap;
+	return InvalidMap;
+}
 
-	inline uint32 GetTypeHash(const FCachedAssetKey& A)
+inline const TMap<FName, const FAssetPackageData*>& FAssetRegistryState::GetAssetPackageDataMap() const
+{
+	return reinterpret_cast<const TMap<FName, const FAssetPackageData*>&>(CachedPackageData);
+}
+
+inline void FAssetRegistryState::GetPackageNames(TArray<FName>& OutPackageNames) const
+{
+	OutPackageNames.Reserve(CachedAssetsByPackageName.Num());
+	for (auto It = CachedAssetsByPackageName.CreateConstIterator(); It; ++It)
 	{
-		return HashCombineQuick(GetTypeHash(A.OuterPath), GetTypeHash(A.ObjectName));
+		OutPackageNames.Add(It.Key());
 	}
 }
+
+inline void FAssetRegistryState::InitializeFromExisting(const FAssetRegistryState& Existing,
+	const FAssetRegistrySerializationOptions& Options,
+	EInitializationMode InitializationMode)
+{
+	InitializeFromExisting(Existing.CachedAssets, Existing.CachedDependsNodes,
+		Existing.CachedPackageData, Options, InitializationMode);
+}
