AssetRegistry - Memory optimization - In runtime configurations, remove CachedAssetsByTag and implement EnumerateAssetsByTagName (more slowly but with less memory) using a new CachedClassesByTag and the existing CachedAssetsByClass.
#rnx

diff --git a/Engine/Source/Runtime/AssetRegistry/Private/AssetRegistryState.cpp b/Engine/Source/Runtime/AssetRegistry/Private/AssetRegistryState.cpp
--- a/Engine/Source/Runtime/AssetRegistry/Private/AssetRegistryState.cpp
+++ b/Engine/Source/Runtime/AssetRegistry/Private/AssetRegistryState.cpp
@@ -44,7 +44,11 @@ FAssetRegistryState& FAssetRegistryState::operator=(FAssetRegistryState&& Rhs)
 	CachedAssetsByPackageName			= MoveTemp(Rhs.CachedAssetsByPackageName);
 	CachedAssetsByPath					= MoveTemp(Rhs.CachedAssetsByPath);
 	CachedAssetsByClass					= MoveTemp(Rhs.CachedAssetsByClass);
+#if UE_ASSETREGISTRY_CACHEDASSETSBYTAG
 	CachedAssetsByTag					= MoveTemp(Rhs.CachedAssetsByTag);
+#else
+	CachedClassesByTag					= MoveTemp(Rhs.CachedClassesByTag);
+#endif
 	CachedDependsNodes					= MoveTemp(Rhs.CachedDependsNodes);
 	CachedPackageData					= MoveTemp(Rhs.CachedPackageData);
 	PreallocatedAssetDataBuffers		= MoveTemp(Rhs.PreallocatedAssetDataBuffers);
@@ -150,7 +154,11 @@ void FAssetRegistryState::Reset()
 	CachedAssetsByPackageName.Empty();
 	CachedAssetsByPath.Empty();
 	CachedAssetsByClass.Empty();
+#if UE_ASSETREGISTRY_CACHEDASSETSBYTAG
 	CachedAssetsByTag.Empty();
+#else
+	CachedClassesByTag.Empty();
+#endif
 	CachedDependsNodes.Empty();
 	CachedPackageData.Empty();
 }
@@ -945,7 +953,87 @@ void FilterAssets(TArray<const FAssetData*>& InOutResults, const AccelerationMap
 	}
 }
 
+#if !UE_ASSETREGISTRY_CACHEDASSETSBYTAG
+template<typename CallbackType>
+void FilterAssetsByCachedClassesByTag(TArray<const FAssetData*>& InOutResults,
+	const TMap<FName, TSet<FTopLevelAssetPath>>& CachedClassesByTag,
+	const TMap<FTopLevelAssetPath, TArray<FAssetData*>>& CachedAssetsByClass,
+	const TMultiMap<FName, TOptional<FString>>& TagsAndValues,
+	CallbackType&& FunctionToKeepAsset, int32 FilterComplexity)
+{
+	TArray<TArray<const FAssetData*>, TInlineAllocator<10>> Matches;
+	Matches.Reserve(TagsAndValues.Num());
+	uint32 TotalMatches = 0;
+
+	for (const TPair<FName, TOptional<FString>>& TagPair : TagsAndValues)
+	{
+		TArray<const FAssetData*>& Results = Matches.Emplace_GetRef();
+
+		// The lists of assets in CachedAssetsByClass are non-intersecting (each list is only the exact instances of that
+		// class and does not include subclasses), so we do not need to handle removing duplicates when merging lists from
+		// multiple classes.
+		if (const TSet<FTopLevelAssetPath>* TagClasses = CachedClassesByTag.Find(TagPair.Key))
+		{
+			for (const FTopLevelAssetPath& ClassPath : *TagClasses)
+			{
+				if (const TArray<FAssetData*>* ClassAssets = CachedAssetsByClass.Find(ClassPath))
+				{
+					Results.Append(*ClassAssets);
+				}
+			}
+		}
+		// Some assets are in a class that could have the tag, but the specific asset actually does not have the tag.
+		for (TArray<const FAssetData*>::TIterator Iter(Results); Iter; ++Iter)
+		{
+			const FAssetData* AssetData = *Iter;
+			if (!AssetDataMatchesTag(AssetData, TagPair))
+			{
+				Iter.RemoveCurrentSwap();
+			}
+		}
+		TotalMatches += Results.Num();
+	}
+
+	if (InOutResults.IsEmpty())
+	{
+		// No previous Results; set Results equal to the values found
+		InOutResults.Reserve(TotalMatches);
+		for (TArray<const FAssetData*>& Assets : Matches)
+		{
+			InOutResults.Append(Assets);
+		}
+		// Remove duplicates
+		Algo::Sort(InOutResults);
+		InOutResults.SetNum(Algo::Unique(InOutResults));
+	}
+	else
+	{
+		bool bUseFiltering = DecideIntersectionMethod(InOutResults.Num(), TotalMatches, FilterComplexity);
+		if (bUseFiltering)
+		{
+			InOutResults.RemoveAllSwap([&FunctionToKeepAsset](const FAssetData* AssetData)
+				{
+					return !FunctionToKeepAsset(AssetData);
+				});
+		}
+		else
+		{
+			// Convert Array of Arrays into format required by ArrayIntersection: Array of ArrayViews
+			TArray<TArrayView<const FAssetData*>, TInlineAllocator<10>> ArrayViewMatches;
+			ArrayViewMatches.Reserve(Matches.Num());
+			for (TArray<const FAssetData*>& MatchesElement : Matches)
+			{
+				ArrayViewMatches.Emplace(MatchesElement);
+			}
+
+			// ArrayIntersection handles removing any duplicates from Matches
+			ArrayIntersection(InOutResults, ArrayViewMatches, TotalMatches);
+		}
+	}
 }
+#endif
+
+} // namespace UE::AssetRegistry::Private
 
 bool FAssetRegistryState::EnumerateAssets(const FARCompiledFilter& Filter, const TSet<FName>& PackageNamesToSkip,
 	TFunctionRef<bool(const FAssetData&)> Callback, bool bSkipARFilteredAssets) const
@@ -1052,7 +1140,12 @@ bool FAssetRegistryState::EnumerateAssets(const FARCompiledFilter& Filter, const
 
 	if (Filter.TagsAndValues.Num() > 0)
 	{
-		FilterAssets(AccumulatedResults, CachedAssetsByTag, Filter.TagsAndValues,
+#if UE_ASSETREGISTRY_CACHEDASSETSBYTAG
+		FilterAssets(AccumulatedResults, CachedAssetsByTag,
+#else
+		FilterAssetsByCachedClassesByTag(AccumulatedResults, CachedClassesByTag, CachedAssetsByClass,
+#endif
+			Filter.TagsAndValues,
 			[&Filter](const FAssetData* AssetData)
 			{
 				for (const TPair<FName, TOptional<FString>>& TagPair : Filter.TagsAndValues)
@@ -1831,7 +1924,11 @@ SIZE_T FAssetRegistryState::GetAllocatedSize(bool bLogDetailed) const
 	MapMemory += CachedAssetsByPackageName.GetAllocatedSize();
 	MapMemory += CachedAssetsByPath.GetAllocatedSize();
 	MapMemory += CachedAssetsByClass.GetAllocatedSize();
+#if UE_ASSETREGISTRY_CACHEDASSETSBYTAG
 	MapMemory += CachedAssetsByTag.GetAllocatedSize();
+#else
+	MapMemory += CachedClassesByTag.GetAllocatedSize();
+#endif
 	MapMemory += CachedDependsNodes.GetAllocatedSize();
 	MapMemory += CachedPackageData.GetAllocatedSize();
 	MapMemory += PreallocatedAssetDataBuffers.GetAllocatedSize();
@@ -1855,7 +1952,11 @@ SIZE_T FAssetRegistryState::GetAllocatedSize(bool bLogDetailed) const
 		MapArrayMemory += Pair.Value.GetAllocatedSize();
 	}
 
+#if UE_ASSETREGISTRY_CACHEDASSETSBYTAG
 	SubArray(CachedAssetsByTag);
+#else
+	SubArray(CachedClassesByTag);
+#endif
 
 	if (bLogDetailed)
 	{
@@ -2094,6 +2195,7 @@ void FAssetRegistryState::SetAssetDatas(TArrayView<FAssetData> AssetDatas, const
 		}
 		ShrinkMultimap(CachedAssetsByClass);
 
+#if UE_ASSETREGISTRY_CACHEDASSETSBYTAG
 		CachedAssetsByTag.Empty();
 		for (FAssetData& AssetData : AssetDatas)
 		{
@@ -2109,6 +2211,23 @@ void FAssetRegistryState::SetAssetDatas(TArrayView<FAssetData> AssetDatas, const
 		{
 			Pair.Value.Shrink();
 		}
+#else
+		CachedClassesByTag.Empty();
+		for (FAssetData& AssetData : AssetDatas)
+		{
+			for (const TPair<FName, FAssetTagValueRef>& Pair : AssetData.TagsAndValues)
+			{
+				TSet<FTopLevelAssetPath>& TagClasses = CachedClassesByTag.FindOrAdd(Pair.Key);
+				TagClasses.Add(AssetData.AssetClassPath);
+			}
+		}
+
+		CachedClassesByTag.Shrink();
+		for (TPair<FName, TSet<FTopLevelAssetPath>>& Pair : CachedClassesByTag)
+		{
+			Pair.Value.Shrink();
+		}
+#endif
 	};
 
 	if (Options.ParallelWorkers <= 1)
@@ -2153,8 +2272,13 @@ void FAssetRegistryState::AddAssetData(FAssetData* AssetData)
 	{
 		FName Key = TagIt.Key();
 
+#if UE_ASSETREGISTRY_CACHEDASSETSBYTAG
 		TSet<FAssetData*>& TagAssets = CachedAssetsByTag.FindOrAdd(Key);
 		TagAssets.Add(AssetData);
+#else
+		TSet<FTopLevelAssetPath>& TagClasses = CachedClassesByTag.FindOrAdd(Key);
+		TagClasses.Add(AssetData->AssetClassPath);
+#endif
 	}
 }
 
@@ -2183,12 +2307,20 @@ void FAssetRegistryState::AddTagsToAssetData(const FSoftObjectPath& InObjectPath
 void FAssetRegistryState::FilterTags(const FAssetRegistrySerializationOptions& Options)
 {
 	// Calling SetTagsOnExistingAsset for any changed tags might be slow.
-	// For cases where many Assets change it might be faster to recreate CachedAssetsByTag
-	// rather than trying to update its elements for each Asset change
+	// For cases where many Assets change it might be faster to recreate CachedAssetsByTag/CachedClassesByTag rather
+	// than trying to update its elements for each Asset change. For that reason we (currently) always recreate
+	// CachedAssetsByTag/CachedClassesByTag.
+#if UE_ASSETREGISTRY_CACHEDASSETSBYTAG
 	for (TPair<FName, TSet<FAssetData*>>& Pair : CachedAssetsByTag)
 	{
 		Pair.Value.Reset();
 	}
+#else
+	for (TPair<FName, TSet<FTopLevelAssetPath>>& Pair : CachedClassesByTag)
+	{
+		Pair.Value.Reset();
+	}
+#endif
 
 	for (FAssetData* AssetData : CachedAssets) 
 	{
@@ -2202,11 +2334,15 @@ void FAssetRegistryState::FilterTags(const FAssetRegistrySerializationOptions& O
 			AssetData->TagsAndValues = FAssetDataTagMapSharedView(MoveTemp(LocalTagsAndValues));
 		}
 
-		// Add the AssetData to all its CachedAssetsByTag keys even if nothing changed, because
+		// Add the AssetData to all its CachedAssetsByTag/CachedClassesByTag keys even if nothing changed, because
 		// we are reconstructing all CachedAssetsByTag.
 		for (auto TagIt = AssetData->TagsAndValues.CreateConstIterator(); TagIt; ++TagIt)
 		{
+#if UE_ASSETREGISTRY_CACHEDASSETSBYTAG
 			CachedAssetsByTag.FindOrAdd(TagIt.Key()).Add(AssetData);
+#else
+			CachedClassesByTag.FindOrAdd(TagIt.Key()).Add(AssetData->AssetClassPath);
+#endif
 		}
 	}
 }
@@ -2220,11 +2356,16 @@ void FAssetRegistryState::SetTagsOnExistingAsset(FAssetData* AssetData, FAssetDa
 
 		if (!NewTags.Contains(FNameKey))
 		{
+#if UE_ASSETREGISTRY_CACHEDASSETSBYTAG
 			TSet<FAssetData*>* OldTagAssets = CachedAssetsByTag.Find(FNameKey);
 			if (OldTagAssets)
 			{
 				OldTagAssets->Remove(AssetData);
 			}
+#else
+			// For CachedClassesByTag, we do not need to remove the asset's class from the entries
+			// for the old tags. The class is not changed and still has the possibility of containing the tags.
+#endif
 		}
 	}
 	// Update the tag cache map to add added tags
@@ -2234,7 +2375,11 @@ void FAssetRegistryState::SetTagsOnExistingAsset(FAssetData* AssetData, FAssetDa
 
 		if (!AssetData->TagsAndValues.Contains(FNameKey))
 		{
+#if UE_ASSETREGISTRY_CACHEDASSETSBYTAG
 			CachedAssetsByTag.FindOrAdd(FNameKey).Add(AssetData);
+#else
+			CachedClassesByTag.FindOrAdd(FNameKey).Add(AssetData->AssetClassPath);
+#endif
 		}
 	}
 	AssetData->TagsAndValues = FAssetDataTagMapSharedView(MoveTemp(NewTags));
@@ -2349,8 +2494,13 @@ void FAssetRegistryState::UpdateAssetData(FAssetData* AssetData, FAssetData&& Ne
 
 			if (!NewAssetData.TagsAndValues.Contains(FNameKey))
 			{
+#if UE_ASSETREGISTRY_CACHEDASSETSBYTAG
 				TSet<FAssetData*>* OldTagAssets = CachedAssetsByTag.Find(FNameKey);
 				OldTagAssets->Remove(AssetData);
+#else
+				// For CachedClassesByTag, we do not need to remove the asset's class from the entries
+				// for the old tags. The class is not changed and still has the possibility of containing the tags.
+#endif
 			}
 		}
 
@@ -2360,8 +2510,13 @@ void FAssetRegistryState::UpdateAssetData(FAssetData* AssetData, FAssetData&& Ne
 
 			if (!AssetData->TagsAndValues.Contains(FNameKey))
 			{
+#if UE_ASSETREGISTRY_CACHEDASSETSBYTAG
 				TSet<FAssetData*>& NewTagAssets = CachedAssetsByTag.FindOrAdd(FNameKey);
 				NewTagAssets.Add(AssetData);
+#else
+				TSet<FTopLevelAssetPath>& NewTagClasses = CachedClassesByTag.FindOrAdd(FNameKey);
+				NewTagClasses.Add(AssetData->AssetClassPath);
+#endif
 			}
 		}
 	}
@@ -2479,11 +2634,16 @@ void FAssetRegistryState::RemoveAssetData(FAssetData* AssetData, const FCachedAs
 	OldPathAssets->RemoveSingleSwap(AssetData);
 	OldClassAssets->RemoveSingleSwap(AssetData);
 
+#if UE_ASSETREGISTRY_CACHEDASSETSBYTAG
 	for (auto TagIt = AssetData->TagsAndValues.CreateConstIterator(); TagIt; ++TagIt)
 	{
 		TSet<FAssetData*>* OldTagAssets = CachedAssetsByTag.Find(TagIt.Key());
 		OldTagAssets->Remove(AssetData);
 	}
+#else
+	// For CachedClassesByTag, we do not need to remove the asset's class from the entries
+	// for the old tags. The class is not changed and still has the possibility of containing the tags.
+#endif
 
 	// Only remove dependencies and package data if there are no other known assets in the package
 	if (OldPackageAssets->Num() == 0)
@@ -2681,6 +2841,104 @@ const TArray<const FAssetData*>& FAssetRegistryState::GetAssetsByClassName(const
 	return InvalidArray;
 }
 
+void FAssetRegistryState::EnumerateAssetsByTagName(const FName TagName,
+	TFunctionRef<bool(const FAssetData* AssetData)> Callback) const
+{
+#if UE_ASSETREGISTRY_CACHEDASSETSBYTAG
+	const TSet<FAssetData*>* FoundAssets = CachedAssetsByTag.Find(TagName);
+	if (FoundAssets)
+	{
+		for (const FAssetData* AssetData : *FoundAssets)
+		{
+			if (!Callback(AssetData))
+			{
+				break;
+			}
+		}
+	}
+#else
+	const TSet<FTopLevelAssetPath>* FoundClasses = CachedClassesByTag.Find(TagName);
+	if (!FoundClasses)
+	{
+		return;
+	}
+
+	// The lists of assets in CachedAssetsByClass are non-intersecting (each list is only the exact instances of that
+	// class and does not include subclasses), so we do not need to handle removing duplicates when merging lists from
+	// multiple classes.
+	TArray<FAssetData*> PossibleAssets;
+	for (const FTopLevelAssetPath& ClassPath : *FoundClasses)
+	{
+		const TArray<FAssetData*>* ClassAssets = CachedAssetsByClass.Find(ClassPath);
+		if (ClassAssets)
+		{
+			PossibleAssets.Append(*ClassAssets);
+		}
+	}
+
+	for (const FAssetData* AssetData : PossibleAssets)
+	{
+		// Some assets are in a class that could have the tag, but the specific asset actually does not have the tag.
+		if (AssetData->FindTag(TagName))
+		{
+			if (!Callback(AssetData))
+			{
+				break;
+			}
+		}
+	}
+#endif
+}
+
+void FAssetRegistryState::EnumerateTagToAssetDatas(
+	TFunctionRef<bool(FName TagName, IAssetRegistry::FEnumerateAssetDatasFunc EnumerateAssets)> Callback) const
+{
+#if UE_ASSETREGISTRY_CACHEDASSETSBYTAG
+	for (const TPair<FName, TSet<FAssetData*>>& Pair : CachedAssetsByTag)
+	{
+		const bool bKeepEnumerating = Callback(Pair.Key, [&Pair](IAssetRegistry::FAssetDataFunc AssetCallback)
+			{
+				for (const FAssetData* AssetData : Pair.Value)
+				{
+					if (!AssetCallback(AssetData))
+					{
+						return false;
+					}
+				}
+
+				return true;
+			});
+
+		if (!bKeepEnumerating)
+		{
+			break;
+		}
+	}
+#else
+	for (const TPair<FName, TSet<FTopLevelAssetPath>>& Pair : CachedClassesByTag)
+	{
+		const bool bKeepEnumerating = Callback(Pair.Key, [&Pair, this](IAssetRegistry::FAssetDataFunc AssetCallback)
+			{
+				EnumerateAssetsByTagName(Pair.Key, [&AssetCallback](const FAssetData* AssetData)
+					{
+						if (!AssetCallback(AssetData))
+						{
+							return false;
+						}
+						return true;
+					});
+
+				return true;
+			});
+
+		if (!bKeepEnumerating)
+		{
+			break;
+		}
+	}
+#endif
+}
+
 bool FAssetRegistryState::IsPackageUnmountedAndFiltered(const FName PackageName) const
 {
 	// TODO: This can be removed once UE-178174 is fixed, as there will no longer be unmounted content to enumerate
@@ -2814,8 +3072,9 @@ static void PrintAssetDataMap(FString Name, const MapType& AssetMap, TStringBuil
 			Items.Add(It);
 		}
 		Items.Sort([](const FAssetData& A, const FAssetData& B)
-			{ return A.GetSoftObjectPath().LexicalLess(B.GetSoftObjectPath()); }
-		);
+			{
+				return A.GetSoftObjectPath().LexicalLess(B.GetSoftObjectPath());
+			});
 
 		PageBuffer.Append(TEXT("\t"));
 		Key.AppendString(PageBuffer);
@@ -2837,6 +3096,58 @@ static void PrintAssetDataMap(FString Name, const MapType& AssetMap, TStringBuil
 	AddLine();
 };
 
+template <typename MapType>
+static void PrintClassDataMap(FString Name, const MapType& ClassPathMap, TStringBuilder<16>& PageBuffer,
+	const TFunctionRef<void()>& AddLine)
+{
+	PageBuffer.Appendf(TEXT("--- Begin %s ---"), *Name);
+	AddLine();
+
+	TArray<typename MapType::KeyType> Keys;
+	ClassPathMap.GenerateKeyArray(Keys);
+
+	Keys.Sort(FPrintAssetDataMapKeyIsLess<typename MapType::KeyType>());
+
+	TArray<FTopLevelAssetPath> Items;
+	Items.Reserve(1024);
+
+	int32 ValidCount = 0;
+	for (const typename MapType::KeyType& Key : Keys)
+	{
+		const TSet<FTopLevelAssetPath>& ClassPaths = ClassPathMap.FindChecked(Key);
+		if (ClassPaths.Num() == 0)
+		{
+			continue;
+		}
+		++ValidCount;
+
+		Items.Reset();
+		Items.Reserve(ClassPaths.Num());
+		for (const FTopLevelAssetPath& ClassPath : ClassPaths)
+		{
+			Items.Add(ClassPath);
+		}
+		Items.Sort([](const FTopLevelAssetPath& A, const FTopLevelAssetPath& B)
+			{
+				return A.Compare(B) < 0;
+			});
+
+		PageBuffer.Append(TEXT("\t"));
+		Key.AppendString(PageBuffer);
+		PageBuffer.Appendf(TEXT(" : %d item(s)"), Items.Num());
+		AddLine();
+		for (const FTopLevelAssetPath& Data : Items)
+		{
+			PageBuffer.Append(TEXT("\t\t"));
+			Data.AppendString(PageBuffer);
+			AddLine();
+		}
+	}
+
+	PageBuffer.Appendf(TEXT("--- End %s : %d entries ---"), *Name, ValidCount);
+	AddLine();
+};
+
 void FAssetRegistryState::Dump(const TArray<FString>& Arguments, TArray<FString>& OutPages, int32 LinesPerPage) const
 {
 	int32 ExpectedNumLines = 14 + CachedAssets.Num() * 5 + CachedDependsNodes.Num() + CachedPackageData.Num();
@@ -3009,11 +3320,15 @@ void FAssetRegistryState::Dump(const TArray<FString>& Arguments, TArray<FString>
 	if (Arguments.Contains(TEXT("Tag")))
 	{
 		AddPageBreak();
+#if UE_ASSETREGISTRY_CACHEDASSETSBYTAG
 		PrintAssetDataMap(TEXT("CachedAssetsByTag"), CachedAssetsByTag, PageBuffer, AddLine,
 			[&PageBuffer, &AddLine](const FName& TagName, const FAssetData& Data)
 			{
 				PageBuffer << TEXT(", ") << Data.TagsAndValues.FindTag(TagName).ToLoose();
 			});
+#else
+		PrintClassDataMap(TEXT("CachedClassesByTag"), CachedClassesByTag, PageBuffer, AddLine);
+#endif
 	}
 
 	TArray<FAssetData*> SortedAssets;
diff --git a/Engine/Source/Runtime/AssetRegistry/Public/AssetRegistry/AssetRegistryState.h b/Engine/Source/Runtime/AssetRegistry/Public/AssetRegistry/AssetRegistryState.h
--- a/Engine/Source/Runtime/AssetRegistry/Public/AssetRegistry/AssetRegistryState.h
+++ b/Engine/Source/Runtime/AssetRegistry/Public/AssetRegistry/AssetRegistryState.h
@@ -22,6 +22,24 @@
 #include "CoreMinimal.h"
 #endif
 
+/**
+ * UE_ASSETREGISTRY_CACHEDASSETSBYTAG: If non-zero, the CachedAssetsByTag field is defined and used. If zero,
+ * CachedClassesByTag is defined and used. Both of these are used only in queries for assets by tag.
+ * 
+ * If CachedAssetsByTag is defined, the queries are as fast as possible, but a large amount of memory is used.
+ * If not, then the queries are instead executed using a three step process: CachedClassesByTag to find all classes
+ * with the tag, CachedAssetsByClass to find all assets in those classes, and then a filtering step
+ * on the resulting list of assets.
+ * 
+ * The amount of memory used for each:
+ * CachedAssetsByTag:  (number of assets) *(average number of tags per asset)*sizeof(Pointer)
+ * CachedClassesByTag: (number of classes)*(average number of tags per class)*sizeof(FTopLevelAssetPath)
+ * CachedClassesByTag is much smaller because number of classes is smaller than number of assets.
+ */
+#ifndef UE_ASSETREGISTRY_CACHEDASSETSBYTAG
+#define UE_ASSETREGISTRY_CACHEDASSETSBYTAG WITH_EDITORONLY_DATA
+#endif 
+
 class FArchive;
 class FAssetDataTagMap;
 class FAssetDataTagMapSharedView;
@@ -446,7 +464,8 @@ public:
 	 * @param TagName the tag name to search for
 	 * @param Callback the function called for each asset data
 	 */
-	void EnumerateAssetsByTagName(const FName TagName, TFunctionRef<bool(const FAssetData* AssetData)> Callback) const;
+	ASSETREGISTRY_API void EnumerateAssetsByTagName(const FName TagName,
+		TFunctionRef<bool(const FAssetData* AssetData)> Callback) const;
 
 	/**
 	 * Enumerates all tags of any asset in the AssetRegistry including function that can be called to enumerate
@@ -454,7 +473,7 @@ public:
 	 *
 	 * @param Callback the function for each tag pair
 	 */
-	void EnumerateTagToAssetDatas(
+	ASSETREGISTRY_API void EnumerateTagToAssetDatas(
 		TFunctionRef<bool(FName TagName, IAssetRegistry::FEnumerateAssetDatasFunc EnumerateAssets)> Callback) const;
 
 	/** Returns const version of internal ObjectPath->AssetData map for fast iteration */
@@ -706,8 +725,13 @@ private:
 	/** The map of class name to asset data for assets saved to disk */
 	TMap<FTopLevelAssetPath, TArray<FAssetData*> > CachedAssetsByClass;
 
+#if UE_ASSETREGISTRY_CACHEDASSETSBYTAG
 	/** The map of asset tag to asset data for assets saved to disk */
 	TMap<FName, TSet<FAssetData*> > CachedAssetsByTag;
+#else
+	/** The map of asset tag to asset data for assets saved to disk */
+	TMap<FName, TSet<FTopLevelAssetPath> > CachedClassesByTag;
+#endif
 
 	/** A map of object names to dependency data */
 	TMap<FAssetIdentifier, FDependsNode*> CachedDependsNodes;
@@ -933,7 +957,11 @@ inline const TArray<const FAssetData*>& FAssetRegistryState::GetAssetsByClassPat
 
 inline void FAssetRegistryState::EnumerateTags(TFunctionRef<bool(FName TagName)> Callback) const
 {
+#if UE_ASSETREGISTRY_CACHEDASSETSBYTAG
 	for (const TPair<FName, TSet<FAssetData*>>& Pair : CachedAssetsByTag)
+#else
+	for (const TPair<FName, TSet<FTopLevelAssetPath>>& Pair : CachedClassesByTag)
+#endif
 	{
 		if (!Callback(Pair.Key))
 		{
@@ -944,7 +972,11 @@ inline void FAssetRegistryState::EnumerateTags(TFunctionRef<bool(FName TagName)>
 
 inline bool FAssetRegistryState::ContainsTag(FName TagName) const
 {
+#if UE_ASSETREGISTRY_CACHEDASSETSBYTAG
 	return CachedAssetsByTag.Contains(TagName);
+#else
+	return CachedClassesByTag.Contains(TagName);
+#endif
 }
 
 inline const TArray<const FAssetData*>& FAssetRegistryState::GetAssetsByTagName(const FName TagName) const
@@ -954,47 +986,6 @@ inline const TArray<const FAssetData*>& FAssetRegistryState::GetAssetsByTagName(
 	return InvalidArray;
 }
 
-inline void FAssetRegistryState::EnumerateAssetsByTagName(const FName TagName,
-	TFunctionRef<bool(const FAssetData* AssetData)> Callback) const
-{
-	const TSet<FAssetData*>* FoundAssets = CachedAssetsByTag.Find(TagName);
-	if (FoundAssets)
-	{
-		for (const FAssetData* AssetData : *FoundAssets)
-		{
-			if (!Callback(AssetData))
-			{
-				break;
-			}
-		}
-	}
-}
-
-inline void FAssetRegistryState::EnumerateTagToAssetDatas(
-	TFunctionRef<bool(FName TagName, IAssetRegistry::FEnumerateAssetDatasFunc EnumerateAssets)> Callback) const
-{
-	for (const TPair<FName, TSet<FAssetData*>>& Pair : CachedAssetsByTag)
-	{
-		const bool bKeepEnumerating = Callback(Pair.Key, [&Pair](IAssetRegistry::FAssetDataFunc AssetCallback)
-			{
-				for (const FAssetData* AssetData : Pair.Value)
-				{
-					if (!AssetCallback(AssetData))
-					{
-						return false;
-					}
-				}
-
-				return true;
-			});
-
-		if (!bKeepEnumerating)
-		{
-			break;
-		}
-	}
-}
-
 inline const FAssetRegistryState::FConstAssetDataMap& FAssetRegistryState::GetAssetDataMap() const
 {
 	return reinterpret_cast<const FConstAssetDataMap&>(CachedAssets);
