AssetRegistry: Remove direct access to the internal containers on FAssetRegistryState so that we can change them to be more complicated compressed structures in an upcoming change.
Remove some functions that were deprecated in 5.1 .
In AssetRegistryState implementation, tweak some function local variables to work for the upcoming code.
#rnx

diff --git a/Engine/Plugins/Editor/AssetManagerEditor/Source/AssetManagerEditor/Private/TreeView/SAssetTableTreeView.cpp b/Engine/Plugins/Editor/AssetManagerEditor/Source/AssetManagerEditor/Private/TreeView/SAssetTableTreeView.cpp
--- a/Engine/Plugins/Editor/AssetManagerEditor/Source/AssetManagerEditor/Private/TreeView/SAssetTableTreeView.cpp
+++ b/Engine/Plugins/Editor/AssetManagerEditor/Source/AssetManagerEditor/Private/TreeView/SAssetTableTreeView.cpp
@@ -1766,9 +1766,8 @@ void SAssetTableTreeView::RefreshAssets()
 
 		for (int32 SourceAssetIndex = 0; SourceAssetIndex < SourceAssets.Num(); SourceAssetIndex++)
 		{
-			TArrayView<FAssetData const* const> AssetsInSourcePackage = RegistrySource.GetOwnedRegistryState()->GetAssetsByPackageName(SourceAssets[SourceAssetIndex].PackageName);
-
-			for (FAssetData const* const SourceAsset : AssetsInSourcePackage)
+			RegistrySource.GetOwnedRegistryState()->EnumerateAssetsByPackageName(SourceAssets[SourceAssetIndex].PackageName,
+				[&AssetToIndexMap, this, &AssetTable, &PluginToSizeMap](const FAssetData* SourceAsset)
 			{
 				if (AssetToIndexMap.Find(*SourceAsset) == nullptr)
 				{
@@ -1786,7 +1785,8 @@ void SAssetTableTreeView::RefreshAssets()
 						PluginToSizeMap.Add(AssetRow.GetPluginName(), AssetRow.GetStagedCompressedSizeRequiredInstall());
 					}
 				}
-			}
+				return true; // Keep iterating
+			});
 		}
 
 		int32 EntryIndex = AssetTable->GetAssets().Num();
@@ -1809,7 +1809,7 @@ void SAssetTableTreeView::RefreshAssets()
 
 			FAssetIdentifier CurrentIdentifier(SourceAssets[SourceAssetIndex].PackageName);
 			// We'll have to add the dependencies to all these entries
-			TArrayView<FAssetData const* const> AssetsInSourcePackage = RegistrySource.GetOwnedRegistryState()->GetAssetsByPackageName(SourceAssets[SourceAssetIndex].PackageName);
+			TArray<const FAssetData*> AssetsInSourcePackage = RegistrySource.GetOwnedRegistryState()->CopyAssetsByPackageName(SourceAssets[SourceAssetIndex].PackageName);
 
 			// Get the dependencies
 			TArray<FAssetIdentifier> DependencyList;
@@ -1819,8 +1819,10 @@ void SAssetTableTreeView::RefreshAssets()
 
 			for (const FAssetIdentifier& Dependency : DependencyList)
 			{
-				TArrayView<FAssetData const* const> AssetsInDependencyPackage = RegistrySource.GetOwnedRegistryState()->GetAssetsByPackageName(Dependency.PackageName);
-				for (FAssetData const* const DependencyAsset : AssetsInDependencyPackage)
+				RegistrySource.GetOwnedRegistryState()->EnumerateAssetsByPackageName(Dependency.PackageName,
+					[&AssetToIndexMap, &IndicesOfDependenciesInRowTableToAddToCurrentSourceAssetRow, this,
+					&AssetTable, &SourceAssets, &PluginToSizeMap]
+					(const FAssetData* DependencyAsset)
 				{
 					if (int32* DependencyIndex = AssetToIndexMap.Find(*DependencyAsset))
 					{
@@ -1845,7 +1847,8 @@ void SAssetTableTreeView::RefreshAssets()
 							PluginToSizeMap.Add(NewRow.GetPluginName(), NewRow.GetStagedCompressedSizeRequiredInstall());
 						}
 					}
-				}
+					return true; // Keep iterating
+				});
 			}
 
 			if (IndicesOfDependenciesInRowTableToAddToCurrentSourceAssetRow.Num())
diff --git a/Engine/Plugins/Editor/AssetRegistryExport/Source/Private/AssetRegistryExportCommandlet.cpp b/Engine/Plugins/Editor/AssetRegistryExport/Source/Private/AssetRegistryExportCommandlet.cpp
--- a/Engine/Plugins/Editor/AssetRegistryExport/Source/Private/AssetRegistryExportCommandlet.cpp
+++ b/Engine/Plugins/Editor/AssetRegistryExport/Source/Private/AssetRegistryExportCommandlet.cpp
@@ -334,7 +334,7 @@ int32 UAssetRegistryExportCommandlet::Main(const FString& CmdLineParams)
 					FName PackageName;
 					while (PackageQueue.Dequeue(PackageName))
 					{
-						TArrayView<const FAssetData* const> AssetDataList = AssetRegistry.GetAssetsByPackageName(PackageName);
+						TArray<const FAssetData*> AssetDataList = AssetRegistry.CopyAssetsByPackageName(PackageName);
 
 						TArray<FAssetIdentifier> PackageDependencies;
 						AssetRegistry.GetDependencies(FAssetIdentifier(PackageName), PackageDependencies);
@@ -366,7 +366,9 @@ int32 UAssetRegistryExportCommandlet::Main(const FString& CmdLineParams)
 					Line << PrimaryAsset->GetObjectPathString() << ",";
 					Line << PackageDependency.ToString() << ",";
 
-					const FAssetData* UseAsClassData = UE::AssetRegistry::GetMostImportantAsset(AssetRegistry.GetAssetsByPackageName(PackageDependency), UE::AssetRegistry::EGetMostImportantAssetFlags::IgnoreSkipClasses);
+					const FAssetData* UseAsClassData = UE::AssetRegistry::GetMostImportantAsset(
+						AssetRegistry.CopyAssetsByPackageName(PackageDependency),
+						UE::AssetRegistry::EGetMostImportantAssetFlags::IgnoreSkipClasses);
 					if (UseAsClassData)
 					{
 						Line << UseAsClassData->AssetClassPath.ToString() << ",";
@@ -428,7 +430,9 @@ int32 UAssetRegistryExportCommandlet::Main(const FString& CmdLineParams)
 						Pair.Value.UniquePackageDependencies.Add(PackageDependency);
 						UnionUniqueDependencies.Add(PackageDependency);
 
-						const FAssetData* UseAsClassData = UE::AssetRegistry::GetMostImportantAsset(AssetRegistry.GetAssetsByPackageName(PackageDependency), UE::AssetRegistry::EGetMostImportantAssetFlags::IgnoreSkipClasses);
+						const FAssetData* UseAsClassData = UE::AssetRegistry::GetMostImportantAsset(
+							AssetRegistry.CopyAssetsByPackageName(PackageDependency),
+							UE::AssetRegistry::EGetMostImportantAssetFlags::IgnoreSkipClasses);
 						if (UseAsClassData)
 						{
 							uint64 CompressedSize = 0;
@@ -461,7 +465,9 @@ int32 UAssetRegistryExportCommandlet::Main(const FString& CmdLineParams)
 					Line << "Unassigned,";
 					Line << Package.Key.ToString() << ",";
 
-					const FAssetData* UseAsClassData = UE::AssetRegistry::GetMostImportantAsset(AssetRegistry.GetAssetsByPackageName(Package.Key), UE::AssetRegistry::EGetMostImportantAssetFlags::IgnoreSkipClasses);
+					const FAssetData* UseAsClassData = UE::AssetRegistry::GetMostImportantAsset(
+						AssetRegistry.CopyAssetsByPackageName(Package.Key),
+						UE::AssetRegistry::EGetMostImportantAssetFlags::IgnoreSkipClasses);
 					if (UseAsClassData)
 					{
 						Line << UseAsClassData->AssetClassPath.ToString() << ",";
@@ -502,7 +508,9 @@ int32 UAssetRegistryExportCommandlet::Main(const FString& CmdLineParams)
 					Line << "Shared,";
 					Line << Package.Key.ToString() << ",";
 
-					const FAssetData* UseAsClassData = UE::AssetRegistry::GetMostImportantAsset(AssetRegistry.GetAssetsByPackageName(Package.Key), UE::AssetRegistry::EGetMostImportantAssetFlags::IgnoreSkipClasses);
+					const FAssetData* UseAsClassData = UE::AssetRegistry::GetMostImportantAsset(
+						AssetRegistry.CopyAssetsByPackageName(Package.Key),
+						UE::AssetRegistry::EGetMostImportantAssetFlags::IgnoreSkipClasses);
 					if (UseAsClassData)
 					{
 						Line << UseAsClassData->AssetClassPath.ToString() << ",";
@@ -615,7 +623,7 @@ int32 UAssetRegistryExportCommandlet::Main(const FString& CmdLineParams)
 
 				UE_LOG(LogAssetRegistryExport, Warning, TEXT("...Orphaned: %s (%s bytes, in packagemap: %s"), *LD->PackageName.ToString(), *NumberString(value), *LexToString(AllPackages.Contains(LD->PackageName)));
 
-				TArrayView<const FAssetData* const> PackageAssets = AssetRegistry.GetAssetsByPackageName(LD->PackageName);
+				TArray<const FAssetData*> PackageAssets = AssetRegistry.CopyAssetsByPackageName(LD->PackageName);
 				for (const FAssetData* AD : PackageAssets)
 				{
 					uint64 blah = 0;
diff --git a/Engine/Source/Developer/IoStoreUtilities/Private/IoStoreUtilities.cpp b/Engine/Source/Developer/IoStoreUtilities/Private/IoStoreUtilities.cpp
--- a/Engine/Source/Developer/IoStoreUtilities/Private/IoStoreUtilities.cpp
+++ b/Engine/Source/Developer/IoStoreUtilities/Private/IoStoreUtilities.cpp
@@ -4119,7 +4119,9 @@ static void UpdatePluginMetadataAndWriteJsons(
 
 		// Grab the most important asset out and use it to track largest asset classes for the plugin.
 		// This might be null!
-		const FAssetData* AssetData = UE::AssetRegistry::GetMostImportantAsset(AssetRegistry.GetAssetsByPackageName(AssetPackage.Key), UE::AssetRegistry::EGetMostImportantAssetFlags::IgnoreSkipClasses);
+		const FAssetData* AssetData = UE::AssetRegistry::GetMostImportantAsset(
+			AssetRegistry.CopyAssetsByPackageName(AssetPackage.Key),
+			UE::AssetRegistry::EGetMostImportantAssetFlags::IgnoreSkipClasses);
 		
 		const TArray<FIoStoreChunkSource, TInlineAllocator<2>>* PackageChunks = PackageToChunks.Find(FPackageId::FromName(AssetPackage.Key));
 		if (PackageChunks == nullptr)
@@ -4482,7 +4484,9 @@ static void AddChunkInfoToAssetRegistry(TMap<FPackageId, TArray<FIoStoreChunkSou
 
 		FPackageId PackageId = FPackageId::FromName(AssetPackage.Key);
 
-		const FAssetData* AssetData = UE::AssetRegistry::GetMostImportantAsset(AssetRegistry.GetAssetsByPackageName(AssetPackage.Key), UE::AssetRegistry::EGetMostImportantAssetFlags::IgnoreSkipClasses);
+		const FAssetData* AssetData = UE::AssetRegistry::GetMostImportantAsset(
+			AssetRegistry.CopyAssetsByPackageName(AssetPackage.Key),
+			UE::AssetRegistry::EGetMostImportantAssetFlags::IgnoreSkipClasses);
 		if (AssetData == nullptr)
 		{
 			// e.g. /Script packages.
@@ -9611,14 +9615,22 @@ int32 CreateIoStoreContainerFiles(const TCHAR* CmdLine)
 						// skip over packages that were not actually saved out, but were added to the AR - the DLC may now have those packages included,
 						// and there will be a conflict later on if the package is in this list and the DLC list. PackageFlags of 0 means it was 
 						// evaluated and skipped.
-						TArrayView<FAssetData const* const> AssetsForPackage = ReleaseAssetRegistry.GetAssetsByPackageName(PackageName);
-						checkf(AssetsForPackage.Num() > 0, TEXT("It is unexpected that no assets were found in DevelopmentAssetRegistry for the package %s. This indicates an invalid AR."), *PackageName.ToString());
-						// just check the first one in the list, they will all have the same flags
-						if (AssetsForPackage[0]->PackageFlags != 0)
-						{
-							Arguments.ReleasedPackages.PackageNames.Add(PackageName);
-							Arguments.ReleasedPackages.PackageIdToName.Add(FPackageId::FromName(PackageName), PackageName);
-						}
+						bool bHasAny = false;
+						ReleaseAssetRegistry.EnumerateAssetsByPackageName(PackageName,
+							[&bHasAny, &Arguments, PackageName](const FAssetData* AssetData)
+							{
+								// just check the first one in the list, they will all have the same flags
+								if (AssetData->PackageFlags != 0)
+								{
+									Arguments.ReleasedPackages.PackageNames.Add(PackageName);
+									Arguments.ReleasedPackages.PackageIdToName.Add(FPackageId::FromName(PackageName), PackageName);
+								}
+								bHasAny = true;
+								return false; // stop iterating
+							});
+						checkf(bHasAny,
+							TEXT("It is unexpected that no assets were found in DevelopmentAssetRegistry for the package %s. This indicates an invalid AR."),
+							*PackageName.ToString());
 					}
 				}
 			}
diff --git a/Engine/Source/Editor/UnrealEd/Private/Commandlets/AssetRegistryGenerator.cpp b/Engine/Source/Editor/UnrealEd/Private/Commandlets/AssetRegistryGenerator.cpp
--- a/Engine/Source/Editor/UnrealEd/Private/Commandlets/AssetRegistryGenerator.cpp
+++ b/Engine/Source/Editor/UnrealEd/Private/Commandlets/AssetRegistryGenerator.cpp
@@ -883,16 +883,13 @@ void FAssetRegistryGenerator::SetPreviousAssetRegistry(TUniquePtr<FAssetRegistry
 			FIterativelySkippedPackageUpdateData& UpdateData = PreviousPackagesToUpdate.FindOrAdd(PackageName);
 			bool bGenerated = false;
 
-			TArrayView<FAssetData const * const> PreviousAssetDatas = InPreviousState->GetAssetsByPackageName(PackageName);
-			if (PreviousAssetDatas.Num() > 0)
-			{
-				UpdateData.AssetDatas.Reserve(PreviousAssetDatas.Num());
-				for (const FAssetData* AssetData : PreviousAssetDatas)
+			UpdateData.AssetDatas.Reserve(InPreviousState->NumAssetsByPackageName(PackageName));
+			InPreviousState->EnumerateAssetsByPackageName(PackageName, [&bGenerated, &UpdateData](const FAssetData* AssetData)
 				{
 					bGenerated |= (AssetData->PackageFlags & PKG_CookGenerated) != 0;
 					UpdateData.AssetDatas.Emplace(*AssetData);
-				}
-			}
+					return true; // Keep iterating
+				});
 			UpdateData.PackageData = *Pair.Value;
 
 			// Keep the dependencies and referencers of generated packages
@@ -921,11 +918,8 @@ void FAssetRegistryGenerator::InjectEncryptionData(FAssetRegistryState& TargetSt
 
 		for (FName EncryptedRootPackageName : EncryptedRootAssets)
 		{
-			for (const FAssetData* PackageAsset : TargetState.GetAssetsByPackageName(EncryptedRootPackageName))
-			{
-				FAssetData* AssetData = const_cast<FAssetData*>(PackageAsset);
-
-				if (AssetData)
+			TargetState.EnumerateMutableAssetsByPackageName(EncryptedRootPackageName,
+				[&GuidCache, &AssetManager, &TargetState](FAssetData* AssetData)
 				{
 					FString GuidString;
 					const FAssetData::FChunkArrayView ChunkIDs = AssetData->GetChunkIDs();
@@ -956,8 +950,8 @@ void FAssetRegistryGenerator::InjectEncryptionData(FAssetRegistryState& TargetSt
 							TargetState.UpdateAssetData(AssetData, MoveTemp(NewAssetData));
 						}
 					}
-				}
-			}
+					return true; // Keep iterating assets in the package
+				});
 		}
 	}
 }
@@ -1104,7 +1098,7 @@ void FAssetRegistryGenerator::UpdateCollectionAssetData()
 		const FSoftObjectPath& AssetPath = AssetPathToCollectionTagsPair.Key;
 		const TArray<FName>& CollectionTagsForAsset = AssetPathToCollectionTagsPair.Value;
 
-		const FAssetData* AssetData = State.GetAssetByObjectPath(AssetPath);
+		FAssetData* AssetData = State.GetMutableAssetByObjectPath(AssetPath);
 		if (AssetData)
 		{
 			FAssetDataTagMap TagsAndValues = AssetData->TagsAndValues.CopyMap();
@@ -1114,7 +1108,7 @@ void FAssetRegistryGenerator::UpdateCollectionAssetData()
 			}
 			FAssetData NewAssetData(*AssetData);
 			NewAssetData.TagsAndValues = FAssetDataTagMapSharedView(MoveTemp(TagsAndValues));
-			State.UpdateAssetData(const_cast<FAssetData*>(AssetData), MoveTemp(NewAssetData));
+			State.UpdateAssetData(AssetData, MoveTemp(NewAssetData));
 		}
 	}
 }
@@ -1457,8 +1451,12 @@ FAssetPackageData FAssetRegistryGenerator::CopyAssetPackageDataForIncrementalCoo
 
 FName FAssetRegistryGenerator::GetGeneratorPackage(FName PackageName, const FAssetRegistryState& InState)
 {
-	TConstArrayView<const FAssetData*> Assets = InState.GetAssetsByPackageName(PackageName);
-	bool bGenerated = !Assets.IsEmpty() && (Assets[0]->PackageFlags & PKG_CookGenerated) != 0;
+	bool bGenerated = false;
+	InState.EnumerateAssetsByPackageName(PackageName, [&bGenerated](const FAssetData* Asset)
+		{
+			bGenerated = (Asset->PackageFlags & PKG_CookGenerated) != 0;
+			return false; // Stop iterating
+		});
 	if (!bGenerated)
 	{
 		return NAME_None;
@@ -1486,8 +1484,12 @@ void FAssetRegistryGenerator::ComputePackageRemovals(const FAssetRegistryState&
 		{
 			// If it's a generated package, never mark it as removed (that can only be handled by the generator)
 			// Mark it as modified if its generator or any of its dependencies are modified.
-			TConstArrayView<const FAssetData*> PreviousAssets = PreviousState.GetAssetsByPackageName(PackageName);
-			bool bGenerated = !PreviousAssets.IsEmpty() && (PreviousAssets[0]->PackageFlags & PKG_CookGenerated);
+			bool bGenerated = false;
+			PreviousState.EnumerateAssetsByPackageName(PackageName, [&bGenerated](const FAssetData* AssetData)
+				{
+					bGenerated = (AssetData->PackageFlags & PKG_CookGenerated) != 0;
+					return false; // stop iterating
+				});
 			if (bGenerated)
 			{
 				TArray<FAssetIdentifier> Referencers;
@@ -1587,7 +1589,7 @@ void FAssetRegistryGenerator::FinalizeChunkIDs(const TSet<FName>& InCookedPackag
 	}
 
 	// Copy ExplicitChunkIDs and other data from the AssetRegistry into the maps we use during finalization
-	State.EnumerateAllAssets([&](const FAssetData& AssetData)
+	State.EnumerateAllMutableAssets([&](FAssetData& AssetData)
 	{
 		for (int32 ChunkID : AssetData.GetChunkIDs())
 		{
@@ -1602,7 +1604,7 @@ void FAssetRegistryGenerator::FinalizeChunkIDs(const TSet<FName>& InCookedPackag
 
 		// Clear the Asset's chunk id list. We will fill it with the final IDs to use later on.
 		// Chunk Ids are safe to modify in place so do a const cast
-		const_cast<FAssetData&>(AssetData).ClearChunkIDs();
+		AssetData.ClearChunkIDs();
 
 		// Update whether the owner package contains a map
 		if ((AssetData.PackageFlags & PKG_ContainsMap) != 0)
@@ -2429,11 +2431,12 @@ void FAssetRegistryGenerator::FixupPackageDependenciesForChunks(UE::Cook::FCookS
 		check(FinalChunkManifests[PakchunkIndex]);
 		for (const TPair<FName, FString>& Asset : *FinalChunkManifests[PakchunkIndex])
 		{
-			for (const FAssetData* AssetData : State.GetAssetsByPackageName(Asset.Key))
-			{
-				// Chunk Ids are safe to modify in place
-				const_cast<FAssetData*>(AssetData)->AddChunkID(PakchunkIndex);
-			}
+			State.EnumerateMutableAssetsByPackageName(Asset.Key, [PakchunkIndex](FAssetData* AssetData)
+				{
+					// Chunk Ids are safe to modify in place
+					AssetData->AddChunkID(PakchunkIndex);
+					return true;
+				});
 		}
 	}
 }
diff --git a/Engine/Source/Editor/UnrealEd/Private/Commandlets/DiffAssetBulkDataCommandlet.cpp b/Engine/Source/Editor/UnrealEd/Private/Commandlets/DiffAssetBulkDataCommandlet.cpp
--- a/Engine/Source/Editor/UnrealEd/Private/Commandlets/DiffAssetBulkDataCommandlet.cpp
+++ b/Engine/Source/Editor/UnrealEd/Private/Commandlets/DiffAssetBulkDataCommandlet.cpp
@@ -178,7 +178,9 @@ int32 UDiffAssetBulkDataCommandlet::Main(const FString& FullCommandLine)
 		{
 			const FAssetPackageData* Current = CurrentState.GetAssetPackageData(NamePackageDataPair.Key);
 
-			const FAssetData* BaseMIAsset = UE::AssetRegistry::GetMostImportantAsset(BaseState.GetAssetsByPackageName(NamePackageDataPair.Key), UE::AssetRegistry::EGetMostImportantAssetFlags::IgnoreSkipClasses);
+			const FAssetData* BaseMIAsset = UE::AssetRegistry::GetMostImportantAsset(
+				BaseState.CopyAssetsByPackageName(NamePackageDataPair.Key),
+				UE::AssetRegistry::EGetMostImportantAssetFlags::IgnoreSkipClasses);
 			uint64 BaseCompressedSize = 0;
 			if (BaseMIAsset && BaseMIAsset->GetTagValue(UE::AssetRegistry::Stage_ChunkCompressedSizeFName, BaseCompressedSize))
 			{
@@ -201,7 +203,9 @@ int32 UDiffAssetBulkDataCommandlet::Main(const FString& FullCommandLine)
 		{
 			const FAssetPackageData* Base = BaseState.GetAssetPackageData(NamePackageDataPair.Key);
 
-			const FAssetData* CurrentMIAsset = UE::AssetRegistry::GetMostImportantAsset(CurrentState.GetAssetsByPackageName(NamePackageDataPair.Key), UE::AssetRegistry::EGetMostImportantAssetFlags::IgnoreSkipClasses);
+			const FAssetData* CurrentMIAsset = UE::AssetRegistry::GetMostImportantAsset(
+				CurrentState.CopyAssetsByPackageName(NamePackageDataPair.Key),
+				UE::AssetRegistry::EGetMostImportantAssetFlags::IgnoreSkipClasses);
 			uint64 CurrentCompressedSize = 0;
 			if (CurrentMIAsset && CurrentMIAsset->GetTagValue(UE::AssetRegistry::Stage_ChunkCompressedSizeFName, CurrentCompressedSize))
 			{
@@ -316,8 +320,8 @@ int32 UDiffAssetBulkDataCommandlet::Main(const FString& FullCommandLine)
 	TMap < FName /* PackageName */, TPair<uint64, uint64>> PackageSizes;
 	for (const FName& ChangedPackageName : ChangedPackages)
 	{
-		TConstArrayView<FAssetData const*> BaseAssetDatas = BaseState.GetAssetsByPackageName(ChangedPackageName);
-		TConstArrayView<FAssetData const*> CurrentAssetDatas = CurrentState.GetAssetsByPackageName(ChangedPackageName);
+		TArray<FAssetData const*> BaseAssetDatas = BaseState.CopyAssetsByPackageName(ChangedPackageName);
+		TArray<FAssetData const*> CurrentAssetDatas = CurrentState.CopyAssetsByPackageName(ChangedPackageName);
 
 		struct FDiffTag
 		{
diff --git a/Engine/Source/Editor/UnrealEd/Private/Commandlets/DiffAssetRegistriesCommandlet.cpp b/Engine/Source/Editor/UnrealEd/Private/Commandlets/DiffAssetRegistriesCommandlet.cpp
--- a/Engine/Source/Editor/UnrealEd/Private/Commandlets/DiffAssetRegistriesCommandlet.cpp
+++ b/Engine/Source/Editor/UnrealEd/Private/Commandlets/DiffAssetRegistriesCommandlet.cpp
@@ -512,7 +512,7 @@ bool	UDiffAssetRegistriesCommandlet::IsInRelevantChunk(FAssetRegistryState& InRe
 		return true;
 
 	}
-	TArrayView<FAssetData const* const> Assets = InRegistryState.GetAssetsByPackageName(InAssetPath);
+	TArray<const FAssetData*> Assets = InRegistryState.CopyAssetsByPackageName(InAssetPath);
 	if (!Assets.IsEmpty())
 	{
 		const FAssetData::FChunkArrayView ChunkIDs = Assets[0]->GetChunkIDs();
@@ -529,7 +529,7 @@ FName UDiffAssetRegistriesCommandlet::GetClassName(FAssetRegistryState& InRegist
 {
 	if (AssetPathToClassName.Contains(InAssetPath) == false)
 	{
-		TArrayView<FAssetData const * const> Assets = InRegistryState.GetAssetsByPackageName(InAssetPath);
+		TArray<const FAssetData*> Assets = InRegistryState.CopyAssetsByPackageName(InAssetPath);
 
 		FName NewName;
 		if (Assets.Num() > 0)
@@ -558,7 +558,7 @@ TArray<int32> UDiffAssetRegistriesCommandlet::GetAssetChunks(FAssetRegistryState
 {
 	if (ChunkIdByAssetPath.Contains(InAssetPath) == false)
 	{
-		TArrayView<FAssetData const* const> Assets = InRegistryState.GetAssetsByPackageName(InAssetPath);
+		TArray<const FAssetData* > Assets = InRegistryState.CopyAssetsByPackageName(InAssetPath);
 		const FAssetData::FChunkArrayView ChunkIDs = Assets.IsEmpty() ? FAssetData::FChunkArrayView() : Assets[0]->GetChunkIDs();
 		if (!ChunkIDs.IsEmpty())
 		{
@@ -1147,7 +1147,7 @@ void UDiffAssetRegistriesCommandlet::DiffAssetRegistries(const FString& OldPath,
 				newuncooked++;
 			}
 			
-			newassets += NewState.GetAssetsByPackageName(Name).Num();
+			newassets += NewState.NumAssetsByPackageName(Name);
 			
 			if (!PrevData)
 			{
@@ -1217,7 +1217,7 @@ void UDiffAssetRegistriesCommandlet::DiffAssetRegistries(const FString& OldPath,
 				newuncooked++;
 			}
 			
-			newassets += NewState.GetAssetsByPackageName(Name).Num();
+			newassets += NewState.NumAssetsByPackageName(Name);
 			
 			if (!PrevData)
 			{
@@ -1265,7 +1265,7 @@ void UDiffAssetRegistriesCommandlet::DiffAssetRegistries(const FString& OldPath,
 			olduncooked++;
 		}
 
-		oldassets += OldState.GetAssetsByPackageName(Name).Num();
+		oldassets += OldState.NumAssetsByPackageName(Name);
 
 		if (!Data)
 		{
diff --git a/Engine/Source/Editor/UnrealEd/Private/CookOnTheFlyServer.cpp b/Engine/Source/Editor/UnrealEd/Private/CookOnTheFlyServer.cpp
--- a/Engine/Source/Editor/UnrealEd/Private/CookOnTheFlyServer.cpp
+++ b/Engine/Source/Editor/UnrealEd/Private/CookOnTheFlyServer.cpp
@@ -12829,16 +12829,24 @@ bool UCookOnTheFlyServer::GetAllPackageFilenamesFromAssetRegistry(const FString&
 			TSet<FName> Names;
 		} UniquePackageNames[UNIQUEPACKAGENAMES_BUCKETS];
 
-		const int32 NumPackages = SerializedState->GetAssetDataMap().Num();
-		OutPackageDatas.SetNum(NumPackages);
+		const int32 NumAssets = SerializedState->GetNumAssets();
+		TArray<const FAssetData*> StateAssets;
+		StateAssets.Reserve(NumAssets);
+		SerializedState->EnumerateAllAssets([&StateAssets](const FAssetData& AssetData)
+			{
+				StateAssets.Add(&AssetData);
+			});
 
-		// Convert the Map of RegistryData into an Array of FAssetData and populate PackageNames in the output array
-		// We can index directly from the set because we know its congiguous as we just deserialized it.
-		checkf(SerializedState->GetAssetDataMap().GetMaxIndex() == NumPackages, TEXT("The set needs to be contiguous so we can index into it directly"));
-		ParallelFor(NumPackages,
+		// We set the output packages size to the number of assets, even though the number of packages will be less than
+		// the number of assets. We check for duplicates in a critical section inside the parallel for and skip the duplicate
+		// work. We remove the entries for the skipped duplicates after the parallel for.
+		// We are iterating over assets instead of packages because it is faster in the parallelfor to do the flat iteration over assets.
+		OutPackageDatas.SetNum(NumAssets);
+		// populate PackageNames in the output array
+		ParallelFor(NumAssets,
 			[&](int32 Index)
 			{
-				const FAssetData& RegistryData = *SerializedState->GetAssetDataMap()[FSetElementId::FromInteger(Index)];
+				const FAssetData& RegistryData = *StateAssets[Index];
 
 				// If we want to reevaluate (try cooking again) the uncooked packages (packages that were found to be empty when we cooked them before),
 				// then remove the uncooked packages from the set of known packages. Uncooked packages are identified by PackageFlags == 0.
diff --git a/Engine/Source/Editor/UnrealEd/Private/Cooker/LooseCookedPackageWriter.cpp b/Engine/Source/Editor/UnrealEd/Private/Cooker/LooseCookedPackageWriter.cpp
--- a/Engine/Source/Editor/UnrealEd/Private/Cooker/LooseCookedPackageWriter.cpp
+++ b/Engine/Source/Editor/UnrealEd/Private/Cooker/LooseCookedPackageWriter.cpp
@@ -769,10 +769,11 @@ TUniquePtr<FAssetRegistryState> FLooseCookedPackageWriter::LoadPreviousAssetRegi
 				bool bIsCookedOnly = bIsScriptPackage;
 				if (!bIsCookedOnly)
 				{
-					for (const FAssetData* AssetData : PreviousState->GetAssetsByPackageName(PackageName))
-					{
-						bIsCookedOnly |= !!(AssetData->PackageFlags & PKG_CookGenerated);
-					}
+					PreviousState->EnumerateAssetsByPackageName(PackageName, [&bIsCookedOnly](const FAssetData* AssetData)
+						{
+							bIsCookedOnly |= !!(AssetData->PackageFlags & PKG_CookGenerated);
+							return true; // Keep iterating
+						});
 				}
 				bNoLongerExistsInEditor = !bIsCookedOnly;
 			}
diff --git a/Engine/Source/Runtime/AssetRegistry/Private/AssetRegistry.cpp b/Engine/Source/Runtime/AssetRegistry/Private/AssetRegistry.cpp
--- a/Engine/Source/Runtime/AssetRegistry/Private/AssetRegistry.cpp
+++ b/Engine/Source/Runtime/AssetRegistry/Private/AssetRegistry.cpp
@@ -1599,7 +1599,7 @@ void FAssetRegistryImpl::CollectCodeGeneratorClasses()
 		// when AddAssetData is called for those instances, but when we add a new generator class we have to recheck all
 		// instances of the class since they would have failed to detect they were Blueprint classes before.
 		// This can happen if blueprints in plugin B are scanned before their blueprint class from plugin A is scanned.
-		for (const FAssetData* AssetData : State.GetAssetsByClassPathName(BPCoreClassName))
+		State.EnumerateAssetsByClassPathName(BPCoreClassName, [this](const FAssetData* AssetData)
 		{
 			const FString GeneratedClass = AssetData->GetTagValueRef<FString>(FBlueprintTags::GeneratedClassPath);
 			const FString ParentClass = AssetData->GetTagValueRef<FString>(FBlueprintTags::ParentClassPath);
@@ -1616,7 +1616,8 @@ void FAssetRegistryImpl::CollectCodeGeneratorClasses()
 					TempCachedInheritanceBuffer.bDirty = true;
 				}
 			}
-		}
+			return true; // Keep iterating the assets for the class
+		});
 	}
 }
 
@@ -4371,7 +4372,7 @@ void FAssetRegistryImpl::AssetsSaved(UE::AssetRegistry::Impl::FEventContext& Eve
 	for (FAssetData& NewAssetData : Assets)
 	{
 		FCachedAssetKey Key(NewAssetData);
-		FAssetData** DataFromGather = State.CachedAssets.Find(Key);
+		FAssetData* DataFromGather = State.GetMutableAssetByObjectPath(Key);
 
 		AssetDataObjectPathsUpdatedOnLoad.Add(NewAssetData.GetSoftObjectPath());
 
@@ -4382,7 +4383,7 @@ void FAssetRegistryImpl::AssetsSaved(UE::AssetRegistry::Impl::FEventContext& Eve
 		}
 		else
 		{
-			UpdateAssetData(EventContext, *DataFromGather, MoveTemp(NewAssetData), false /* bKeepDeletedTags */);
+			UpdateAssetData(EventContext, DataFromGather, MoveTemp(NewAssetData), false /* bKeepDeletedTags */);
 		}
 	}
 }
@@ -5308,7 +5309,8 @@ void FAssetRegistryImpl::LoadCalculatedDependencies(FName PackageName,
 	FReadScopeLock GathererClassScopeLock(RegisteredDependencyGathererClassesLock);
 
 	TArray<UE::AssetDependencyGatherer::Private::FRegisteredAssetDependencyGatherer*, TInlineAllocator<2>> Gatherers;
-	for (const FAssetData* AssetData : State.GetAssetsByPackageName(PackageName))
+	State.EnumerateAssetsByPackageName(PackageName,
+		[this, &Gatherers, &GetCompiledFilter, &bOutHadActivity, PackageName](const FAssetData* AssetData)
 	{
 		Gatherers.Reset();
 
@@ -5346,7 +5348,9 @@ void FAssetRegistryImpl::LoadCalculatedDependencies(FName PackageName,
 				AddDirectoryReferencer(PackageName, Directory);
 			}
 		}
-	}
+
+		return true; // Keep iterating the assets in the package
+	});
 }
 
 void FAssetRegistryImpl::AddDirectoryReferencer(FName PackageName, const FString& DirectoryLocalPathOrLongPackageName)
@@ -5466,28 +5470,25 @@ void FAssetRegistryImpl::CachePathsFromState(Impl::FEventContext& EventContext,
 	CollectCodeGeneratorClasses();
 
 	// Add paths to cache
-	for (const FAssetData* AssetData : InState.CachedAssets)
+	InState.EnumerateAllAssets([this, &EventContext](const FAssetData& AssetData)
 	{
-		if (AssetData != nullptr)
+		AddAssetPath(EventContext, AssetData.PackagePath);
+
+		// Populate the class map if adding blueprint
+		if (ClassGeneratorNames.Contains(AssetData.AssetClassPath))
 		{
-			AddAssetPath(EventContext, AssetData->PackagePath);
+			FAssetRegistryExportPath GeneratedClass = AssetData.GetTagValueRef<FAssetRegistryExportPath>(FBlueprintTags::GeneratedClassPath);
+			FAssetRegistryExportPath ParentClass = AssetData.GetTagValueRef<FAssetRegistryExportPath>(FBlueprintTags::ParentClassPath);
 
-			// Populate the class map if adding blueprint
-			if (ClassGeneratorNames.Contains(AssetData->AssetClassPath))
+			if (GeneratedClass && ParentClass)
 			{
-				FAssetRegistryExportPath GeneratedClass = AssetData->GetTagValueRef<FAssetRegistryExportPath>(FBlueprintTags::GeneratedClassPath);
-				FAssetRegistryExportPath ParentClass = AssetData->GetTagValueRef<FAssetRegistryExportPath>(FBlueprintTags::ParentClassPath);
+				AddCachedBPClassParent(GeneratedClass.ToTopLevelAssetPath(), ParentClass.ToTopLevelAssetPath());
 
-				if (GeneratedClass && ParentClass)
-				{
-					AddCachedBPClassParent(GeneratedClass.ToTopLevelAssetPath(), ParentClass.ToTopLevelAssetPath());
-
-					// Invalidate caching because CachedBPInheritanceMap got modified
-					TempCachedInheritanceBuffer.bDirty = true;
-				}
+				// Invalidate caching because CachedBPInheritanceMap got modified
+				TempCachedInheritanceBuffer.bDirty = true;
 			}
 		}
-	}
+	});
 }
 
 }
@@ -5857,10 +5858,12 @@ void FAssetRegistryImpl::ScanPathsSynchronous(Impl::FScanPathContext& Context)
 		}
 		for (const FString& PackageName : Context.PackageFiles)
 		{
-			for (const FAssetData* AssetData : State.GetAssetsByPackageName(FName(*PackageName)))
-			{
-				OldAssetsToRemove.Add(AssetData->ToSoftObjectPath());
-			}
+			State.EnumerateAssetsByPackageName(FName(*PackageName),
+				[&OldAssetsToRemove](const FAssetData* AssetData)
+				{
+					OldAssetsToRemove.Add(AssetData->ToSoftObjectPath());
+					return true;
+				});
 			for (const TCHAR* Extension : FAssetDataGatherer::GetVerseFileExtensions())
 			{
 				FName VerseName(*WriteToString<256>(PackageName, Extension), FNAME_Find);
@@ -6026,7 +6029,7 @@ void FAssetRegistryImpl::ScanPathsSynchronous(Impl::FScanPathContext& Context)
 #endif
 	for (FSoftObjectPath& OldAssetToRemove : OldAssetsToRemove)
 	{
-		FAssetData* AssetDataToRemove = const_cast<FAssetData*>(State.GetAssetByObjectPath(OldAssetToRemove));
+		FAssetData* AssetDataToRemove = State.GetMutableAssetByObjectPath(OldAssetToRemove);
 		if (AssetDataToRemove)
 		{
 			RemoveAssetData(Context.EventContext, AssetDataToRemove);
@@ -6426,8 +6429,7 @@ void FAssetRegistryImpl::AssetSearchDataGathered(Impl::FEventContext& EventConte
 
 		// Look for an existing asset to check whether we need to add or update
 		FCachedAssetKey Key(*BackgroundResult);
-		FAssetData* const* FoundData = State.CachedAssets.Find(Key);
-		FAssetData* ExistingAssetData = FoundData ? *FoundData : nullptr;
+		FAssetData* ExistingAssetData = State.GetMutableAssetByObjectPath(Key);
 		// The background result should not already be registered; it should be impossible since it is in TUnqiuePtr
 		check(ExistingAssetData == nullptr || ExistingAssetData != BackgroundResult.Get());
 
@@ -6985,8 +6987,15 @@ bool FAssetRegistryImpl::RemoveAssetData(Impl::FEventContext& EventContext, FAss
 
 void FAssetRegistryImpl::RemovePackageData(Impl::FEventContext& EventContext, const FName PackageName)
 {
-	TArray<FAssetData*, TInlineAllocator<1>>* PackageAssetsPtr = State.CachedAssetsByPackageName.Find(PackageName);
-	if (PackageAssetsPtr && PackageAssetsPtr->Num() > 0)
+	// Even if we could point to the array, we have to copy the array since RemoveAssetData may re-allocate it.
+	TArray<FAssetData*, TInlineAllocator<1>> PackageAssets;
+	State.EnumerateMutableAssetsByPackageName(PackageName, [&PackageAssets](FAssetData* AssetData)
+		{
+			PackageAssets.Add(AssetData);
+			return true;
+		});
+
+	if (PackageAssets.Num() > 0)
 	{
 		FAssetIdentifier PackageAssetIdentifier(PackageName);
 		// If there were any EDependencyCategory::Package referencers, re-add them to a new empty dependency node, as it would be when the referencers are loaded from disk
@@ -7001,8 +7010,6 @@ void FAssetRegistryImpl::RemovePackageData(Impl::FEventContext& EventContext, co
 			}
 		}
 
-		// Copy the array since RemoveAssetData may re-allocate it!
-		TArray<FAssetData*, TInlineAllocator<1>> PackageAssets = *PackageAssetsPtr;
 		for (FAssetData* PackageAsset : PackageAssets)
 		{
 			RemoveAssetData(EventContext, PackageAsset);
@@ -7584,7 +7591,7 @@ void FAssetRegistryImpl::PushProcessLoadedAssetsBatch(Impl::FEventContext& Event
 			continue;
 		}
 		FCachedAssetKey Key(NewAssetData);
-		FAssetData** DataFromGather = State.CachedAssets.Find(Key);
+		FAssetData* DataFromGather = State.GetMutableAssetByObjectPath(Key);
 
 		AssetDataObjectPathsUpdatedOnLoad.Add(NewAssetData.GetSoftObjectPath());
 
@@ -7603,7 +7610,7 @@ void FAssetRegistryImpl::PushProcessLoadedAssetsBatch(Impl::FEventContext& Event
 			// Modified tag values on the other hand do overwrite the old values from disk.
 			// This means that the only way to delete no-longer present tags from an AssetData
 			// is to resave the package, or to manually call AssetUpdateTags(EAssetRegistryTagsCaller::FullUpdate) from c++.
-			UpdateAssetData(EventContext, *DataFromGather, MoveTemp(NewAssetData), true /* bKeepDeletedTags */);
+			UpdateAssetData(EventContext, DataFromGather, MoveTemp(NewAssetData), true /* bKeepDeletedTags */);
 		}
 	}
 
@@ -7617,19 +7624,19 @@ void FAssetRegistryImpl::PushProcessLoadedAssetsBatch(Impl::FEventContext& Event
 
 void FAssetRegistryImpl::UpdateRedirectCollector()
 {
-	// Look for all redirectors in list
-	const TArray<const FAssetData*>& RedirectorAssets = State.GetAssetsByClassPathName(UE::AssetRegistry::GetClassPathObjectRedirector());
-
-	for (const FAssetData* AssetData : RedirectorAssets)
-	{
-		FSoftObjectPath Source = AssetData->GetSoftObjectPath();
-		FSoftObjectPath Destination = GetRedirectedObjectPath(Source, nullptr, nullptr, /*bNeedsScanning*/ false);
-
-		if (Destination != Source)
+	// Look for all redirectors in the AssetRegistry
+	State.EnumerateAssetsByClassPathName(UE::AssetRegistry::GetClassPathObjectRedirector(),
+		[this](const FAssetData* AssetData)
 		{
-			GRedirectCollector.AddAssetPathRedirection(Source, Destination);
-		}
-	}
+			FSoftObjectPath Source = AssetData->GetSoftObjectPath();
+			FSoftObjectPath Destination = GetRedirectedObjectPath(Source, nullptr, nullptr, /*bNeedsScanning*/ false);
+
+			if (Destination != Source)
+			{
+				GRedirectCollector.AddAssetPathRedirection(Source, Destination);
+			}
+			return true; // Keep iterating
+		});
 }
 
 }
@@ -7687,11 +7694,16 @@ void FAssetRegistryImpl::ScanModifiedAssetFiles(Impl::FEventContext& EventContex
 		ExistingAssetDatas.Reserve(InFilePaths.Num());
 		for (const FString& PackageName : ModifiedPackageNames)
 		{
-			TArray<FAssetData*, TInlineAllocator<1>>* PackageAssetsPtr = State.CachedAssetsByPackageName.Find(*PackageName);
-			if (PackageAssetsPtr && PackageAssetsPtr->Num() > 0)
+			TArray<const FAssetData*, TInlineAllocator<1>> PackageAssets;
+			State.EnumerateAssetsByPackageName(*PackageName, [&PackageAssets](const FAssetData* AssetData)
+				{
+					PackageAssets.Add(AssetData);
+					return true;
+				});
+			if (PackageAssets.Num() > 0)
 			{
-				ExistingAssetDatas.Reserve(ExistingAssetDatas.Num() + PackageAssetsPtr->Num());
-				for (FAssetData* AssetData : *PackageAssetsPtr)
+				ExistingAssetDatas.Reserve(ExistingAssetDatas.Num() + PackageAssets.Num());
+				for (const FAssetData* AssetData : PackageAssets)
 				{
 					ExistingAssetDatas.Add(AssetData->ToSoftObjectPath());
 				}
@@ -7712,7 +7724,7 @@ void FAssetRegistryImpl::ScanModifiedAssetFiles(Impl::FEventContext& EventContex
 		{
 			if (!FoundAssets.Contains(OldAssetPath))
 			{
-				FAssetData* OldAssetData = const_cast<FAssetData*>(State.GetAssetByObjectPath(OldAssetPath));
+				FAssetData* OldAssetData = State.GetMutableAssetByObjectPath(OldAssetPath);
 				if (OldAssetData)
 				{
 					RemoveAssetData(EventContext, OldAssetData);
@@ -7723,10 +7735,10 @@ void FAssetRegistryImpl::ScanModifiedAssetFiles(Impl::FEventContext& EventContex
 		// Send ModifiedOnDisk event for every Asset that was modified
 		for (const FSoftObjectPath& FoundAsset : FoundAssets)
 		{
-			FAssetData** AssetData = State.CachedAssets.Find(FCachedAssetKey(FoundAsset));
+			const FAssetData* AssetData = State.GetAssetByObjectPath(FCachedAssetKey(FoundAsset));
 			if (AssetData)
 			{
-				EventContext.AssetEvents.Emplace(**AssetData, Impl::FEventContext::EEvent::UpdatedOnDisk);
+				EventContext.AssetEvents.Emplace(*AssetData, Impl::FEventContext::EEvent::UpdatedOnDisk);
 			}
 		}
 	}
@@ -7905,11 +7917,11 @@ void FAssetRegistryImpl::OnContentPathDismounted(Impl::FEventContext& EventConte
 		for (FName PathName : PathList)
 		{
 			// Gather assets
-			TArray<FAssetData*>* AssetsInPath = State.CachedAssetsByPath.Find(PathName);
-			if (AssetsInPath)
-			{
-				AllAssetDataToRemove.Append(*AssetsInPath);
-			}
+			State.EnumerateMutableAssetsByPackagePath(PathName, [&AllAssetDataToRemove](FAssetData* AssetData)
+				{
+					AllAssetDataToRemove.Add(AssetData);
+					return true;
+				});
 
 			// Forget Verse files
 			const TArray<FName>* VerseFilesInPath = CachedVerseFilesByPath.Find(PathName);
@@ -8595,15 +8607,13 @@ namespace UE::AssetRegistry
 
 bool FAssetRegistryImpl::SetPrimaryAssetIdForObjectPath(Impl::FEventContext& EventContext, const FSoftObjectPath& ObjectPath, FPrimaryAssetId PrimaryAssetId)
 {
-	FAssetData** FoundAssetData = State.CachedAssets.Find(FCachedAssetKey(ObjectPath));
+	FAssetData* AssetData = State.GetMutableAssetByObjectPath(ObjectPath);
 
-	if (!FoundAssetData)
+	if (!AssetData)
 	{
 		return false;
 	}
 
-	FAssetData* AssetData = *FoundAssetData;
-
 	FAssetDataTagMap TagsAndValues = AssetData->TagsAndValues.CopyMap();
 	TagsAndValues.Add(FPrimaryAssetId::PrimaryAssetTypeTag, PrimaryAssetId.PrimaryAssetType.ToString());
 	TagsAndValues.Add(FPrimaryAssetId::PrimaryAssetNameTag, PrimaryAssetId.PrimaryAssetName.ToString());
@@ -8747,14 +8757,16 @@ void FAssetRegistryImpl::PruneAndCoalescePackagesRequiringDependencyCalculation(
 			for (auto Iter = SourceSet.CreateIterator(); Iter; ++Iter)
 			{
 				bool HasAnyRegisteredDependencyGatherers = false;
-				for (const FAssetData* AssetData : State.GetAssetsByPackageName(*Iter))
-				{
-					if (RegisteredDependencyGathererClasses.Contains(AssetData->AssetClassPath))
+				State.EnumerateAssetsByPackageName(*Iter,
+					[this, &HasAnyRegisteredDependencyGatherers](const FAssetData* AssetData)
 					{
-						HasAnyRegisteredDependencyGatherers = true;
-						break;
-					}
-				}
+						if (RegisteredDependencyGathererClasses.Contains(AssetData->AssetClassPath))
+						{
+							HasAnyRegisteredDependencyGatherers = true;
+							return false; // stop iterating
+						}
+						return true; // Keep iterating
+					});
 
 				// If we need to process this asset and we have a destination set, move it there
 				if ((OptDestinationSet != nullptr) && HasAnyRegisteredDependencyGatherers)
diff --git a/Engine/Source/Runtime/AssetRegistry/Private/AssetRegistryState.cpp b/Engine/Source/Runtime/AssetRegistry/Private/AssetRegistryState.cpp
--- a/Engine/Source/Runtime/AssetRegistry/Private/AssetRegistryState.cpp
+++ b/Engine/Source/Runtime/AssetRegistry/Private/AssetRegistryState.cpp
@@ -82,18 +82,15 @@ void FAssetRegistryState::Reset()
 	else
 	{
 		// Delete all assets in the cache
-		for (auto AssetDataIt = CachedAssets.CreateConstIterator(); AssetDataIt; ++AssetDataIt)
+		for (FAssetData* AssetData : CachedAssets)
 		{
-			if (*AssetDataIt)
-			{
-				delete *AssetDataIt;
-				NumAssets--;
-			}
+			delete AssetData;
+			NumAssets--;
 		}
 	}
 
 	// Make sure we have deleted all our allocated FAssetData objects
-	// Mar-06: Temporarily remove this ensure to allow passing builds  while we find and fix the cause
+	// March 2021: Temporarily remove this ensure to allow passing builds  while we find and fix the cause
 	// TODO: Restore the ensure
 	// ensure(NumAssets == 0);
 	UE_CLOG(NumAssets != 0, LogAssetRegistry, Display,
@@ -150,7 +147,6 @@ void FAssetRegistryState::Reset()
 	ensure(NumPackageData == 0);
 
 	// Clear cache
-	CachedAssets.Empty();
 	CachedAssetsByPackageName.Empty();
 	CachedAssetsByPath.Empty();
 	CachedAssetsByClass.Empty();
@@ -161,6 +157,7 @@ void FAssetRegistryState::Reset()
 #endif
 	CachedDependsNodes.Empty();
 	CachedPackageData.Empty();
+	CachedAssets.Empty();
 }
 
 void FAssetRegistryState::FilterTags(const FAssetDataTagMapSharedView& InTagsAndValues, FAssetDataTagMap& OutTagsAndValues,
@@ -220,29 +217,33 @@ void FAssetRegistryState::InitializeFromExistingAndPrune(const FAssetRegistrySta
 	TSet<FName> RequiredDependNodePackages;
 
 	// Duplicate asset data entries
-	for (FAssetData* AssetData : ExistingState.CachedAssets)
+	ExistingState.EnumerateAllMutableAssets(
+		[&RequiredPackages, &RemovePackages, &ChunksToKeep, &Options,
+		bIsFilteredByChunkId, bIsFilteredByRequiredPackages, bIsFilteredByRemovedPackages,
+		&RequiredDependNodePackages, this]
+		(FAssetData& AssetData)
 	{
 		bool bRemoveAssetData = false;
 		bool bRemoveDependencyData = true;
 
 		if (bIsFilteredByChunkId &&
-			!AssetData->GetChunkIDs().ContainsByPredicate([&](int32 ChunkId)
+			!AssetData.GetChunkIDs().ContainsByPredicate([&ChunksToKeep](int32 ChunkId)
 				{
 					return ChunksToKeep.Contains(ChunkId);
 				}))
 		{
 			bRemoveAssetData = true;
 		}
-		else if (bIsFilteredByRequiredPackages && !RequiredPackages.Contains(AssetData->PackageName))
+		else if (bIsFilteredByRequiredPackages && !RequiredPackages.Contains(AssetData.PackageName))
 		{
 			bRemoveAssetData = true;
 		}
-		else if (bIsFilteredByRemovedPackages && RemovePackages.Contains(AssetData->PackageName))
+		else if (bIsFilteredByRemovedPackages && RemovePackages.Contains(AssetData.PackageName))
 		{
 			bRemoveAssetData = true;
 		}
-		else if (Options.bFilterAssetDataWithNoTags && AssetData->TagsAndValues.Num() == 0 &&
-			!FPackageName::IsLocalizedPackage(WriteToString<256>(AssetData->PackageName)))
+		else if (Options.bFilterAssetDataWithNoTags && AssetData.TagsAndValues.Num() == 0 &&
+			!FPackageName::IsLocalizedPackage(WriteToString<256>(AssetData.PackageName)))
 		{
 			bRemoveAssetData = true;
 			bRemoveDependencyData = Options.bFilterDependenciesWithNoTags;
@@ -252,23 +253,23 @@ void FAssetRegistryState::InitializeFromExistingAndPrune(const FAssetRegistrySta
 		{
 			if (!bRemoveDependencyData)
 			{
-				RequiredDependNodePackages.Add(AssetData->PackageName);
+				RequiredDependNodePackages.Add(AssetData.PackageName);
 			}
-			continue;
+			return;
 		}
 
 		FAssetDataTagMap NewTagsAndValues;
-		FAssetRegistryState::FilterTags(AssetData->TagsAndValues, NewTagsAndValues,
-			Options.CookFilterlistTagsByClass.Find(AssetData->AssetClassPath), Options);
+		FAssetRegistryState::FilterTags(AssetData.TagsAndValues, NewTagsAndValues,
+			Options.CookFilterlistTagsByClass.Find(AssetData.AssetClassPath), Options);
 
-		FAssetData* NewAssetData = new FAssetData(AssetData->PackageName, AssetData->PackagePath, AssetData->AssetName,
-			AssetData->AssetClassPath, NewTagsAndValues, AssetData->GetChunkIDs(), AssetData->PackageFlags);
+		FAssetData* NewAssetData = new FAssetData(AssetData.PackageName, AssetData.PackagePath, AssetData.AssetName,
+			AssetData.AssetClassPath, NewTagsAndValues, AssetData.GetChunkIDs(), AssetData.PackageFlags);
 
-		NewAssetData->TaggedAssetBundles = AssetData->TaggedAssetBundles;
+		NewAssetData->TaggedAssetBundles = AssetData.TaggedAssetBundles;
 
 		// Add asset to new state
 		AddAssetData(NewAssetData);
-	}
+	});
 
 	// Create package data for all script and required packages
 	for (const TPair<FName, FAssetPackageData*>& Pair : ExistingState.CachedPackageData)
@@ -373,13 +374,12 @@ void FAssetRegistryState::InitializeFromExisting(const FAssetDataMap& AssetDataM
 			// don't do anything 
 			continue;
 		}
-		
 		const FAssetData& AssetData = *AssetDataPtr;
 
 		FAssetData* ExistingData = nullptr;
 		if (InInitializationMode != EInitializationMode::Rebuild) // minor optimization to avoid lookup in rebuild mode
 		{
-			if (FAssetData** Ptr = CachedAssets.Find(FCachedAssetKey(AssetData)))
+			if (FAssetData*const* Ptr = CachedAssets.Find(FCachedAssetKey(AssetData)))
 			{
 				ExistingData = *Ptr;
 			}
@@ -680,18 +680,20 @@ void FAssetRegistryState::Prune(const FAssetRegistryPruneOptions& PruneOptions)
 
 bool FAssetRegistryState::HasAssets(const FName PackagePath, bool bSkipARFilteredAssets) const
 {
-	const TArray<FAssetData*>* FoundAssetArray = CachedAssetsByPath.Find(PackagePath);
-	if (FoundAssetArray)
-	{
-		return FoundAssetArray->ContainsByPredicate([this, bSkipARFilteredAssets](FAssetData* AssetData)
+	bool bHasAssets = false;
+	EnumerateAssetsByPackagePath(PackagePath, [this, bSkipARFilteredAssets, &bHasAssets](const FAssetData* AssetData)
 		{
-			return AssetData && !IsPackageUnmountedAndFiltered(AssetData->PackageName)
+			if (AssetData && !IsPackageUnmountedAndFiltered(AssetData->PackageName)
 				&& (!bSkipARFilteredAssets
 					|| !UE::AssetRegistry::FFiltering::ShouldSkipAsset(
-						AssetData->AssetClassPath, AssetData->PackageFlags));
+						AssetData->AssetClassPath, AssetData->PackageFlags)))
+			{
+				bHasAssets = true;
+				return false; // Stop iterating
+			}
+			return true; // keep iterating
 		});
-	}
-	return false;
+	return bHasAssets;
 }
 
 bool FAssetRegistryState::GetAssets(const FARCompiledFilter& Filter, const TSet<FName>& PackageNamesToSkip,
@@ -740,7 +742,7 @@ bool DecideIntersectionMethod(int32 PreviousSize, int32 FilterResultsSize, int32
 }
 
 void ArrayIntersection(TArray<const FAssetData*>& InOutResults,
-	TConstArrayView<TArrayView<const FAssetData*>> Matches, int32 TotalMatches)
+	TConstArrayView<TConstArrayView<const FAssetData*>> Matches, int32 TotalMatches)
 {
 	if (InOutResults.Num() < TotalMatches)
 	{
@@ -751,7 +753,7 @@ void ArrayIntersection(TArray<const FAssetData*>& InOutResults,
 			Exists.Add(Result, false);
 		}
 		InOutResults.Empty();
-		for (TArrayView<const FAssetData*> Assets : Matches)
+		for (TConstArrayView<const FAssetData*> Assets : Matches)
 		{
 			for (const FAssetData* Asset : Assets)
 			{
@@ -771,7 +773,7 @@ void ArrayIntersection(TArray<const FAssetData*>& InOutResults,
 	{
 		TSet<const FAssetData*> Exists;
 		Exists.Reserve(TotalMatches);
-		for (TArrayView<const FAssetData*> Assets : Matches)
+		for (TConstArrayView<const FAssetData*> Assets : Matches)
 		{
 			for (const FAssetData* Asset : Assets)
 			{
@@ -790,7 +792,7 @@ template<class ArrayType, typename KeyType, typename CallbackType>
 void FilterAssets(TArray<const FAssetData*>& InOutResults, const TMap<KeyType, ArrayType>& AccelerationMap,
 	const TSet<KeyType>& Keys, CallbackType&& FunctionToKeepAsset, int32 FilterComplexity)
 {
-	TArray<TArrayView<const FAssetData*>, TInlineAllocator<10>> Matches;
+	TArray<TConstArrayView<const FAssetData*>, TInlineAllocator<10>> Matches;
 	Matches.Reserve(Keys.Num());
 	uint32 TotalMatches = 0;
 
@@ -799,7 +801,7 @@ void FilterAssets(TArray<const FAssetData*>& InOutResults, const TMap<KeyType, A
 		if (const ArrayType* Assets = AccelerationMap.Find(Key))
 		{
 			const FAssetData** AssetPtr = const_cast<const FAssetData**>(Assets->GetData());
-			Matches.Add(TArrayView<const FAssetData*>(AssetPtr, Assets->Num()));
+			Matches.Add(TConstArrayView<const FAssetData*>(AssetPtr, Assets->Num()));
 			TotalMatches += Assets->Num();
 		}
 	}
@@ -810,7 +812,7 @@ void FilterAssets(TArray<const FAssetData*>& InOutResults, const TMap<KeyType, A
 	{
 		// No previous Results; set Results equal to the values found in AccelerationMap
 		InOutResults.Reserve(TotalMatches);
-		for (TArrayView<const FAssetData*> Assets : Matches)
+		for (TConstArrayView<const FAssetData*> Assets : Matches)
 		{
 			InOutResults.Append(Assets);
 		}
@@ -827,7 +829,7 @@ void FilterAssets(TArray<const FAssetData*>& InOutResults, const TMap<KeyType, A
 		}
 		else
 		{
-			ArrayIntersection(InOutResults, TConstArrayView<TArrayView<const FAssetData*>>(Matches), TotalMatches);
+			ArrayIntersection(InOutResults, TConstArrayView<TConstArrayView<const FAssetData*>>(Matches), TotalMatches);
 		}
 	}
 }
@@ -866,8 +868,8 @@ void FilterAssets(TArray<const FAssetData*>&InOutResults,
 		}
 		else
 		{
-			TArrayView<const FAssetData*> ArrayView(Matches);
-			TConstArrayView<TArrayView<const FAssetData*>> ArrayViewOfArrayViews(&ArrayView, 1);
+			TConstArrayView<const FAssetData*> ArrayView(Matches);
+			TConstArrayView<TConstArrayView<const FAssetData*>> ArrayViewOfArrayViews(&ArrayView, 1);
 			ArrayIntersection(InOutResults, ArrayViewOfArrayViews, Matches.Num());
 		}
 	}
@@ -901,7 +903,7 @@ void FilterAssets(TArray<const FAssetData*>& InOutResults, const AccelerationMap
 	for (const TPair<FName, TOptional<FString>>& TagPair : TagsAndValues)
 	{
 		TArray<const FAssetData*>& Results = Matches.Emplace_GetRef();
-		if (const auto TagAssets = AccelerationMap.Find(TagPair.Key))
+		if (const auto* TagAssets = AccelerationMap.Find(TagPair.Key))
 		{
 			Results.Reserve(TagAssets->Num());
 			for (FAssetData* AssetData : *TagAssets)
@@ -940,7 +942,7 @@ void FilterAssets(TArray<const FAssetData*>& InOutResults, const AccelerationMap
 		else
 		{
 			// Convert Array of Arrays into format required by ArrayIntersection: Array of ArrayViews
-			TArray<TArrayView<const FAssetData*>, TInlineAllocator<10>> ArrayViewMatches;
+			TArray<TConstArrayView<const FAssetData*>, TInlineAllocator<10>> ArrayViewMatches;
 			ArrayViewMatches.Reserve(Matches.Num());
 			for (TArray<const FAssetData*>& MatchesElement : Matches)
 			{
@@ -1019,7 +1021,7 @@ void FilterAssetsByCachedClassesByTag(TArray<const FAssetData*>& InOutResults,
 		else
 		{
 			// Convert Array of Arrays into format required by ArrayIntersection: Array of ArrayViews
-			TArray<TArrayView<const FAssetData*>, TInlineAllocator<10>> ArrayViewMatches;
+			TArray<TConstArrayView<const FAssetData*>, TInlineAllocator<10>> ArrayViewMatches;
 			ArrayViewMatches.Reserve(Matches.Num());
 			for (TArray<const FAssetData*>& MatchesElement : Matches)
 			{
@@ -1112,7 +1114,8 @@ bool FAssetRegistryState::EnumerateAssets(const FARCompiledFilter& Filter, const
 
 	if (Filter.PackageNames.Num() > 0)
 	{
-		FilterAssets(AccumulatedResults, CachedAssetsByPackageName, Filter.PackageNames,
+		FilterAssets(
+			AccumulatedResults, CachedAssetsByPackageName, Filter.PackageNames,
 			[&Filter](const FAssetData* AssetData)
 			{
 				return Filter.PackageNames.Contains(AssetData->PackageName);
@@ -1228,12 +1231,18 @@ bool FAssetRegistryState::EnumerateAllAssets(const TSet<FName>& PackageNamesToSk
 
 void FAssetRegistryState::EnumerateAllAssets(TFunctionRef<void(const FAssetData&)> Callback) const
 {
-	for (const FAssetData* AssetData : CachedAssets)
-	{
-		if (AssetData)
+	EnumerateAllMutableAssets([&Callback](FAssetData& AssetData)
 		{
-			Callback(*AssetData);
-		}
+			Callback(AssetData);
+		});
+}
+
+void FAssetRegistryState::EnumerateAllMutableAssets(TFunctionRef<void(FAssetData&)> Callback) const
+{
+	for (FAssetData* AssetData : CachedAssets)
+	{
+		check(AssetData);
+		Callback(*AssetData);
 	}
 }
 
@@ -1241,22 +1250,24 @@ bool FAssetRegistryState::EnumerateAllAssets(const TSet<FName>& PackageNamesToSk
 	TFunctionRef<bool(const FAssetData&)> Callback,
 	UE::AssetRegistry::EEnumerateAssetsFlags InEnumerateFlags) const
 {
-	for (const FAssetData* AssetData : CachedAssets)
-	{
-		if (AssetData &&
-			!PackageNamesToSkip.Contains(AssetData->PackageName)
-			&& (EnumHasAnyFlags(InEnumerateFlags, UE::AssetRegistry::EEnumerateAssetsFlags::AllowUnmountedPaths) 
-				|| !IsPackageUnmountedAndFiltered(AssetData->PackageName))
-			&& (EnumHasAnyFlags(InEnumerateFlags, UE::AssetRegistry::EEnumerateAssetsFlags::AllowUnfilteredArAssets)
-				|| !UE::AssetRegistry::FFiltering::ShouldSkipAsset(AssetData->AssetClassPath,
-					AssetData->PackageFlags)))
+	using namespace UE::AssetRegistry;
+
+	EnumerateAllMutableAssets([&PackageNamesToSkip, &Callback, InEnumerateFlags, this](const FAssetData& AssetData)
 		{
-			if (!Callback(*AssetData))
+			if (!PackageNamesToSkip.Contains(AssetData.PackageName)
+				&& (EnumHasAnyFlags(InEnumerateFlags, EEnumerateAssetsFlags::AllowUnmountedPaths)
+					|| !IsPackageUnmountedAndFiltered(AssetData.PackageName))
+				&& (EnumHasAnyFlags(InEnumerateFlags, EEnumerateAssetsFlags::AllowUnfilteredArAssets)
+					|| !FFiltering::ShouldSkipAsset(AssetData.AssetClassPath,
+						AssetData.PackageFlags)))
 			{
-				return true;
+				if (!Callback(AssetData))
+				{
+					return false; // Stop iterating
+				}
 			}
-		}
-	}
+			return true; // keep iterating
+		});
 	return true;
 }
 
@@ -1547,10 +1558,10 @@ bool FAssetRegistryState::Save(FArchive& OriginalAr, const FAssetRegistrySeriali
 	{
 		TArray<TPair<FAssetData*, FSoftObjectPath>> SortedAssetsByObjectPath;
 		SortedAssetsByObjectPath.Reserve(AssetCount);
-		for (FAssetData* AssetData : CachedAssets)
-		{
-			SortedAssetsByObjectPath.Add({ AssetData, AssetData->GetSoftObjectPath() });
-		}
+		EnumerateAllMutableAssets([&SortedAssetsByObjectPath](FAssetData& AssetData)
+			{
+				SortedAssetsByObjectPath.Add({ &AssetData, AssetData.GetSoftObjectPath() });
+			});
 		Algo::Sort(SortedAssetsByObjectPath, [](const TPair<FAssetData*, FSoftObjectPath>& A,
 			const TPair<FAssetData*, FSoftObjectPath>& B)
 		{
@@ -1970,16 +1981,19 @@ SIZE_T FAssetRegistryState::GetAllocatedSize(bool bLogDetailed) const
 	SIZE_T NumTopLevelAssetPaths = 0;
 	FAssetDataTagMapSharedView::FMemoryCounter TagMemoryUsage;
 
-	for (const FAssetData* AssetData : CachedAssets)
+	EnumerateAllAssets(
+		[&AssetDataSize, &TagMemoryUsage, &NumAssetBundles, &NumSoftObjectPaths, &AssetBundlesSize,
+		&NumTopLevelAssetPaths]
+		(const FAssetData& AssetData)
 	{
-		AssetDataSize += sizeof(*AssetData);
-		TagMemoryUsage.Include(AssetData->TagsAndValues);
-		if (AssetData->TaggedAssetBundles.IsValid())
+		AssetDataSize += sizeof(AssetData);
+		TagMemoryUsage.Include(AssetData.TagsAndValues);
+		if (AssetData.TaggedAssetBundles.IsValid())
 		{
 			AssetBundlesSize += sizeof(FAssetBundleData);
-			AssetBundlesSize += AssetData->TaggedAssetBundles->Bundles.GetAllocatedSize();
-			NumAssetBundles += AssetData->TaggedAssetBundles->Bundles.Num();
-			for( const FAssetBundleEntry& Entry : AssetData->TaggedAssetBundles->Bundles)
+			AssetBundlesSize += AssetData.TaggedAssetBundles->Bundles.GetAllocatedSize();
+			NumAssetBundles += AssetData.TaggedAssetBundles->Bundles.Num();
+			for (const FAssetBundleEntry& Entry : AssetData.TaggedAssetBundles->Bundles)
 			{
 #if WITH_EDITORONLY_DATA
 				PRAGMA_DISABLE_DEPRECATION_WARNINGS;
@@ -1995,7 +2009,7 @@ SIZE_T FAssetRegistryState::GetAllocatedSize(bool bLogDetailed) const
 				NumTopLevelAssetPaths += Entry.AssetPaths.Num();
 			}
 		}
-	}
+	});
 
 	if (bLogDetailed)
 	{
@@ -2074,16 +2088,16 @@ FDependsNode* FAssetRegistryState::ResolveRedirector(FDependsNode* InDependency,
 		if (CachedAssetsByPackageName.Contains(CurrentDependency->GetPackageName()))
 		{
 			// Get the list of assets contained in this package
-			TArray<FAssetData*, TInlineAllocator<1>>& Assets = CachedAssetsByPackageName[CurrentDependency->GetPackageName()];
-
-			for (FAssetData* Asset : Assets)
+			EnumerateAssetsByPackageName(CurrentDependency->GetPackageName(), 
+				[&CurrentDependency, &InAllowedAssets, &Result, this](const FAssetData* Asset)
 			{
 				if (Asset->IsRedirector())
 				{
 					FDependsNode* ChainedRedirector = nullptr;
 					// This asset is a redirector, so we want to look at its dependencies and find the asset that
 					// it is redirecting to
-					CurrentDependency->IterateOverDependencies([&](FDependsNode* InDepends,
+					CurrentDependency->IterateOverDependencies(
+						[&InAllowedAssets, &ChainedRedirector, &Result, this](FDependsNode* InDepends,
 						UE::AssetRegistry::EDependencyCategory Category,
 						UE::AssetRegistry::EDependencyProperty Property, bool bDuplicate)
 					{
@@ -2110,9 +2124,8 @@ FDependsNode* FAssetRegistryState::ResolveRedirector(FDependsNode* InDependency,
 
 					if (ChainedRedirector)
 					{
-						// Found a redirector, break for loop
 						CurrentDependency = ChainedRedirector;
-						break;
+						return false; // Found a redirector, stop iterating assets for the current package
 					}
 				}
 				else
@@ -2123,9 +2136,10 @@ FDependsNode* FAssetRegistryState::ResolveRedirector(FDependsNode* InDependency,
 				if (Result)
 				{
 					// We found an allowed asset from the original dependency node. We're finished!
-					break;
+					return false; // stop iterating assets for the current package
 				}
-			}
+				return true; // keep iterating assets for the current package
+			});
 		}
 		else
 		{
@@ -2137,8 +2151,8 @@ FDependsNode* FAssetRegistryState::ResolveRedirector(FDependsNode* InDependency,
 	return Result;
 }
 
-template <typename KeyType>
-void ShrinkMultimap(TMap<KeyType, TArray<FAssetData*> >& Map)
+template <typename KeyType, typename ValueType>
+void ShrinkMultimap(TMap<KeyType, TArray<ValueType>>& Map)
 {
 	Map.Shrink();
 	for (auto& Pair : Map)
@@ -2154,10 +2168,10 @@ void FAssetRegistryState::SetAssetDatas(TArrayView<FAssetData> AssetDatas, const
 
 	NumAssets = AssetDatas.Num();
 	
-	auto SetPathCache = [&]() 
+	auto SetObjectPathCache = [this, &AssetDatas]()
 	{
-		CachedAssets.Empty(AssetDatas.Num());
-		for (FAssetData& AssetData : AssetDatas)
+		CachedAssets.Empty(NumAssets);
+		for (FAssetData& AssetData: AssetDatas)
 		{
 			CachedAssets.Add(&AssetData);
 		}
@@ -2166,32 +2180,31 @@ void FAssetRegistryState::SetAssetDatas(TArrayView<FAssetData> AssetDatas, const
 
 	// FAssetDatas sharing package name are very rare.
 	// Reserve up front and don't bother shrinking. 
-	auto SetPackageNameCache = [&]() 
+	auto SetPackageNameCache = [this, &AssetDatas]()
 	{
 		CachedAssetsByPackageName.Empty(AssetDatas.Num());
 		for (FAssetData& AssetData : AssetDatas)
 		{
-			TArray<FAssetData*, TInlineAllocator<1>>& PackageAssets
-				= CachedAssetsByPackageName.FindOrAdd(AssetData.PackageName);
-			PackageAssets.Add(&AssetData);
+			CachedAssetsByPackageName.FindOrAdd(AssetData.PackageName).Add(&AssetData);
 		}
 	};
 
-	auto SetOtherCaches = [&]()
+	auto SetPackagePathCache = [this, &AssetDatas]()
 	{
 		CachedAssetsByPath.Empty();
 		for (FAssetData& AssetData : AssetDatas)
 		{
-			TArray<FAssetData*>& PathAssets = CachedAssetsByPath.FindOrAdd(AssetData.PackagePath);
-			PathAssets.Add(&AssetData);
+			CachedAssetsByPath.FindOrAdd(AssetData.PackagePath).Add(&AssetData);
 		}
 		ShrinkMultimap(CachedAssetsByPath);
+	};
 
+	auto SetClassAndTagCaches = [this, &AssetDatas]()
+	{
 		CachedAssetsByClass.Empty();
 		for (FAssetData& AssetData : AssetDatas)
 		{
-			TArray<FAssetData*>& ClassAssets = CachedAssetsByClass.FindOrAdd(AssetData.AssetClassPath);
-			ClassAssets.Add(&AssetData);
+			CachedAssetsByClass.FindOrAdd(AssetData.AssetClassPath).Add(&AssetData);
 		}
 		ShrinkMultimap(CachedAssetsByClass);
 
@@ -2201,13 +2214,11 @@ void FAssetRegistryState::SetAssetDatas(TArrayView<FAssetData> AssetDatas, const
 		{
 			for (const TPair<FName, FAssetTagValueRef>& Pair : AssetData.TagsAndValues)
 			{
-				TSet<FAssetData*>& TagAssets = CachedAssetsByTag.FindOrAdd(Pair.Key);
-				TagAssets.Add(&AssetData);
+				CachedAssetsByTag.FindOrAdd(Pair.Key).Add(&AssetData);
 			}
 		}
-
 		CachedAssetsByTag.Shrink();
-		for (TPair<FName, TSet<FAssetData*>>& Pair : CachedAssetsByTag)
+		for (auto& Pair : CachedAssetsByTag)
 		{
 			Pair.Value.Shrink();
 		}
@@ -2217,8 +2228,7 @@ void FAssetRegistryState::SetAssetDatas(TArrayView<FAssetData> AssetDatas, const
 		{
 			for (const TPair<FName, FAssetTagValueRef>& Pair : AssetData.TagsAndValues)
 			{
-				TSet<FTopLevelAssetPath>& TagClasses = CachedClassesByTag.FindOrAdd(Pair.Key);
-				TagClasses.Add(AssetData.AssetClassPath);
+				CachedClassesByTag.FindOrAdd(Pair.Key).Add(AssetData.AssetClassPath);
 			}
 		}
 
@@ -2232,15 +2242,17 @@ void FAssetRegistryState::SetAssetDatas(TArrayView<FAssetData> AssetDatas, const
 
 	if (Options.ParallelWorkers <= 1)
 	{
-		SetPathCache();
+		SetObjectPathCache();
 		SetPackageNameCache();
-		SetOtherCaches();
+		SetPackagePathCache();
+		SetClassAndTagCaches();
 	}
 	else
 	{
-		TFuture<void> Task1 = Async(EAsyncExecution::TaskGraph, [=](){ SetPathCache(); });
-		TFuture<void> Task2 = Async(EAsyncExecution::TaskGraph, [=](){ SetPackageNameCache(); });
-		SetOtherCaches();
+		TFuture<void> Task1 = Async(EAsyncExecution::TaskGraph, [=]() { SetObjectPathCache(); });
+		TFuture<void> Task2 = Async(EAsyncExecution::TaskGraph, [=]() { SetPackageNameCache(); });
+		SetPackagePathCache();
+		SetClassAndTagCaches();
 		Task1.Wait();
 		Task2.Wait();
 	}
@@ -2250,6 +2262,7 @@ void FAssetRegistryState::AddAssetData(FAssetData* AssetData)
 {
 	bool bAlreadyInSet = false;
 	CachedAssets.Add(AssetData, &bAlreadyInSet);
+	FAssetData* MapElement = AssetData;
 	if (bAlreadyInSet)
 	{
 		UE_LOG(LogAssetRegistry, Error, TEXT("AddAssetData called with ObjectPath %s which already exists. ")
@@ -2260,24 +2273,18 @@ void FAssetRegistryState::AddAssetData(FAssetData* AssetData)
 		++NumAssets;
 	}
 
-	TArray<FAssetData*, TInlineAllocator<1>>& PackageAssets = CachedAssetsByPackageName.FindOrAdd(AssetData->PackageName);
-	TArray<FAssetData*>& PathAssets = CachedAssetsByPath.FindOrAdd(AssetData->PackagePath);
-	TArray<FAssetData*>& ClassAssets = CachedAssetsByClass.FindOrAdd(AssetData->AssetClassPath);
-
-	PackageAssets.Add(AssetData);
-	PathAssets.Add(AssetData);
-	ClassAssets.Add(AssetData);
+	CachedAssetsByPackageName.FindOrAdd(AssetData->PackageName).Add(MapElement);
+	CachedAssetsByPath.FindOrAdd(AssetData->PackagePath).Add(MapElement);
+	CachedAssetsByClass.FindOrAdd(AssetData->AssetClassPath).Add(MapElement);
 
 	for (auto TagIt = AssetData->TagsAndValues.CreateConstIterator(); TagIt; ++TagIt)
 	{
 		FName Key = TagIt.Key();
 
 #if UE_ASSETREGISTRY_CACHEDASSETSBYTAG
-		TSet<FAssetData*>& TagAssets = CachedAssetsByTag.FindOrAdd(Key);
-		TagAssets.Add(AssetData);
+		CachedAssetsByTag.FindOrAdd(Key).Add(MapElement);
 #else
-		TSet<FTopLevelAssetPath>& TagClasses = CachedClassesByTag.FindOrAdd(Key);
-		TagClasses.Add(AssetData->AssetClassPath);
+		CachedClassesByTag.FindOrAdd(Key).Add(AssetData->AssetClassPath);
 #endif
 	}
 }
@@ -2311,7 +2318,7 @@ void FAssetRegistryState::FilterTags(const FAssetRegistrySerializationOptions& O
 	// than trying to update its elements for each Asset change. For that reason we (currently) always recreate
 	// CachedAssetsByTag/CachedClassesByTag.
 #if UE_ASSETREGISTRY_CACHEDASSETSBYTAG
-	for (TPair<FName, TSet<FAssetData*>>& Pair : CachedAssetsByTag)
+	for (auto& Pair : CachedAssetsByTag)
 	{
 		Pair.Value.Reset();
 	}
@@ -2322,26 +2329,26 @@ void FAssetRegistryState::FilterTags(const FAssetRegistrySerializationOptions& O
 	}
 #endif
 
-	for (FAssetData* AssetData : CachedAssets) 
+	for (FAssetData* AssetDataPtr : CachedAssets)
 	{
-		check(AssetData);
-
+		FAssetData& AssetData = *AssetDataPtr;
+		FAssetData* AssetIndex = AssetDataPtr;
 		FAssetDataTagMap LocalTagsAndValues;
-		FAssetRegistryState::FilterTags(AssetData->TagsAndValues, LocalTagsAndValues,
-			Options.CookFilterlistTagsByClass.Find(AssetData->AssetClassPath), Options);
-		if (LocalTagsAndValues != AssetData->TagsAndValues)
+		FAssetRegistryState::FilterTags(AssetData.TagsAndValues, LocalTagsAndValues,
+			Options.CookFilterlistTagsByClass.Find(AssetData.AssetClassPath), Options);
+		if (LocalTagsAndValues != AssetData.TagsAndValues)
 		{
-			AssetData->TagsAndValues = FAssetDataTagMapSharedView(MoveTemp(LocalTagsAndValues));
+			AssetData.TagsAndValues = FAssetDataTagMapSharedView(MoveTemp(LocalTagsAndValues));
 		}
 
 		// Add the AssetData to all its CachedAssetsByTag/CachedClassesByTag keys even if nothing changed, because
 		// we are reconstructing all CachedAssetsByTag.
-		for (auto TagIt = AssetData->TagsAndValues.CreateConstIterator(); TagIt; ++TagIt)
+		for (auto TagIt = AssetData.TagsAndValues.CreateConstIterator(); TagIt; ++TagIt)
 		{
 #if UE_ASSETREGISTRY_CACHEDASSETSBYTAG
-			CachedAssetsByTag.FindOrAdd(TagIt.Key()).Add(AssetData);
+			CachedAssetsByTag.FindOrAdd(TagIt.Key()).Add(AssetIndex);
 #else
-			CachedClassesByTag.FindOrAdd(TagIt.Key()).Add(AssetData->AssetClassPath);
+			CachedClassesByTag.FindOrAdd(TagIt.Key()).Add(AssetData.AssetClassPath);
 #endif
 		}
 	}
@@ -2350,6 +2357,7 @@ void FAssetRegistryState::FilterTags(const FAssetRegistrySerializationOptions& O
 void FAssetRegistryState::SetTagsOnExistingAsset(FAssetData* AssetData, FAssetDataTagMap&& NewTags)
 {
 	// Update the tag cache map to remove deleted tags
+	FAssetData* AssetIndex = AssetData;
 	for (auto TagIt = AssetData->TagsAndValues.CreateConstIterator(); TagIt; ++TagIt)
 	{
 		const FName FNameKey = TagIt.Key();
@@ -2357,10 +2365,10 @@ void FAssetRegistryState::SetTagsOnExistingAsset(FAssetData* AssetData, FAssetDa
 		if (!NewTags.Contains(FNameKey))
 		{
 #if UE_ASSETREGISTRY_CACHEDASSETSBYTAG
-			TSet<FAssetData*>* OldTagAssets = CachedAssetsByTag.Find(FNameKey);
+			auto* OldTagAssets = CachedAssetsByTag.Find(FNameKey);
 			if (OldTagAssets)
 			{
-				OldTagAssets->Remove(AssetData);
+				OldTagAssets->Remove(AssetIndex);
 			}
 #else
 			// For CachedClassesByTag, we do not need to remove the asset's class from the entries
@@ -2376,7 +2384,7 @@ void FAssetRegistryState::SetTagsOnExistingAsset(FAssetData* AssetData, FAssetDa
 		if (!AssetData->TagsAndValues.Contains(FNameKey))
 		{
 #if UE_ASSETREGISTRY_CACHEDASSETSBYTAG
-			CachedAssetsByTag.FindOrAdd(FNameKey).Add(AssetData);
+			CachedAssetsByTag.FindOrAdd(FNameKey).Add(AssetIndex);
 #else
 			CachedClassesByTag.FindOrAdd(FNameKey).Add(AssetData->AssetClassPath);
 #endif
@@ -2397,11 +2405,12 @@ void FAssetRegistryState::SetDependencyNodeSorting(bool bSortDependencies, bool
 
 void FAssetRegistryState::UpdateAssetData(const FAssetData& NewAssetData, bool bCreateIfNotExists)
 {
-	if (FAssetData** AssetData = CachedAssets.Find(FCachedAssetKey(NewAssetData)))
+	FAssetData* AssetData = GetMutableAssetByObjectPath(FCachedAssetKey(NewAssetData));
+	if (AssetData)
 	{
-		UpdateAssetData(*AssetData, NewAssetData);
+		UpdateAssetData(AssetData, NewAssetData);
 	}
-	else
+	else if (bCreateIfNotExists)
 	{
 		AddAssetData(new FAssetData(NewAssetData));
 	}
@@ -2409,9 +2418,10 @@ void FAssetRegistryState::UpdateAssetData(const FAssetData& NewAssetData, bool b
 
 void FAssetRegistryState::UpdateAssetData(FAssetData&& NewAssetData, bool bCreateIfNotExists)
 {
-	if (FAssetData** AssetData = CachedAssets.Find(FCachedAssetKey(NewAssetData)))
+	FAssetData* AssetData = GetMutableAssetByObjectPath(FCachedAssetKey(NewAssetData));
+	if (AssetData)
 	{
-		UpdateAssetData(*AssetData, MoveTemp(NewAssetData));
+		UpdateAssetData(AssetData, MoveTemp(NewAssetData));
 	}
 	else if (bCreateIfNotExists)
 	{
@@ -2430,6 +2440,8 @@ void FAssetRegistryState::UpdateAssetData(FAssetData* AssetData, FAssetData&& Ne
 	FCachedAssetKey OldKey(AssetData);
 	FCachedAssetKey NewKey(NewAssetData);
 
+	FAssetData* AssetIndex = AssetData;
+
 	// Update ObjectPath
 	if (OldKey != NewKey)
 	{
@@ -2449,24 +2461,17 @@ void FAssetRegistryState::UpdateAssetData(FAssetData* AssetData, FAssetData&& Ne
 	if (AssetData->PackageName != NewAssetData.PackageName)
 	{
 		bKeyFieldIsModified = true;
-		TArray<FAssetData*, TInlineAllocator<1>>& NewPackageAssets
-			= CachedAssetsByPackageName.FindOrAdd(NewAssetData.PackageName);
-		TArray<FAssetData*, TInlineAllocator<1>>* OldPackageAssets
-			= CachedAssetsByPackageName.Find(AssetData->PackageName);
-
-		OldPackageAssets->Remove(AssetData);
-		NewPackageAssets.Add(AssetData);
+		CachedAssetsByPackageName.Find(AssetData->PackageName)->Remove(AssetIndex);
+		CachedAssetsByPackageName.FindOrAdd(NewAssetData.PackageName).Add(AssetIndex);
 	}
 
 	// Update PackagePath
 	if (AssetData->PackagePath != NewAssetData.PackagePath)
 	{
 		bKeyFieldIsModified = true;
-		TArray<FAssetData*>& NewPathAssets = CachedAssetsByPath.FindOrAdd(NewAssetData.PackagePath);
-		TArray<FAssetData*>* OldPathAssets = CachedAssetsByPath.Find(AssetData->PackagePath);
 
-		OldPathAssets->Remove(AssetData);
-		NewPathAssets.Add(AssetData);
+		CachedAssetsByPath.Find(AssetData->PackagePath)->Remove(AssetIndex);
+		CachedAssetsByPath.FindOrAdd(NewAssetData.PackagePath).Add(AssetIndex);
 	}
 
 	// AssetName is not a keyfield; compared below
@@ -2475,11 +2480,9 @@ void FAssetRegistryState::UpdateAssetData(FAssetData* AssetData, FAssetData&& Ne
 	if (AssetData->AssetClassPath != NewAssetData.AssetClassPath)
 	{
 		bKeyFieldIsModified = true;
-		TArray<FAssetData*>& NewClassAssets = CachedAssetsByClass.FindOrAdd(NewAssetData.AssetClassPath);
-		TArray<FAssetData*>* OldClassAssets = CachedAssetsByClass.Find(AssetData->AssetClassPath);
 
-		OldClassAssets->Remove(AssetData);
-		NewClassAssets.Add(AssetData);
+		CachedAssetsByClass.Find(AssetData->AssetClassPath)->Remove(AssetIndex);
+		CachedAssetsByClass.FindOrAdd(NewAssetData.AssetClassPath).Add(AssetIndex);
 	}
 
 	// PackageFlags is not a keyfield; compared below
@@ -2495,8 +2498,7 @@ void FAssetRegistryState::UpdateAssetData(FAssetData* AssetData, FAssetData&& Ne
 			if (!NewAssetData.TagsAndValues.Contains(FNameKey))
 			{
 #if UE_ASSETREGISTRY_CACHEDASSETSBYTAG
-				TSet<FAssetData*>* OldTagAssets = CachedAssetsByTag.Find(FNameKey);
-				OldTagAssets->Remove(AssetData);
+				CachedAssetsByTag.Find(FNameKey)->Remove(AssetIndex);
 #else
 				// For CachedClassesByTag, we do not need to remove the asset's class from the entries
 				// for the old tags. The class is not changed and still has the possibility of containing the tags.
@@ -2511,11 +2513,9 @@ void FAssetRegistryState::UpdateAssetData(FAssetData* AssetData, FAssetData&& Ne
 			if (!AssetData->TagsAndValues.Contains(FNameKey))
 			{
 #if UE_ASSETREGISTRY_CACHEDASSETSBYTAG
-				TSet<FAssetData*>& NewTagAssets = CachedAssetsByTag.FindOrAdd(FNameKey);
-				NewTagAssets.Add(AssetData);
+				CachedAssetsByTag.FindOrAdd(FNameKey).Add(AssetIndex);
 #else
-				TSet<FTopLevelAssetPath>& NewTagClasses = CachedClassesByTag.FindOrAdd(FNameKey);
-				NewTagClasses.Add(AssetData->AssetClassPath);
+				CachedClassesByTag.FindOrAdd(FNameKey).Add(AssetData->AssetClassPath);
 #endif
 			}
 		}
@@ -2547,7 +2547,7 @@ void FAssetRegistryState::UpdateAssetData(FAssetData* AssetData, FAssetData&& Ne
 	// Copy in new values
 	*AssetData = MoveTemp(NewAssetData);
 
-	// Can only re-add to asset map when key fields are updated
+	// Can only re-add to asset map after we update the key fields, because those change the hashvalue in CachedAssets
 	if (OldKey != NewKey)
 	{
 		bool bExisting = false;
@@ -2568,17 +2568,14 @@ void FAssetRegistryState::UpdateAssetData(FAssetData* AssetData, FAssetData&& Ne
 
 bool FAssetRegistryState::UpdateAssetDataPackageFlags(FName PackageName, uint32 PackageFlags)
 {
-	if (const TArray<FAssetData*, TInlineAllocator<1>>* Assets = CachedAssetsByPackageName.Find(PackageName))
-	{
-		for (FAssetData* Asset : *Assets)
+	bool bFoundValue = false;
+	EnumerateMutableAssetsByPackageName(PackageName, [&bFoundValue, PackageFlags](FAssetData* AssetData)
 		{
-			Asset->PackageFlags = PackageFlags;
-		}
-
-		return true;
-	}
-
-	return false;
+			AssetData->PackageFlags = PackageFlags;
+			bFoundValue = true;
+			return true;
+		});
+	return bFoundValue;
 }
 
 void FAssetRegistryState::RemoveAssetData(FAssetData* AssetData, bool bRemoveDependencyData,
@@ -2590,7 +2587,7 @@ void FAssetRegistryState::RemoveAssetData(FAssetData* AssetData, bool bRemoveDep
 		bOutRemovedPackageData = false;
 		return;
 	}
-	RemoveAssetData(AssetData, FCachedAssetKey(AssetData), bRemoveDependencyData, bOutRemovedAssetData,
+	RemoveAssetData(AssetData,  FCachedAssetKey(AssetData), bRemoveDependencyData, bOutRemovedAssetData,
 		bOutRemovedPackageData);
 	if (!bOutRemovedAssetData)
 	{
@@ -2605,13 +2602,14 @@ void FAssetRegistryState::RemoveAssetData(const FSoftObjectPath& AssetPath, bool
 {
 	FCachedAssetKey Key(AssetPath);
 	FAssetData** AssetDataPtrPtr = CachedAssets.Find(Key);
-	if (!AssetDataPtrPtr)
+	FAssetData* AssetData = AssetDataPtrPtr ? *AssetDataPtrPtr : nullptr;
+	if (!AssetData)
 	{
 		bOutRemovedAssetData = false;
 		bOutRemovedPackageData = false;
 		return;
 	}
-	RemoveAssetData(*AssetDataPtrPtr, Key, bRemoveDependencyData, bOutRemovedAssetData, bOutRemovedPackageData);
+	RemoveAssetData(AssetData, Key, bRemoveDependencyData, bOutRemovedAssetData, bOutRemovedPackageData);
 }
 
 void FAssetRegistryState::RemoveAssetData(FAssetData* AssetData, const FCachedAssetKey& Key,
@@ -2626,19 +2624,23 @@ void FAssetRegistryState::RemoveAssetData(FAssetData* AssetData, const FCachedAs
 		return;
 	}
 
-	TArray<FAssetData*, TInlineAllocator<1>>* OldPackageAssets
-		= CachedAssetsByPackageName.Find(AssetData->PackageName);
-	TArray<FAssetData*>* OldPathAssets = CachedAssetsByPath.Find(AssetData->PackagePath);
-	TArray<FAssetData*>* OldClassAssets = CachedAssetsByClass.Find(AssetData->AssetClassPath);
-	OldPackageAssets->RemoveSingleSwap(AssetData);
-	OldPathAssets->RemoveSingleSwap(AssetData);
-	OldClassAssets->RemoveSingleSwap(AssetData);
+	FAssetData* AssetIndex = AssetData;
+
+	TArray<FAssetData*, TInlineAllocator<1>>* OldPackageAssets =
+		CachedAssetsByPackageName.Find(AssetData->PackageName);
+	OldPackageAssets->RemoveSingleSwap(AssetIndex);
+	bool bOldPackageAssetsEmpty = OldPackageAssets->Num() == 0;
+	if (bOldPackageAssetsEmpty)
+	{
+		CachedAssetsByPackageName.Remove(AssetData->PackageName);
+	}
+	CachedAssetsByPath.Find(AssetData->PackagePath)->RemoveSingleSwap(AssetIndex);
+	CachedAssetsByClass.Find(AssetData->AssetClassPath)->RemoveSingleSwap(AssetIndex);
 
 #if UE_ASSETREGISTRY_CACHEDASSETSBYTAG
 	for (auto TagIt = AssetData->TagsAndValues.CreateConstIterator(); TagIt; ++TagIt)
 	{
-		TSet<FAssetData*>* OldTagAssets = CachedAssetsByTag.Find(TagIt.Key());
-		OldTagAssets->Remove(AssetData);
+		CachedAssetsByTag.Find(TagIt.Key())->Remove(AssetIndex);
 	}
 #else
 	// For CachedClassesByTag, we do not need to remove the asset's class from the entries
@@ -2646,10 +2648,8 @@ void FAssetRegistryState::RemoveAssetData(FAssetData* AssetData, const FCachedAs
 #endif
 
 	// Only remove dependencies and package data if there are no other known assets in the package
-	if (OldPackageAssets->Num() == 0)
+	if (bOldPackageAssetsEmpty)
 	{
-		CachedAssetsByPackageName.Remove(AssetData->PackageName);
-
 		// We need to update the cached dependencies references cache so that they know we no
 		// longer exist and so don't reference them.
 		if (bRemoveDependencyData)
@@ -2743,17 +2743,14 @@ bool FAssetRegistryState::RemoveDependsNode(const FAssetIdentifier& Identifier)
 
 void FAssetRegistryState::GetPrimaryAssetsIds(TSet<FPrimaryAssetId>& OutPrimaryAssets) const
 {
-	for (FAssetData* AssetData: CachedAssets)
-	{
-		if (AssetData)
+	EnumerateAllAssets([&OutPrimaryAssets](const FAssetData& AssetData)
 		{
-			FPrimaryAssetId PrimaryAssetId = AssetData->GetPrimaryAssetId();
+			FPrimaryAssetId PrimaryAssetId = AssetData.GetPrimaryAssetId();
 			if (PrimaryAssetId.IsValid())
 			{
 				OutPrimaryAssets.Add(PrimaryAssetId);
 			}
-		}
-	}
+		});
 }
 
 const FAssetPackageData* FAssetRegistryState::GetAssetPackageData(FName PackageName) const
@@ -2828,19 +2825,6 @@ bool FAssetRegistryState::IsFilterValid(const FARCompiledFilter& Filter)
 	return UE::AssetRegistry::Utils::IsFilterValid(Filter);
 }
 
-const TArray<const FAssetData*>& FAssetRegistryState::GetAssetsByClassName(const FName ClassName) const
-{
-	static TArray<const FAssetData*> InvalidArray;
-	for (const TPair< FTopLevelAssetPath, TArray<FAssetData*> >& Pair : CachedAssetsByClass)
-	{
-		if (Pair.Key.GetAssetName() == ClassName)
-		{
-			return reinterpret_cast<const TArray<const FAssetData*>&>(Pair.Value);
-		}
-	}
-	return InvalidArray;
-}
-
 void FAssetRegistryState::EnumerateAssetsByTagName(const FName TagName,
 	TFunctionRef<bool(const FAssetData* AssetData)> Callback) const
 {
@@ -2896,7 +2880,8 @@ void FAssetRegistryState::EnumerateTagToAssetDatas(
 #if UE_ASSETREGISTRY_CACHEDASSETSBYTAG
 	for (const TPair<FName, TSet<FAssetData*>>& Pair : CachedAssetsByTag)
 	{
-		const bool bKeepEnumerating = Callback(Pair.Key, [&Pair](IAssetRegistry::FAssetDataFunc AssetCallback)
+		const bool bKeepEnumerating = Callback(Pair.Key,
+			[&Pair, this](IAssetRegistry::FAssetDataFunc AssetCallback)
 			{
 				for (const FAssetData* AssetData : Pair.Value)
 				{
@@ -3278,10 +3263,10 @@ void FAssetRegistryState::Dump(const TArray<FString>& Arguments, TArray<FString>
 
 		TArray<FCachedAssetKey> Keys;
 		Keys.Reserve(CachedAssets.Num());
-		for( const FAssetData* AssetData : CachedAssets)
-		{
-			Keys.Emplace(FCachedAssetKey(*AssetData));
-		}
+		EnumerateAllAssets([&Keys](const FAssetData& AssetData)
+			{
+				Keys.Emplace(FCachedAssetKey(AssetData));
+			});
 		Keys.Sort([](const FCachedAssetKey& A, const FCachedAssetKey& B) {
 			return WriteToString<1024>(A).ToView().Compare(WriteToString<1024>(B).ToView(),
 				ESearchCase::IgnoreCase) < 0;
@@ -3331,13 +3316,17 @@ void FAssetRegistryState::Dump(const TArray<FString>& Arguments, TArray<FString>
 #endif
 	}
 
-	TArray<FAssetData*> SortedAssets;
+	TArray<const FAssetData*> SortedAssets;
 	auto InitializeSortedAssets = [&SortedAssets, this]()
 	{
 		if (SortedAssets.Num() != CachedAssets.Num())
 		{
-			SortedAssets = CachedAssets.Array();
-			Algo::Sort(SortedAssets, [](const FAssetData* A, FAssetData* B)
+			SortedAssets.Reserve(CachedAssets.Num());
+			EnumerateAllAssets([&SortedAssets](const FAssetData& AssetData)
+				{
+					SortedAssets.Add(&AssetData);
+				});
+			Algo::Sort(SortedAssets, [](const FAssetData* A, const FAssetData* B)
 				{ return A->GetSoftObjectPath().LexicalLess(B->GetSoftObjectPath()); }
 			);
 		}
diff --git a/Engine/Source/Runtime/AssetRegistry/Public/AssetRegistry/AssetRegistryState.h b/Engine/Source/Runtime/AssetRegistry/Public/AssetRegistry/AssetRegistryState.h
--- a/Engine/Source/Runtime/AssetRegistry/Public/AssetRegistry/AssetRegistryState.h
+++ b/Engine/Source/Runtime/AssetRegistry/Public/AssetRegistry/AssetRegistryState.h
@@ -43,19 +43,15 @@
 class FArchive;
 class FAssetDataTagMap;
 class FAssetDataTagMapSharedView;
+class FAssetRegistryGenerator;
 class FDependsNode;
 class FString;
+namespace UE::AssetRegistry { class FAssetRegistryImpl; }
 struct FARCompiledFilter;
+struct FAssetRegistryHeader;
 struct FPrimaryAssetId;
 template <typename FuncType> class TFunctionRef;
 
-namespace UE::AssetRegistry
-{
-	class FAssetRegistryImpl;
-}
-
-struct FAssetRegistryHeader;
-
 /**
  * Load/Save options used to modify how the cache is serialized.
  * These are read out of the AssetRegistry section of Engine.ini and can be changed per platform.
@@ -214,6 +210,7 @@ using FConstAssetDataMap = TSet<const FAssetData*, FCachedAssetKeyFuncs>;
  */
 class FAssetRegistryState
 {
+private:
 	using FCachedAssetKey = UE::AssetRegistry::Private::FCachedAssetKey;
 public:
 	// These types are an implementation detail and they and the functions which take/return them are subject to change
@@ -221,9 +218,9 @@ public:
 	using FAssetDataMap = UE::AssetRegistry::Private::FAssetDataMap;
 	using FConstAssetDataMap = UE::AssetRegistry::Private::FConstAssetDataMap;
 
-	FAssetRegistryState() = default;
+	FAssetRegistryState();
 	FAssetRegistryState(const FAssetRegistryState&) = delete;
-	FAssetRegistryState(FAssetRegistryState&& Rhs) { *this = MoveTemp(Rhs); }
+	FAssetRegistryState(FAssetRegistryState&& Rhs);
 	ASSETREGISTRY_API ~FAssetRegistryState();
 
 	FAssetRegistryState& operator=(const FAssetRegistryState&) = delete;
@@ -397,15 +394,6 @@ public:
 		UE::AssetRegistry::EDependencyCategory Category = UE::AssetRegistry::EDependencyCategory::All,
 		const UE::AssetRegistry::FDependencyQuery& Flags = UE::AssetRegistry::FDependencyQuery()) const;
 
-	/**
-	 * Gets the asset data for the specified object path
-	 *
-	 * @param ObjectPath the path of the object to be looked up
-	 * @return the assets data, null if not found
-	 */
-	UE_DEPRECATED(5.1, "Asset path FNames have been deprecated, use FSoftObjectPath instead.")
-	const FAssetData* GetAssetByObjectPath(const FName ObjectPath) const;
-
 	/**
 	 * Gets the asset data for the specified object path
 	 *
@@ -414,22 +402,24 @@ public:
 	 */
 	const FAssetData* GetAssetByObjectPath(const FSoftObjectPath& ObjectPath) const;
 
+	const FAssetData* GetAssetByObjectPath(const UE::AssetRegistry::Private::FCachedAssetKey& Key) const;
+
 	/**
 	 * Gets the asset data for the specified package name
 	 *
 	 * @param PackageName the path of the package to be looked up
 	 * @return an array of AssetData*, empty if nothing found
 	 */
+	UE_DEPRECATED(5.5, "Use EnumerateAssetsByPackageName or CopyAssetsByPackageName instead.")
 	TArrayView<FAssetData const* const> GetAssetsByPackageName(const FName PackageName) const;
+	void EnumerateAssetsByPackageName(const FName PackageName,
+		TFunctionRef<bool(const FAssetData* AssetData)> Callback) const;
+	/** Gets the array of AssetData pointers for the package; does not copy the AssetDatas, just the pointers. */
+	TArray<const FAssetData*> CopyAssetsByPackageName(const FName PackageName) const;
+	int32 NumAssetsByPackageName(const FName PackageName) const;
 
-	/**
-	 * Gets the asset data for the specified asset class
-	 *
-	 * @param ClassName the class name of the assets to look for
-	 * @return An array of AssetData*, empty if nothing found
-	 */
-	UE_DEPRECATED(5.1, "Class names are now represented by path names. Please use GetAssetsByClassPathName")
-	ASSETREGISTRY_API const TArray<const FAssetData*>& GetAssetsByClassName(const FName ClassName) const;
+	void EnumerateAssetsByPackagePath(FName LongPackagePathName,
+		TFunctionRef<bool(const FAssetData* AssetData)> Callback) const;
 
 	/**
 	 * Gets the asset data for the specified asset class
@@ -437,7 +427,10 @@ public:
 	 * @param ClassPathName the class path name of the assets to look for
 	 * @return An array of AssetData*, empty if nothing found
 	 */
+	UE_DEPRECATED(5.5, "Use EnumerateAssetsByClassPathName instead.")
 	const TArray<const FAssetData*>& GetAssetsByClassPathName(const FTopLevelAssetPath ClassPathName) const;
+	void EnumerateAssetsByClassPathName(const FTopLevelAssetPath ClassPathName,
+		TFunctionRef<bool(const FAssetData* AssetData)> Callback) const;
 
 	/**
 	 * Enumerates all tags of any asset in the AssetRegistry
@@ -455,7 +448,7 @@ public:
 	 * @param TagName the tag name to search for
 	 * @return An array of AssetData*, empty if nothing found
 	 */
-	UE_DEPRECATED(5.5, "GetAssetsByTagName has been deprecated. Please use EnumerateAssetsByTagName")
+	UE_DEPRECATED(5.5, "Use EnumerateAssetsByTagName instead.")
 	const TArray<const FAssetData*>& GetAssetsByTagName(const FName TagName) const;
 
 	/**
@@ -477,10 +470,11 @@ public:
 		TFunctionRef<bool(FName TagName, IAssetRegistry::FEnumerateAssetDatasFunc EnumerateAssets)> Callback) const;
 
 	/** Returns const version of internal ObjectPath->AssetData map for fast iteration */
+	UE_DEPRECATED(5.5, "FAssetDataMap is a complicated implementation detail of FAssetRegistryState. Use the enumeration functions on FAssetRegistryState instead of using it directly.")
 	const FConstAssetDataMap& GetAssetDataMap() const;
 
 	/** Returns const version of internal Tag->AssetDatas map for fast iteration */
-	UE_DEPRECATED(5.5, "GetTagToAssetDatasMap has been deprecated. Please use EnumerateTags or EnumerateTagToAssetDatas")
+	UE_DEPRECATED(5.5, "Use EnumerateTags or EnumerateTagToAssetDatas instead.")
 	const TMap<FName, const TArray<const FAssetData*>> GetTagToAssetDatasMap() const;
 
 	/** Returns const version of internal PackageName->PackageData map for fast iteration */
@@ -577,12 +571,6 @@ public:
 	/** Resets to default state */
 	ASSETREGISTRY_API void Reset();
 
-	/** Initializes cache from existing set of asset data and depends nodes */
-	ASSETREGISTRY_API void InitializeFromExisting(const FAssetDataMap& AssetDataMap,
-		const TMap<FAssetIdentifier, FDependsNode*>& DependsNodeMap,
-		const TMap<FName, FAssetPackageData*>& AssetPackageDataMap,
-		const FAssetRegistrySerializationOptions& Options,
-		EInitializationMode InitializationMode = EInitializationMode::Rebuild);
 	void InitializeFromExisting(const FAssetRegistryState& Existing, const FAssetRegistrySerializationOptions& Options,
 		EInitializationMode InitializationMode = EInitializationMode::Rebuild);
 
@@ -646,10 +634,10 @@ public:
 	static ASSETREGISTRY_API bool IsFilterValid(const FARCompiledFilter& Filter);
 
 	/** Returns the number of assets in this state */
-	int32 GetNumAssets() const { return NumAssets; }
+	int32 GetNumAssets() const;
 
 	/** Returns the number of packages in this state */
-	int32 GetNumPackages() const { return CachedAssetsByPackageName.Num(); }
+	int32 GetNumPackages() const;
 
 #if ASSET_REGISTRY_STATE_DUMPING_ENABLED
 	/**
@@ -669,6 +657,22 @@ public:
 #endif
 
 private:
+	// Accessors of mutable FAssetData*, for friend classes
+	ASSETREGISTRY_API void EnumerateAllMutableAssets(TFunctionRef<void(FAssetData&)> Callback) const;
+	FAssetData* GetMutableAssetByObjectPath(const FSoftObjectPath& ObjectPath) const;
+	FAssetData* GetMutableAssetByObjectPath(const UE::AssetRegistry::Private::FCachedAssetKey& Key) const;
+	void EnumerateMutableAssetsByPackageName(const FName PackageName,
+		TFunctionRef<bool(FAssetData* AssetData)> Callback) const;
+	void EnumerateMutableAssetsByPackagePath(FName LongPackagePathName,
+		TFunctionRef<bool(FAssetData* AssetData)> Callback) const;
+
+	/** Initializes cache from existing set of asset data and depends nodes */
+	ASSETREGISTRY_API void InitializeFromExisting(const FAssetDataMap& AssetDataMap,
+		const TMap<FAssetIdentifier, FDependsNode*>& DependsNodeMap,
+		const TMap<FName, FAssetPackageData*>& AssetPackageDataMap,
+		const FAssetRegistrySerializationOptions& Options,
+		EInitializationMode InitializationMode = EInitializationMode::Rebuild);
+
 	template<class Archive>
 	void Load(Archive&& Ar, const FAssetRegistryHeader& Header, const FAssetRegistryLoadOptions& Options);
 
@@ -758,6 +762,7 @@ private:
 	/** True if this asset registry state was loaded from a cooked asset registry */
 	bool bCookedGlobalAssetRegistryState = false;
 
+	friend class FAssetRegistryGenerator;
 	friend class UAssetRegistryImpl;
 	friend class UE::AssetRegistry::FAssetRegistryImpl;
 };
@@ -915,21 +920,36 @@ inline FAssetRegistryLoadOptions::FAssetRegistryLoadOptions(const FAssetRegistry
 {
 }
 
-inline const FAssetData* FAssetRegistryState::GetAssetByObjectPath(const FName ObjectPath) const
+inline FAssetRegistryState::FAssetRegistryState()
 {
-	return GetAssetByObjectPath(FSoftObjectPath(ObjectPath.ToString()));
+}
+
+inline FAssetRegistryState::FAssetRegistryState(FAssetRegistryState&& Rhs)
+{
+	*this = MoveTemp(Rhs);
 }
 
 inline const FAssetData* FAssetRegistryState::GetAssetByObjectPath(const FSoftObjectPath& ObjectPath) const
 {
-	FCachedAssetKey Key(ObjectPath);
-	FAssetData* const* FoundAsset = CachedAssets.Find(Key);
-	if (FoundAsset)
-	{
-		return *FoundAsset;
-	}
+	return GetAssetByObjectPath(UE::AssetRegistry::Private::FCachedAssetKey(ObjectPath));
+}
+
+inline const FAssetData* FAssetRegistryState::GetAssetByObjectPath(
+	const UE::AssetRegistry::Private::FCachedAssetKey& Key) const
+{
+	return GetMutableAssetByObjectPath(Key);
+}
+
+inline FAssetData* FAssetRegistryState::GetMutableAssetByObjectPath(const FSoftObjectPath& ObjectPath) const
+{
+	return GetMutableAssetByObjectPath(UE::AssetRegistry::Private::FCachedAssetKey(ObjectPath));
+}
 
-	return nullptr;
+inline FAssetData* FAssetRegistryState::GetMutableAssetByObjectPath(
+	const UE::AssetRegistry::Private::FCachedAssetKey& Key) const
+{
+	FAssetData* const* FoundAsset = CachedAssets.Find(Key);
+	return FoundAsset ? *FoundAsset : nullptr;
 }
 
 inline TArrayView<FAssetData const* const> FAssetRegistryState::GetAssetsByPackageName(const FName PackageName) const
@@ -942,6 +962,78 @@ inline TArrayView<FAssetData const* const> FAssetRegistryState::GetAssetsByPacka
 	return TArrayView<FAssetData* const>();
 }
 
+inline void FAssetRegistryState::EnumerateAssetsByPackageName(const FName PackageName,
+	TFunctionRef<bool(const FAssetData* AssetData)> Callback) const
+{
+	EnumerateMutableAssetsByPackageName(PackageName, [&Callback](FAssetData* AssetData)
+		{
+			return Callback(AssetData);
+		});
+}
+
+inline void FAssetRegistryState::EnumerateMutableAssetsByPackageName(const FName PackageName,
+	TFunctionRef<bool(FAssetData* AssetData)> Callback) const
+{
+	using namespace UE::AssetRegistry::Private;
+
+	if (const TArray<FAssetData*, TInlineAllocator<1>>* FoundAssetArray
+		= CachedAssetsByPackageName.Find(PackageName))
+	{
+		for (FAssetData* AssetData : *FoundAssetArray)
+		{
+			if (!Callback(AssetData))
+			{
+				break;
+			}
+		}
+	}
+}
+
+inline TArray<const FAssetData*> FAssetRegistryState::CopyAssetsByPackageName(const FName PackageName) const
+{
+	TArray<const FAssetData*> Result;
+	Result.Reserve(NumAssetsByPackageName(PackageName));
+	EnumerateAssetsByPackageName(PackageName, [&Result](const FAssetData* AssetData)
+		{
+			Result.Add(AssetData);
+			return true;
+		});
+	return Result;
+}
+
+inline int32 FAssetRegistryState::NumAssetsByPackageName(const FName PackageName) const
+{
+	const auto* Array = CachedAssetsByPackageName.Find(PackageName);
+	return Array ? Array->Num() : 0;
+}
+
+inline void FAssetRegistryState::EnumerateAssetsByPackagePath(FName LongPackagePathName,
+	TFunctionRef<bool(const FAssetData* AssetData)> Callback) const
+{
+	EnumerateMutableAssetsByPackagePath(LongPackagePathName, [&Callback](FAssetData* AssetData)
+		{
+			return Callback(AssetData);
+		});
+}
+
+inline void FAssetRegistryState::EnumerateMutableAssetsByPackagePath(FName LongPackagePathName,
+	TFunctionRef<bool(FAssetData* AssetData)> Callback) const
+{
+	using namespace UE::AssetRegistry::Private;
+
+	const TArray<FAssetData*>* AssetsInPath = CachedAssetsByPath.Find(LongPackagePathName);
+	if (AssetsInPath)
+	{
+		for (FAssetData* AssetData : *AssetsInPath)
+		{
+			if (!Callback(AssetData))
+			{
+				break;
+			}
+		}
+	}
+}
+
 inline const TArray<const FAssetData*>& FAssetRegistryState::GetAssetsByClassPathName(
 	const FTopLevelAssetPath ClassPathName) const
 {
@@ -955,10 +1047,24 @@ inline const TArray<const FAssetData*>& FAssetRegistryState::GetAssetsByClassPat
 	return InvalidArray;
 }
 
+inline void FAssetRegistryState::EnumerateAssetsByClassPathName(
+	const FTopLevelAssetPath ClassPathName, TFunctionRef<bool(const FAssetData* AssetData)> Callback) const
+{
+	using namespace UE::AssetRegistry::Private;
+
+	if (const TArray<FAssetData*>* FoundAssetArray = CachedAssetsByClass.Find(ClassPathName))
+	{
+		for (FAssetData* AssetData : *FoundAssetArray)
+		{
+			Callback(AssetData);
+		}
+	}
+}
+
 inline void FAssetRegistryState::EnumerateTags(TFunctionRef<bool(FName TagName)> Callback) const
 {
 #if UE_ASSETREGISTRY_CACHEDASSETSBYTAG
-	for (const TPair<FName, TSet<FAssetData*>>& Pair : CachedAssetsByTag)
+	for (const auto& Pair : CachedAssetsByTag)
 #else
 	for (const TPair<FName, TSet<FTopLevelAssetPath>>& Pair : CachedClassesByTag)
 #endif
@@ -1007,16 +1113,26 @@ inline const TMap<FName, const FAssetPackageData*>& FAssetRegistryState::GetAsse
 inline void FAssetRegistryState::GetPackageNames(TArray<FName>& OutPackageNames) const
 {
 	OutPackageNames.Reserve(CachedAssetsByPackageName.Num());
-	for (auto It = CachedAssetsByPackageName.CreateConstIterator(); It; ++It)
+	for (const auto& It : CachedAssetsByPackageName)
 	{
-		OutPackageNames.Add(It.Key());
+		OutPackageNames.Add(It.Key);
 	}
 }
 
+inline int32 FAssetRegistryState::GetNumAssets() const
+{
+	return NumAssets;
+}
+
+inline int32 FAssetRegistryState::GetNumPackages() const
+{
+	return CachedAssetsByPackageName.Num();
+}
+
 inline void FAssetRegistryState::InitializeFromExisting(const FAssetRegistryState& Existing,
 	const FAssetRegistrySerializationOptions& Options,
 	EInitializationMode InitializationMode)
 {
 	InitializeFromExisting(Existing.CachedAssets, Existing.CachedDependsNodes,
 		Existing.CachedPackageData, Options, InitializationMode);
-}
+}
\ No newline at end of file
