Reducing thread contention in the asset registry by making bInitialSearchCompleted so it doesn't need to be protected by the interface RW lock and changing DeferredEvents to be protected by its own lock.

diff --git a/Engine/Source/Runtime/AssetRegistry/Private/AssetDataGatherer.cpp b/Engine/Source/Runtime/AssetRegistry/Private/AssetDataGatherer.cpp
--- a/Engine/Source/Runtime/AssetRegistry/Private/AssetDataGatherer.cpp
+++ b/Engine/Source/Runtime/AssetRegistry/Private/AssetDataGatherer.cpp
@@ -3722,8 +3722,7 @@ uint32 FAssetDataGatherer::Run()
 			if ((Status != UE::AssetRegistry::Impl::EGatherStatus::TickActiveGatherActive)
 				&& (Status != UE::AssetRegistry::Impl::EGatherStatus::TickActiveGatherIdle))
 			{
-				// Consider creating a TRACE_CPUPROFILER_EVENT_SCOPE_STR_CONDITIONAL
-				TRACE_CPUPROFILER_EVENT_SCOPE_TEXT_CONDITIONAL("FAssetDataGatherer Sleep", !bIsInitialSearchCompleted.load(std::memory_order_relaxed));
+				TRACE_CPUPROFILER_EVENT_SCOPE_STR_CONDITIONAL("FAssetDataGatherer Sleep", !bIsInitialSearchCompleted.load(std::memory_order_relaxed));
 				FPlatformProcess::Sleep(bLocalIdle ? IdleSleepTime : PausedSleepTime);
 			}
 		}
diff --git a/Engine/Source/Runtime/AssetRegistry/Private/AssetRegistry.cpp b/Engine/Source/Runtime/AssetRegistry/Private/AssetRegistry.cpp
--- a/Engine/Source/Runtime/AssetRegistry/Private/AssetRegistry.cpp
+++ b/Engine/Source/Runtime/AssetRegistry/Private/AssetRegistry.cpp
@@ -946,7 +946,7 @@ void FAssetRegistryImpl::Initialize(Impl::FInitializeContext& Context)
 	const double StartupStartTime = FPlatformTime::Seconds();
 
 	bInitialSearchStarted = false;
-	bInitialSearchCompleted = true;
+	bInitialSearchCompleted.store(true, std::memory_order_relaxed);
 	GatherStatus = Impl::EGatherStatus::TickActiveGatherActive;
 	PerformanceMode = Impl::EPerformanceMode::MostlyStatic;
 
@@ -1915,7 +1915,7 @@ void FAssetRegistryImpl::SearchAllAssets(Impl::FEventContext& EventContext,
 	{
 		InitialSearchStartTime = FPlatformTime::Seconds();
 		bInitialSearchStarted = true;
-		bInitialSearchCompleted = false;
+		bInitialSearchCompleted.store(false, std::memory_order_relaxed);
 	}
 
 	FAssetDataGatherer& Gatherer = *GlobalGatherer;
@@ -2010,6 +2010,7 @@ void UAssetRegistryImpl::WaitForCompletion()
 			if (IsInGameThread())
 			{
 				// Process any deferred events. Required since deferred events would block sending the FileLoadedEvent
+				FScopeLock DeferredEventsLock(&DeferredEventsCriticalSection);
 				EventContext = MoveTemp(DeferredEvents);
 				DeferredEvents.Clear();
 			}
@@ -4487,7 +4488,7 @@ namespace UE::AssetRegistry
 
 bool FAssetRegistryImpl::IsLoadingAssets() const
 {
-	return !bInitialSearchCompleted;
+	return !IsInitialSearchCompleted();
 }
 
 }
@@ -4519,8 +4520,11 @@ UE::AssetRegistry::Impl::EGatherStatus UAssetRegistryImpl::TickOnBackgroundThrea
 			TickContext.bHandleDeferred = true;
 			Status = GuardedData.TickGatherer(TickContext);
 
-			DeferredEvents.Append(MoveTemp(EventContext));
-			EventContext.Clear();
+			{
+				FScopeLock DeferredEventsLock(&DeferredEventsCriticalSection);
+				DeferredEvents.Append(MoveTemp(EventContext));
+				EventContext.Clear();
+			}
 		}
 		else 
 		{
@@ -4577,8 +4581,12 @@ void UAssetRegistryImpl::Tick(float DeltaTime)
 			GetInheritanceContextWithRequiredLock(InterfaceScopeLock, InheritanceContext, InheritanceBuffer);
 
 			// Process any deferred events
-			EventContext = MoveTemp(DeferredEvents);
-			DeferredEvents.Clear();
+			{
+				FScopeLock DeferredEventsLock(&DeferredEventsCriticalSection);
+				EventContext = MoveTemp(DeferredEvents);
+				DeferredEvents.Clear();
+			}
+
 			if (EventContext.IsEmpty())
 			{
 				// Tick the Gatherer
@@ -4598,9 +4606,7 @@ void UAssetRegistryImpl::Tick(float DeltaTime)
 		}
 		else
 		{
-			// We still take the interface lock but only briefly in this case. We don't try to actually process the data 
-			// because the background thread is working on it. We just pump the events.
-			UE::AssetRegistry::FInterfaceWriteScopeLock InterfaceScopeLock(InterfaceLock);
+			FScopeLock DeferredEventsLock(&DeferredEventsCriticalSection);
 			EventContext.Append(MoveTemp(DeferredEvents));
 			DeferredEvents.Clear();
 		}
@@ -4652,8 +4658,7 @@ bool FAssetRegistryImpl::ClassRequiresGameThreadProcessing(const UClass* Class)
 
 Impl::EGatherStatus FAssetRegistryImpl::TickGatherer(Impl::FTickContext& TickContext)
 {
-	// Consider creating a TRACE_CPUPROFILER_EVENT_SCOPE_STR_CONDITIONAL
-	TRACE_CPUPROFILER_EVENT_SCOPE_TEXT_CONDITIONAL("FAssetRegistryImpl::TickGatherer", !bInitialSearchCompleted);
+	TRACE_CPUPROFILER_EVENT_SCOPE_STR_CONDITIONAL("FAssetRegistryImpl::TickGatherer", !IsInitialSearchCompleted());
 
 	using namespace UE::AssetRegistry::Impl;
 	FEventContext& EventContext = TickContext.EventContext;
@@ -5020,7 +5025,7 @@ Impl::EGatherStatus FAssetRegistryImpl::TickGatherer(Impl::FTickContext& TickCon
 
 	if (OutStatus == EGatherStatus::Complete)
 	{
-		if (!bInitialSearchCompleted)
+		if (!IsInitialSearchCompleted())
 		{
 			// Finishing the background search is blocked until preloading complete because plugins can be mounted during
 			// startup up until that point, and we need to wait for all the plugins to load before declaring completion.
@@ -5082,8 +5087,6 @@ Impl::EGatherStatus FAssetRegistryImpl::TickGatherer(Impl::FTickContext& TickCon
 
 void FAssetRegistryImpl::OnInitialSearchCompleted(Impl::FEventContext& EventContext)
 {
-	bInitialSearchCompleted = true;
-
 #if WITH_EDITOR
 	// update redirectors
 	UpdateRedirectCollector();
@@ -5098,6 +5101,8 @@ void FAssetRegistryImpl::OnInitialSearchCompleted(Impl::FEventContext& EventCont
 	GlobalGatherer->OnInitialSearchCompleted();
 
 	EventContext.bFileLoadedEventBroadcast = true;
+	
+	bInitialSearchCompleted.store(true, std::memory_order_relaxed);
 }
 
 void FAssetRegistryImpl::LogSearchDiagnostics(double StartTime)
@@ -7159,7 +7164,7 @@ void FAssetRegistryImpl::OnDirectoryChanged(Impl::FEventContext& EventContext,
 	{
 		if (FileChangesProcessed[FileIdx].Action == FFileChangeData::FCA_RescanRequired)
 		{
-			if (bInitialSearchStarted && !bInitialSearchCompleted)
+			if (bInitialSearchStarted && !IsInitialSearchCompleted())
 			{
 				// Ignore rescan request during initial scan as it is probably caused by the scan itself
 				UE_LOG(LogAssetRegistry, Log, TEXT("FAssetRegistry ignoring rescan request for %s during startup"), *FileChangesProcessed[FileIdx].Filename);
@@ -8880,9 +8885,9 @@ void UAssetRegistryImpl::Broadcast(UE::AssetRegistry::Impl::FEventContext& Event
 		{
 			return;
 		}
-		UE::AssetRegistry::FInterfaceWriteScopeLock InterfaceScopeLock(InterfaceLock);
 		// Broadcast should not be called on DeferredEvents; DeferredEvents should be moved to a separate EventContext
-		// while under the InterfaceLock, and broadcast called on that separate Eventcontext outside of the lock.
+		// and broadcast called on that separate EventContext outside of the lock.
+		FScopeLock DeferredEventsLock(&DeferredEventsCriticalSection);
 		check(&EventContext != &DeferredEvents);
 		DeferredEvents.Append(MoveTemp(EventContext));
 		EventContext.Clear();
@@ -9054,9 +9059,9 @@ void UAssetRegistryImpl::Broadcast(UE::AssetRegistry::Impl::FEventContext& Event
 		if (!bAllowFileLoadedEvent)
 		{
 			// Do not send the file loaded event yet; pass the flag on instead
-			UE::AssetRegistry::FInterfaceWriteScopeLock InterfaceScopeLock(InterfaceLock);
+			FScopeLock DeferredEventsLock(&DeferredEventsCriticalSection);
 			// Broadcast should not be called on DeferredEvents; DeferredEvents should be moved to a separate EventContext
-			// while under the InterfaceLock, and broadcast called on that separate Eventcontext outside of the lock.
+			// and broadcast called on that separate EventContext outside of the lock.
 			check(&EventContext != &DeferredEvents);
 			DeferredEvents.Append(MoveTemp(EventContext));
 			EventContext.Clear();
@@ -9067,7 +9072,7 @@ void UAssetRegistryImpl::Broadcast(UE::AssetRegistry::Impl::FEventContext& Event
 
 		FEventContext CopiedDeferredEvents;
 		{
-			UE::AssetRegistry::FInterfaceWriteScopeLock InterfaceScopeLock(InterfaceLock);
+			FScopeLock DeferredEventsLock(&DeferredEventsCriticalSection);
 			check(&EventContext != &DeferredEvents);
 			CopiedDeferredEvents = MoveTemp(DeferredEvents);
 			DeferredEvents.Clear();
diff --git a/Engine/Source/Runtime/AssetRegistry/Private/AssetRegistry.h b/Engine/Source/Runtime/AssetRegistry/Private/AssetRegistry.h
--- a/Engine/Source/Runtime/AssetRegistry/Private/AssetRegistry.h
+++ b/Engine/Source/Runtime/AssetRegistry/Private/AssetRegistry.h
@@ -412,7 +412,11 @@ private:
 	/** The delegate to execute scanning has ended */
 	FScanEndedEvent ScanEndedEvent;
 
+	/** Storage for events that will be broadcast later on the game thread. Only safe
+	 *  to access under the DeferredEventsCriticalSection. 
+	 */
 	UE::AssetRegistry::Impl::FEventContext DeferredEvents;
+	FCriticalSection DeferredEventsCriticalSection;
 
 	friend class UE::AssetRegistry::FAssetRegistryImpl;
 	friend struct UE::AssetRegistry::Premade::FAsyncConsumer;
diff --git a/Engine/Source/Runtime/AssetRegistry/Private/AssetRegistryImpl.h b/Engine/Source/Runtime/AssetRegistry/Private/AssetRegistryImpl.h
--- a/Engine/Source/Runtime/AssetRegistry/Private/AssetRegistryImpl.h
+++ b/Engine/Source/Runtime/AssetRegistry/Private/AssetRegistryImpl.h
@@ -352,7 +352,7 @@ public:
 	void GetSubClasses(Impl::FClassInheritanceContext& InheritanceContext, const TArray<FTopLevelAssetPath>& InClassNames,
 		const TSet<FTopLevelAssetPath>& ExcludedClassNames, TSet<FTopLevelAssetPath>& SubClassNames) const;
 
-	bool IsInitialSearchCompleted() const { return bInitialSearchCompleted; }
+	bool IsInitialSearchCompleted() const { return bInitialSearchCompleted.load(std::memory_order_relaxed); }
 	bool IsTempCachingEnabled() const { return bIsTempCachingEnabled; }
 	bool IsTempCachingAlwaysEnabled() const { return bIsTempCachingAlwaysEnabled; }
 	bool IsInitialSearchStarted() const { return bInitialSearchStarted; }
@@ -579,8 +579,10 @@ private:
 	double InitialSearchStartTime = 0.0f;
 	/** Flag to indicate if we used an initial async search */
 	bool bInitialSearchStarted;
-	/** Flag to indicate if the initial background search has completed */
-	bool bInitialSearchCompleted;
+	/** Flag to indicate if the initial background search has completed. All access are relaxed because
+	 *  the actual search data can only be accessed under a proper lock.
+	 */
+	std::atomic<bool> bInitialSearchCompleted;
 	/**
 	 * Flag to indicate PreloadingComplete; finishing the background search is blocked until preloading complete
 	 * because preloading can add assets.
