AssetRegistry: SetKeyFuncsTest: Change test framework to Test2 so it is tested during Foundation LowLevelTests.
#rnx

diff --git a/Engine/Source/Runtime/AssetRegistry/Private/SetKeyFuncsTest.cpp b/Engine/Source/Runtime/AssetRegistry/Private/SetKeyFuncsTest.cpp
--- a/Engine/Source/Runtime/AssetRegistry/Private/SetKeyFuncsTest.cpp
+++ b/Engine/Source/Runtime/AssetRegistry/Private/SetKeyFuncsTest.cpp
@@ -1,20 +1,13 @@
 // Copyright Epic Games, Inc. All Rights Reserved.
 
-#include "SetKeyFuncs.h"
+#if WITH_TESTS
 
-#include "Misc/AutomationTest.h"
+#include "SetKeyFuncs.h"
 
-#if WITH_DEV_AUTOMATION_TESTS
 #include "Containers/Array.h"
-#include "Templates/Tuple.h"
-#endif
+#include "Tests/TestHarnessAdapter.h"
 
-#if WITH_DEV_AUTOMATION_TESTS
-IMPLEMENT_SIMPLE_AUTOMATION_TEST(FSetKeyFuncsTest, "System.AssetRegistry.SetKeyFuncs",
-	EAutomationTestFlags::EditorContext | EAutomationTestFlags::ClientContext | EAutomationTestFlags::ServerContext
-	| EAutomationTestFlags::EngineFilter);
-
-bool FSetKeyFuncsTest::RunTest(const FString& Parameters)
+TEST_CASE_NAMED(FSetKeyFuncsTest, "System::AssetRegistry::SetKeyFuncs", "[ApplicationContextMask][EngingeFilter]")
 {
 	struct FData
 	{
@@ -47,21 +40,20 @@ bool FSetKeyFuncsTest::RunTest(const FString& Parameters)
 					return TypeHash;
 				}
 			}
-			Test->AddError(FString::Printf(TEXT("GetTypeHash was unexpectedly called on unknown value %u."), Value));
+			FAIL_CHECK(FString::Printf(TEXT("GetTypeHash was unexpectedly called on unknown value %u."), Value));
 			return (uint32)-1;
 		}
 		bool Matches(uint32 A, uint32 B)
 		{
 			return A == B;
 		}
-		FSetKeyFuncsTest* Test;
 		FData* Data;
 	};
 
 	FData Data;
-	TSetKeyFuncs<uint32, FKeyFuncs1> Set(FKeyFuncs1{ this, &Data });
+	TSetKeyFuncs<uint32, FKeyFuncs1> Set(FKeyFuncs1{ &Data });
 
-	auto ValidateExpectedSetContents = [&Data, &Set, this]()
+	auto ValidateExpectedSetContents = [&Data, &Set]()
 		{
 			for (uint32 VFind = 0; VFind < Data.EndValue; ++VFind)
 			{
@@ -74,16 +66,16 @@ bool FSetKeyFuncsTest::RunTest(const FString& Parameters)
 				{
 					if (Data.ValueInSet[VFind])
 					{
-						AddError(FString::Printf(TEXT("Expected in-set value %u was unexpectedly not found."), VFind));
+						FAIL_CHECK(FString::Printf(TEXT("Expected in-set value %u was unexpectedly not found."), VFind));
 					}
 					else
 					{
-						AddError(FString::Printf(TEXT("Expected not-in-set value %u was unexpectedly found."), VFind));
+						FAIL_CHECK(FString::Printf(TEXT("Expected not-in-set value %u was unexpectedly found."), VFind));
 					}
 				}
 				else if (ExistingValue && *ExistingValue != VFind)
 				{
-					AddError(FString::Printf(TEXT("Expected value %u returned invalid result %u."),
+					FAIL_CHECK(FString::Printf(TEXT("Expected value %u returned invalid result %u."),
 						VFind, *ExistingValue));
 				}
 			}
@@ -96,7 +88,7 @@ bool FSetKeyFuncsTest::RunTest(const FString& Parameters)
 			{
 				if (Data.ValueInSetCopiedFromSet[VIter])
 				{
-					AddError(FString::Printf(TEXT("Value %u unexpectedly encountered twice in Set iterator."), VIter));
+					FAIL_CHECK(FString::Printf(TEXT("Value %u unexpectedly encountered twice in Set iterator."), VIter));
 				}
 				Data.ValueInSetCopiedFromSet[VIter] = true;
 			}
@@ -106,17 +98,17 @@ bool FSetKeyFuncsTest::RunTest(const FString& Parameters)
 				{
 					if (Data.ValueInSet[VIter])
 					{
-						AddError(FString::Printf(TEXT("Expected in-set value %u was unexpectedly not found in the Set iterator."), VIter));
+						FAIL_CHECK(FString::Printf(TEXT("Expected in-set value %u was unexpectedly not found in the Set iterator."), VIter));
 					}
 					else
 					{
-						AddError(FString::Printf(TEXT("Expected not-in-set value %u was unexpectedly found in the Set iterator."), VIter));
+						FAIL_CHECK(FString::Printf(TEXT("Expected not-in-set value %u was unexpectedly found in the Set iterator."), VIter));
 					}
 				}
 			}
 			if (Set.Num() != Data.NumInSet)
 			{
-				AddError(FString::Printf(TEXT("Set Num(%d) != expected num(%d)."), Set.Num(), Data.NumInSet));
+				FAIL_CHECK(FString::Printf(TEXT("Set Num(%d) != expected num(%d)."), Set.Num(), Data.NumInSet));
 			}
 		};
 
@@ -196,7 +188,7 @@ bool FSetKeyFuncsTest::RunTest(const FString& Parameters)
 		case 5:
 		{
 			TSetKeyFuncs<uint32, FKeyFuncs1> CopySet(Set);
-			CopySet.SetKeyFuncs(FKeyFuncs1{ this, &Data });
+			CopySet.SetKeyFuncs(FKeyFuncs1{ &Data });
 			Set.Empty();
 			Set = CopySet;
 			break;
@@ -233,17 +225,15 @@ bool FSetKeyFuncsTest::RunTest(const FString& Parameters)
 		}
 
 		FSetKeyFuncsStats Stats = Set.GetStats();
-		TestTrue(TEXT("AverageSearch >= 1"), Stats.AverageSearch >= 1.f);
-		TestTrue(TEXT("LongestSearch >= 1"), Stats.LongestSearch >= 1);
-		TestTrue(TEXT("LongestSearch >= 1"), Set.GetAllocatedSize() >= Set.Num()*1);
+		CHECK_MESSAGE(TEXT("AverageSearch >= 1"), Stats.AverageSearch >= 1.f);
+		CHECK_MESSAGE(TEXT("LongestSearch >= 1"), Stats.LongestSearch >= 1);
+		CHECK_MESSAGE(TEXT("LongestSearch >= 1"), Set.GetAllocatedSize() >= Set.Num() * 1);
 	}
 
-	TSetKeyFuncs<uint32, FKeyFuncs1> EmptySet(FKeyFuncs1{ this, &Data });
+	TSetKeyFuncs<uint32, FKeyFuncs1> EmptySet(FKeyFuncs1{ &Data });
 	FSetKeyFuncsStats Stats = EmptySet.GetStats();
-	TestTrue(TEXT("Empty AverageSearch == 0.0"), Stats.AverageSearch == 0.0f);
-	TestTrue(TEXT("Empty LongestSearch == 0"), Stats.AverageSearch == 0);
-
-	return true;
+	CHECK_MESSAGE(TEXT("Empty AverageSearch == 0.0"), Stats.AverageSearch == 0.0f);
+	CHECK_MESSAGE(TEXT("Empty LongestSearch == 0"), Stats.AverageSearch == 0);
 }
 
-#endif
\ No newline at end of file
+#endif // WITH_TESTS
\ No newline at end of file
