AssetRegistry - Bugfix to Memory optimization -
When removing AssetData from AssetRegistryState, we have to remove from CachedAssets container last, since this container holds the memory for FAssetDataPtrIndex that is referred to from the other containers.
#rnx

diff --git a/Engine/Source/Runtime/AssetRegistry/Private/AssetRegistryState.cpp b/Engine/Source/Runtime/AssetRegistry/Private/AssetRegistryState.cpp
--- a/Engine/Source/Runtime/AssetRegistry/Private/AssetRegistryState.cpp
+++ b/Engine/Source/Runtime/AssetRegistry/Private/AssetRegistryState.cpp
@@ -2917,9 +2917,8 @@ void FAssetRegistryState::RemoveAssetData(FAssetData* AssetData, const FCachedAs
 
 	bOutRemovedAssetData = false;
 	bOutRemovedPackageData = false;
-	int32 NumRemoved = CachedAssets.Remove(Key);
-	check(NumRemoved <= 1);
-	if (NumRemoved == 0)
+
+	if (!CachedAssets.Find(Key))
 	{
 		return;
 	}
@@ -2952,6 +2951,10 @@ void FAssetRegistryState::RemoveAssetData(FAssetData* AssetData, const FCachedAs
 	// for the old tags. The class is not changed and still has the possibility of containing the tags.
 #endif
 
+	// In the UE_ASSETREGISTRY_INDIRECT_ASSETDATA_POINTERS case the other containers hold an index into CachedAssets,
+	// so we can only remove from CachedAssets after removing from all other containers.
+	CachedAssets.Remove(Key);
+
 	// Only remove dependencies and package data if there are no other known assets in the package
 	if (bOldPackageAssetsEmpty)
 	{
