Added a function get the asset package data via a batched call.

diff --git a/Engine/Source/Runtime/AssetRegistry/Private/AssetRegistry.cpp b/Engine/Source/Runtime/AssetRegistry/Private/AssetRegistry.cpp
--- a/Engine/Source/Runtime/AssetRegistry/Private/AssetRegistry.cpp
+++ b/Engine/Source/Runtime/AssetRegistry/Private/AssetRegistry.cpp
@@ -3089,6 +3089,27 @@ TOptional<FAssetPackageData> UAssetRegistryImpl::GetAssetPackageDataCopy(FName P
 	return AssetPackageData ? *AssetPackageData : TOptional<FAssetPackageData>();
 }
 
+TArray<TOptional<FAssetPackageData>> UAssetRegistryImpl::GetAssetPackageDatasCopy(TArrayView<FName> PackageNames) const
+{
+	TArray<TOptional<FAssetPackageData>> OutAssetPackagesData;
+	OutAssetPackagesData.Reserve(PackageNames.Num());
+
+	UE::AssetRegistry::FInterfaceReadScopeLock InterfaceScopeLock(InterfaceLock);
+	for (FName PackageName : PackageNames)
+	{
+		if (const FAssetPackageData* AssetPackageData = GuardedData.GetState().GetAssetPackageData(PackageName))
+		{
+			OutAssetPackagesData.Emplace(*AssetPackageData);
+		}
+		else
+		{
+			OutAssetPackagesData.Add(TOptional<FAssetPackageData>());
+		}
+	}
+
+	return OutAssetPackagesData;
+}
+
 void UAssetRegistryImpl::EnumerateAllPackages(TFunctionRef<void(FName PackageName, const FAssetPackageData& PackageData)> Callback) const
 {
 	UE::AssetRegistry::FInterfaceReadScopeLock InterfaceScopeLock(InterfaceLock);
diff --git a/Engine/Source/Runtime/AssetRegistry/Private/AssetRegistry.h b/Engine/Source/Runtime/AssetRegistry/Private/AssetRegistry.h
--- a/Engine/Source/Runtime/AssetRegistry/Private/AssetRegistry.h
+++ b/Engine/Source/Runtime/AssetRegistry/Private/AssetRegistry.h
@@ -65,6 +65,7 @@ public:
 	virtual bool GetReferencers(const FAssetIdentifier& AssetIdentifier, TArray<FAssetDependency>& OutReferencers, UE::AssetRegistry::EDependencyCategory Category = UE::AssetRegistry::EDependencyCategory::All, const UE::AssetRegistry::FDependencyQuery& Flags = UE::AssetRegistry::FDependencyQuery()) const override;
 	virtual bool GetReferencers(FName PackageName, TArray<FName>& OutReferencers, UE::AssetRegistry::EDependencyCategory Category = UE::AssetRegistry::EDependencyCategory::Package, const UE::AssetRegistry::FDependencyQuery& Flags = UE::AssetRegistry::FDependencyQuery()) const override; //-V1101
 	virtual TOptional<FAssetPackageData> GetAssetPackageDataCopy(FName PackageName) const override;
+	virtual TArray<TOptional<FAssetPackageData>> GetAssetPackageDatasCopy(TArrayView<FName> PackageNames) const override;
 	virtual void EnumerateAllPackages(TFunctionRef<void(FName PackageName, const FAssetPackageData& PackageData)> Callback) const override;
 	virtual bool DoesPackageExistOnDisk(FName PackageName, FString* OutCorrectCasePackageName = nullptr, FString* OutExtension = nullptr) const override;
 	virtual FSoftObjectPath GetRedirectedObjectPath(const FSoftObjectPath& ObjectPath) override;
diff --git a/Engine/Source/Runtime/AssetRegistry/Public/AssetRegistry/IAssetRegistry.h b/Engine/Source/Runtime/AssetRegistry/Public/AssetRegistry/IAssetRegistry.h
--- a/Engine/Source/Runtime/AssetRegistry/Public/AssetRegistry/IAssetRegistry.h
+++ b/Engine/Source/Runtime/AssetRegistry/Public/AssetRegistry/IAssetRegistry.h
@@ -517,6 +517,9 @@ public:
 	/** Finds Package data for a package name. This data is only updated on save and can only be accessed for valid packages */
 	virtual TOptional<FAssetPackageData> GetAssetPackageDataCopy(FName PackageName) const = 0;
 
+	/** Find Package Data for the packages name. This data is only updated on save and can only be accessed for valid packages */
+	virtual TArray<TOptional<FAssetPackageData>> GetAssetPackageDatasCopy(TArrayView<FName> PackageNames) const = 0;
+
 	/**
 	 * Enumerate all PackageDatas in the AssetRegistry. The callback is called from within the AssetRegistry's lock, so it must not call
 	 * arbitrary code that could call back into the AssetRegistry; doing so would deadlock.
