AssetRegistry: Background Ticking: Remove some unnecessary reads of AssetRegistry data during editor startup to prevent interrupting the background tick. Add a counter for how many times it is interrupted.
#rnx

diff --git a/Engine/Source/Runtime/AssetRegistry/Private/AssetRegistry.cpp b/Engine/Source/Runtime/AssetRegistry/Private/AssetRegistry.cpp
--- a/Engine/Source/Runtime/AssetRegistry/Private/AssetRegistry.cpp
+++ b/Engine/Source/Runtime/AssetRegistry/Private/AssetRegistry.cpp
@@ -4503,7 +4503,6 @@ void UAssetRegistryImpl::PackageDeleted(UPackage* DeletedPackage)
 
 bool UAssetRegistryImpl::IsLoadingAssets() const
 {
-	UE::AssetRegistry::FInterfaceReadScopeLock InterfaceScopeLock(InterfaceLock);
 	return GuardedData.IsLoadingAssets();
 }
 
@@ -4535,7 +4534,22 @@ UE::AssetRegistry::Impl::EGatherStatus UAssetRegistryImpl::TickOnBackgroundThrea
 			GetInheritanceContextWithRequiredLock(InterfaceScopeLock, InitializeContext.InheritanceContext, InitializeContext.InheritanceBuffer);
 
 			UE::AssetRegistry::Impl::FInterruptionContext::ShouldExitEarlyCallbackType EarlyExitHelper = 
-				[this]()->bool { return InterfaceLock.HasWaiters() || IsBackgroundProcessingPaused(); };
+				[this]()->bool
+				{
+					if (InterfaceLock.HasWaiters() || IsBackgroundProcessingPaused())
+					{
+#if WITH_EDITOR
+						// During EngineStartup many packages are loading and need to query the AssetRegistry; do not
+						// count them in the metric for backgroundtick interruptions.
+						if (IsEngineStartupModuleLoadingComplete())
+						{
+							++GuardedData.GetBackgroundTickInterruptionsCount();
+						}
+#endif
+						return true;
+					}
+					return false;
+				};
 
 			UE::AssetRegistry::Impl::FTickContext TickContext(EventContext, InheritanceContext);
 			TickContext.InterruptionContext.SetLimitedTickTime(FPlatformTime::Seconds(),
@@ -5144,13 +5158,19 @@ void FAssetRegistryImpl::LogSearchDiagnostics(double StartTime)
 	Telemetry.NumCachedAssetFiles = Diagnostics.NumCachedAssetFiles;
 	Telemetry.NumUncachedAssetFiles = Diagnostics.NumUncachedAssetFiles;
 	FTelemetryRouter::Get().ProvideTelemetry(Telemetry);
-	UE_LOG(LogAssetRegistry, Log, TEXT("AssetRegistryGather time %.4fs: AssetDataDiscovery %0.4fs, AssetDataGather %0.4fs, StoreResults %0.4fs. Wall time %0.4fs.")
+#if !NO_LOGGING
+	TStringBuilder<256> Message;
+	Message.Appendf(TEXT("AssetRegistryGather time %.4fs: AssetDataDiscovery %0.4fs, AssetDataGather %0.4fs, StoreResults %0.4fs. Wall time %0.4fs.")
 		TEXT("\n\tNumCachedDirectories %d. NumUncachedDirectories %d. NumCachedFiles %d. NumUncachedFiles %d."),
 		Total, Diagnostics.DiscoveryTimeSeconds, Diagnostics.GatherTimeSeconds, StoreGatherResultsTimeSeconds,
 		Diagnostics.WallTimeSeconds, Diagnostics.NumCachedDirectories, Diagnostics.NumUncachedDirectories,
 		Diagnostics.NumCachedAssetFiles, Diagnostics.NumUncachedAssetFiles);
+#if WITH_EDITOR
+	Message.Appendf(TEXT("\n\tBackgroundTickInterruptions %d."), BackgroundTickInterruptionsCount);
+#endif // WITH_EDITOR
+
+	UE_LOG(LogAssetRegistry, Log, TEXT("%s"), *Message);
 
-#if !NO_LOGGING
 	if (bVerboseLogging)
 	{
 		UE_LOG(LogAssetRegistry, Verbose, TEXT("TagMemoryUse:"));
@@ -5166,7 +5186,7 @@ void FAssetRegistryImpl::LogSearchDiagnostics(double StartTime)
 				*Pair.Key.ToString(), (float)Pair.Value / (1000.f * 1000.f));
 		}
 	}
-#endif
+#endif // !NO_LOGGING
 }
 
 void FAssetRegistryImpl::TickGatherPackage(Impl::FEventContext& EventContext, const FString& PackageName, const FString& LocalPath)
@@ -8222,13 +8242,13 @@ void UAssetRegistryImpl::OnGetExtraObjectTags(FAssetRegistryTagsContext Context)
 {
 	if (bAddMetaDataTagsToOnGetExtraObjectTags)
 	{
-		// It is critical that bIncludeOnlyOnDiskAssets=true otherwise this will cause an infinite loop
-		const FAssetData AssetData = GetAssetByObjectPath(FSoftObjectPath::ConstructFromObject(Context.GetObject()), /*bIncludeOnlyOnDiskAssets=*/true);
 		// Adding metadata tags from disk is only necessary for cooked assets; uncooked assets still have the metadata and add them elsewhere
 		// in UObject::GetAssetRegistryTags. Adding the tags from disk into uncooked assets would make the tags impossible to remove when
 		// the uncooked assets are resaved.
-		if ((AssetData.PackageFlags & PKG_Cooked) != 0)
+		if (Context.GetObject()->GetPackage()->HasAnyPackageFlags(PKG_Cooked))
 		{
+			// It is critical that bIncludeOnlyOnDiskAssets=true otherwise this will cause an infinite loop
+			const FAssetData AssetData = GetAssetByObjectPath(FSoftObjectPath::ConstructFromObject(Context.GetObject()), /*bIncludeOnlyOnDiskAssets=*/true);
 			TSet<FName>& MetaDataTags = UObject::GetMetaDataTagsForAssetRegistry();
 			for (const FName MetaDataTag : MetaDataTags)
 			{
diff --git a/Engine/Source/Runtime/AssetRegistry/Private/AssetRegistryImpl.h b/Engine/Source/Runtime/AssetRegistry/Private/AssetRegistryImpl.h
--- a/Engine/Source/Runtime/AssetRegistry/Private/AssetRegistryImpl.h
+++ b/Engine/Source/Runtime/AssetRegistry/Private/AssetRegistryImpl.h
@@ -399,6 +399,7 @@ public:
 	void RequestPauseBackgroundProcessing();
 	void RequestResumeBackgroundProcessing();
 	bool IsBackgroundProcessingPaused() const;
+	uint32& GetBackgroundTickInterruptionsCount();
 
 #endif
 private:
@@ -639,6 +640,12 @@ private:
 	/** A map from directoryname to packagename of Packages that have CalculatedDependencies on packages in the directory. */
 	TMultiMap<FString, FName> DirectoryReferencers;
 
+	/**
+	 * Number of times during the startup scan that TickOnBackgroundThread was interrupted by a request for the
+	 * AssetRegistry's lock from another thread.
+	 */
+	uint32 BackgroundTickInterruptionsCount = 0;
+
 	/** A map of per asset class dependency gatherer called in LoadCalculatedDependencies */
 	TMultiMap<FTopLevelAssetPath, UE::AssetDependencyGatherer::Private::FRegisteredAssetDependencyGatherer*> RegisteredDependencyGathererClasses;
 	mutable FRWLock RegisteredDependencyGathererClassesLock;
@@ -841,4 +848,18 @@ void EnumerateMemoryAssets(const FARCompiledFilter& InFilter, TSet<FName>& OutPa
 
 } // namespace Utils within namespace UE::AssetRegistry
 
+
+///////////////////////////////////////////////////////
+// Inline implementations
+///////////////////////////////////////////////////////
+
+
+#if WITH_EDITOR
+inline uint32& FAssetRegistryImpl::GetBackgroundTickInterruptionsCount()
+{
+	return BackgroundTickInterruptionsCount;
+}
+#endif
+
+
 } // namespace UE::AssetRegistry
