Fix SortByAssetName sorting incorrectly: "Scene_10" should be lexically smaller than "Scene_01".

diff --git a/Engine/Source/Runtime/AssetRegistry/Private/AssetRegistryHelpers.cpp b/Engine/Source/Runtime/AssetRegistry/Private/AssetRegistryHelpers.cpp
--- a/Engine/Source/Runtime/AssetRegistry/Private/AssetRegistryHelpers.cpp
+++ b/Engine/Source/Runtime/AssetRegistry/Private/AssetRegistryHelpers.cpp
@@ -160,7 +160,14 @@ void UAssetRegistryHelpers::SortByAssetName(TArray<FAssetData>& Assets, EAssetRe
 	UE::AssetRegistry::SortAssets(Assets,
 		[](const FAssetData& Left, const FAssetData& Right)
 		{
-			return Left.AssetName.LexicalLess(Right.AssetName);
+			// Summary: String compare needed instead of FName::LexicalLess.
+			// Reason: FName::LexicalLess says e.g. FName("Scene_10") < FName("Scene_01") (while: FString::operator< says "Scene_01" < "Scene_10").
+			// Explanation:
+			// - "Scene_10" has ComparisionIndex of "Scene" and Number = 11,
+			// - "Scene_01" has ComparisionIndex of "Scene_01" and number 0
+			// - Thus, (FName("Scene_10").LexicalLess(FName("Scene_01")) internally ends up checking "Scene" < "Scene_01" , which is true.
+			// - For reference, "Scene_1" has ComparisionIndex of "Scene" and number 2, which, when sorting, we'd "expect" Scene_01 to have, too.
+			return Left.AssetName.ToString() < Right.AssetName.ToString();
 		}, SortOrder);
 }
 
