ObjectVersion Bump: AssetRegistry dependency section now has an extra TArray<TPair<FName, uint32>>.
AssetRegistry: Change EDependencyProperty::Build to trigger the propagation of manage dependencies. We need this so that a UWorld can propagate its manage dependencies to its LevelInstance, which is an editor-only reference. Turn the Build property off of the dependencies created from imports and softobjectpaths. Add build dependencies during editor package save via the function FObjectPreSaveContext.AddCookBuildDependency.
#rnx

diff --git a/Engine/Plugins/Editor/AssetManagerEditor/Source/AssetManagerEditor/Private/AssetManagerEditorCommands.cpp b/Engine/Plugins/Editor/AssetManagerEditor/Source/AssetManagerEditor/Private/AssetManagerEditorCommands.cpp
--- a/Engine/Plugins/Editor/AssetManagerEditor/Source/AssetManagerEditor/Private/AssetManagerEditorCommands.cpp
+++ b/Engine/Plugins/Editor/AssetManagerEditor/Source/AssetManagerEditor/Private/AssetManagerEditorCommands.cpp
@@ -44,8 +44,20 @@ void FAssetManagerEditorCommands::RegisterCommands()
 
 	UI_COMMAND(ShowSoftReferences, "Show Soft References", "Toggles visibility of Soft References", EUserInterfaceActionType::ToggleButton, FInputChord(EKeys::S));
 	UI_COMMAND(ShowHardReferences, "Show Hard References", "Toggles visibility of Hard References", EUserInterfaceActionType::ToggleButton, FInputChord(EKeys::H));
-	UI_COMMAND(ShowEditorOnlyReferences, "Show EditorOnly References","Toggles visibility of EditorOnly References", EUserInterfaceActionType::ToggleButton, FInputChord(EKeys::E));
-
+	UI_COMMAND(EditorOnlyReferenceFilterTypeGame, "Show Game References",
+		"Only UsedInGame references are displayed. UsedInGame references cause the dependencies to be cooked and"
+		" cause management assignments from the referencer to be propagated to the dependencies.",
+		EUserInterfaceActionType::RadioButton, FInputChord(EKeys::G));
+	UI_COMMAND(EditorOnlyReferenceFilterTypePropagation, "Show Propagation References",
+		"UsedInGame and Build references are displayed. Build references are displayed even if they are EditorOnly."
+		" All of these references cause management assignments from the referencer to be propagated to the"
+		" dependencies, but the EditorOnly Build dependencies do not cause the dependencies to be cooked.",
+		EUserInterfaceActionType::RadioButton, FInputChord(EKeys::B));
+	UI_COMMAND(EditorOnlyReferenceFilterTypeEditorOnly, "Show EditorOnly References",
+		"UsedInGame, Build, and EditorOnly references are displayed. EditorOnly references do not cause the"
+		" dependencies to be cooked and do not propagate management assignments. Use this mode to see how assets"
+		" interact in the editor.",
+		EUserInterfaceActionType::RadioButton, FInputChord(EKeys::E));
 	UI_COMMAND(ShowManagementReferences, "Show Management References","Toggles visibility of Management References (i.e. PrimaryAssetIDs)", EUserInterfaceActionType::ToggleButton, FInputChord(EKeys::M));
 	UI_COMMAND(ShowNameReferences, "Show Name References","Toggles visibility of Name References (i.e. Gameplay Tags and Data Table Row Handles)", EUserInterfaceActionType::ToggleButton, FInputChord(EKeys::N));
 	UI_COMMAND(ShowCodePackages, "Show C++ Packages","Toggles visibility of C++ Packages", EUserInterfaceActionType::ToggleButton, FInputChord(EKeys::C));
diff --git a/Engine/Plugins/Editor/AssetManagerEditor/Source/AssetManagerEditor/Private/AssetManagerEditorCommands.h b/Engine/Plugins/Editor/AssetManagerEditor/Source/AssetManagerEditor/Private/AssetManagerEditorCommands.h
--- a/Engine/Plugins/Editor/AssetManagerEditor/Source/AssetManagerEditor/Private/AssetManagerEditorCommands.h
+++ b/Engine/Plugins/Editor/AssetManagerEditor/Source/AssetManagerEditor/Private/AssetManagerEditorCommands.h
@@ -63,7 +63,7 @@ public:
 	TSharedPtr<FUICommandInfo> DecreaseBreadth;
 
 	// Set the Breadth LImit
-	TSharedPtr<FUICommandInfo> SetBreadth;;
+	TSharedPtr<FUICommandInfo> SetBreadth;
 
 	// Toggles visiblity of Soft References
 	TSharedPtr<FUICommandInfo> ShowSoftReferences;
@@ -71,8 +71,10 @@ public:
 	// Toggles visiblity of Hard References
 	TSharedPtr<FUICommandInfo> ShowHardReferences;
 
-	// Toggles visiblity of EditorOnly References
-	TSharedPtr<FUICommandInfo> ShowEditorOnlyReferences;
+	// Sets which level of filter type between Game and EditorOnly is visible
+	TSharedPtr<FUICommandInfo> EditorOnlyReferenceFilterTypeGame;
+	TSharedPtr<FUICommandInfo> EditorOnlyReferenceFilterTypePropagation;
+	TSharedPtr<FUICommandInfo> EditorOnlyReferenceFilterTypeEditorOnly;
 
 	// Toggles visiblity of Management References (i.e. PrimaryAssetIDs)
 	TSharedPtr<FUICommandInfo> ShowManagementReferences;
diff --git a/Engine/Plugins/Editor/AssetManagerEditor/Source/AssetManagerEditor/Private/ReferenceViewer/EdGraph_ReferenceViewer.cpp b/Engine/Plugins/Editor/AssetManagerEditor/Source/AssetManagerEditor/Private/ReferenceViewer/EdGraph_ReferenceViewer.cpp
--- a/Engine/Plugins/Editor/AssetManagerEditor/Source/AssetManagerEditor/Private/ReferenceViewer/EdGraph_ReferenceViewer.cpp
+++ b/Engine/Plugins/Editor/AssetManagerEditor/Source/AssetManagerEditor/Private/ReferenceViewer/EdGraph_ReferenceViewer.cpp
@@ -189,7 +189,13 @@ FAssetManagerDependencyQuery UEdGraph_ReferenceViewer::GetReferenceSearchFlags(b
 		Query.Categories |= EDependencyCategory::Package;
 		Query.Flags |= bLocalIsShowSoftReferences ? EDependencyQuery::NoRequirements : EDependencyQuery::Hard;
 		Query.Flags |= Settings->IsShowHardReferences() ? EDependencyQuery::NoRequirements : EDependencyQuery::Soft;
-		Query.Flags |= Settings->IsShowEditorOnlyReferences() ? EDependencyQuery::NoRequirements : EDependencyQuery::Game;
+		switch (Settings->GetEditorOnlyReferenceFilterType())
+		{
+		case EEditorOnlyReferenceFilterType::Game: Query.Flags |= EDependencyQuery::Game; break;
+		case EEditorOnlyReferenceFilterType::Propagation: Query.Flags |= EDependencyQuery::Propagation; break;
+		case EEditorOnlyReferenceFilterType::EditorOnly: [[fallthrough]];
+		default: /* No requirements */ ; break;
+		}
 	}
 	if (Settings->IsShowSearchableNames() && !bHardOnly)
 	{
diff --git a/Engine/Plugins/Editor/AssetManagerEditor/Source/AssetManagerEditor/Private/ReferenceViewer/ReferenceViewerSettings.cpp b/Engine/Plugins/Editor/AssetManagerEditor/Source/AssetManagerEditor/Private/ReferenceViewer/ReferenceViewerSettings.cpp
--- a/Engine/Plugins/Editor/AssetManagerEditor/Source/AssetManagerEditor/Private/ReferenceViewer/ReferenceViewerSettings.cpp
+++ b/Engine/Plugins/Editor/AssetManagerEditor/Source/AssetManagerEditor/Private/ReferenceViewer/ReferenceViewerSettings.cpp
@@ -39,9 +39,9 @@ bool UReferenceViewerSettings::IsShowDuplicates() const
 	return bIsShowDuplicates;
 }
 
-bool UReferenceViewerSettings::IsShowEditorOnlyReferences() const
+EEditorOnlyReferenceFilterType UReferenceViewerSettings::GetEditorOnlyReferenceFilterType() const
 {
-	return bIsShowEditorOnlyReferences;
+	return EditorOnlyReferenceFilterType;
 }
 
 bool UReferenceViewerSettings::IsShowManagementReferences() const
@@ -111,9 +111,9 @@ void UReferenceViewerSettings::SetShowDuplicatesEnabled(bool bNewEnabled)
 	SaveConfig();
 }
 
-void UReferenceViewerSettings::SetShowEditorOnlyReferencesEnabled(bool bNewEnabled)
+void UReferenceViewerSettings::SetEditorOnlyReferenceFilterType(EEditorOnlyReferenceFilterType Value)
 {
-	bIsShowEditorOnlyReferences = bNewEnabled;
+	EditorOnlyReferenceFilterType = Value;
 	SaveConfig();
 }
 
diff --git a/Engine/Plugins/Editor/AssetManagerEditor/Source/AssetManagerEditor/Private/ReferenceViewer/SReferenceViewer.cpp b/Engine/Plugins/Editor/AssetManagerEditor/Source/AssetManagerEditor/Private/ReferenceViewer/SReferenceViewer.cpp
--- a/Engine/Plugins/Editor/AssetManagerEditor/Source/AssetManagerEditor/Private/ReferenceViewer/SReferenceViewer.cpp
+++ b/Engine/Plugins/Editor/AssetManagerEditor/Source/AssetManagerEditor/Private/ReferenceViewer/SReferenceViewer.cpp
@@ -1444,21 +1444,20 @@ bool SReferenceViewer::IsShowDuplicatesChecked() const
 	return Settings->GetFindPathEnabled() || Settings->IsShowDuplicates();
 }
 
-void SReferenceViewer::OnShowEditorOnlyReferencesChanged()
+void SReferenceViewer::OnEditorOnlyReferenceFilterTypeChanged(EEditorOnlyReferenceFilterType Value)
 {
-	Settings->SetShowEditorOnlyReferencesEnabled(!Settings->IsShowEditorOnlyReferences());
+	Settings->SetEditorOnlyReferenceFilterType(Value);
 	if (GraphObj)
 	{
 		GraphObj->RebuildGraph();
 	}
 }
 
-bool SReferenceViewer::IsShowEditorOnlyReferencesChecked() const
+EEditorOnlyReferenceFilterType SReferenceViewer::GetEditorOnlyReferenceFilterType() const
 {
-	return Settings->IsShowEditorOnlyReferences();
+	return Settings->GetEditorOnlyReferenceFilterType();
 }
 
-
 bool SReferenceViewer::GetManagementReferencesVisibility() const
 {
 	return bShowShowReferencesOptions;
@@ -1633,10 +1632,24 @@ void SReferenceViewer::RegisterActions()
 		FIsActionButtonVisible::CreateLambda([this] { return bShowShowReferencesOptions; }));
 
 	ReferenceViewerActions->MapAction(
-		FAssetManagerEditorCommands::Get().ShowEditorOnlyReferences,
-		FExecuteAction::CreateSP(this, &SReferenceViewer::OnShowEditorOnlyReferencesChanged),
+		FAssetManagerEditorCommands::Get().EditorOnlyReferenceFilterTypeGame,
+		FExecuteAction::CreateSPLambda(this, [this]() { OnEditorOnlyReferenceFilterTypeChanged(EEditorOnlyReferenceFilterType::Game); }),
 		FCanExecuteAction(),	
-		FIsActionChecked::CreateSP(this, &SReferenceViewer::IsShowEditorOnlyReferencesChecked),
+		FIsActionChecked::CreateSPLambda(this, [this]() { return GetEditorOnlyReferenceFilterType() == EEditorOnlyReferenceFilterType::Game; }),
+		FIsActionButtonVisible::CreateLambda([this] { return bShowShowReferencesOptions; }));
+
+	ReferenceViewerActions->MapAction(
+		FAssetManagerEditorCommands::Get().EditorOnlyReferenceFilterTypePropagation,
+		FExecuteAction::CreateSPLambda(this, [this]() { OnEditorOnlyReferenceFilterTypeChanged(EEditorOnlyReferenceFilterType::Propagation); }),
+		FCanExecuteAction(),
+		FIsActionChecked::CreateSPLambda(this, [this]() { return GetEditorOnlyReferenceFilterType() == EEditorOnlyReferenceFilterType::Propagation; }),
+		FIsActionButtonVisible::CreateLambda([this] { return bShowShowReferencesOptions; }));
+
+	ReferenceViewerActions->MapAction(
+		FAssetManagerEditorCommands::Get().EditorOnlyReferenceFilterTypeEditorOnly,
+		FExecuteAction::CreateSPLambda(this, [this]() { OnEditorOnlyReferenceFilterTypeChanged(EEditorOnlyReferenceFilterType::EditorOnly); }),
+		FCanExecuteAction(),
+		FIsActionChecked::CreateSPLambda(this, [this]() { return GetEditorOnlyReferenceFilterType() == EEditorOnlyReferenceFilterType::EditorOnly; }),
 		FIsActionButtonVisible::CreateLambda([this] { return bShowShowReferencesOptions; }));
 
 	ReferenceViewerActions->MapAction(
@@ -2468,7 +2481,12 @@ TSharedRef<SWidget> SReferenceViewer::GetShowMenuContent()
 	MenuBuilder.BeginSection("ReferenceTypes", LOCTEXT("ReferenceTypes", "Reference Types"));
 	MenuBuilder.AddMenuEntry(FAssetManagerEditorCommands::Get().ShowSoftReferences);
 	MenuBuilder.AddMenuEntry(FAssetManagerEditorCommands::Get().ShowHardReferences);
-	MenuBuilder.AddMenuEntry(FAssetManagerEditorCommands::Get().ShowEditorOnlyReferences);
+	MenuBuilder.EndSection();
+
+	MenuBuilder.BeginSection("EditorOnlyReferenceTypes", LOCTEXT("EditorOnlyReferenceTypes", "Editor Only Reference Types"));
+	MenuBuilder.AddMenuEntry(FAssetManagerEditorCommands::Get().EditorOnlyReferenceFilterTypeGame);
+	MenuBuilder.AddMenuEntry(FAssetManagerEditorCommands::Get().EditorOnlyReferenceFilterTypePropagation);
+	MenuBuilder.AddMenuEntry(FAssetManagerEditorCommands::Get().EditorOnlyReferenceFilterTypeEditorOnly);
 	MenuBuilder.EndSection();
 
 	MenuBuilder.BeginSection("Assets", LOCTEXT("Assets", "Assets"));
diff --git a/Engine/Plugins/Editor/AssetManagerEditor/Source/AssetManagerEditor/Private/ReferenceViewer/SReferenceViewer.h b/Engine/Plugins/Editor/AssetManagerEditor/Source/AssetManagerEditor/Private/ReferenceViewer/SReferenceViewer.h
--- a/Engine/Plugins/Editor/AssetManagerEditor/Source/AssetManagerEditor/Private/ReferenceViewer/SReferenceViewer.h
+++ b/Engine/Plugins/Editor/AssetManagerEditor/Source/AssetManagerEditor/Private/ReferenceViewer/SReferenceViewer.h
@@ -6,6 +6,7 @@
 #include "AssetRegistry/AssetData.h"
 #include "HistoryManager.h"
 #include "CollectionManagerTypes.h"
+#include "ReferenceViewer/ReferenceViewerSettings.h"
 
 class FUICommandList;
 class SComboButton;
@@ -144,8 +145,8 @@ private:
 	bool IsShowSoftReferencesChecked() const;
 	void OnShowHardReferencesChanged();
 	bool IsShowHardReferencesChecked() const;
-	void OnShowEditorOnlyReferencesChanged();
-	bool IsShowEditorOnlyReferencesChecked() const;
+	void OnEditorOnlyReferenceFilterTypeChanged(EEditorOnlyReferenceFilterType Value);
+	EEditorOnlyReferenceFilterType GetEditorOnlyReferenceFilterType() const;
 
 	void OnShowFilteredPackagesOnlyChanged();
 	bool IsShowFilteredPackagesOnlyChecked() const;
diff --git a/Engine/Plugins/Editor/AssetManagerEditor/Source/AssetManagerEditor/Public/ReferenceViewer/ReferenceViewerSettings.h b/Engine/Plugins/Editor/AssetManagerEditor/Source/AssetManagerEditor/Public/ReferenceViewer/ReferenceViewerSettings.h
--- a/Engine/Plugins/Editor/AssetManagerEditor/Source/AssetManagerEditor/Public/ReferenceViewer/ReferenceViewerSettings.h
+++ b/Engine/Plugins/Editor/AssetManagerEditor/Source/AssetManagerEditor/Public/ReferenceViewer/ReferenceViewerSettings.h
@@ -32,6 +32,15 @@ struct FilterState
 	bool bIsEnabled;
 };
 
+/** @see tooltips in FAssetManagerEditorCommands::RegisterCommands. */
+UENUM()
+enum class EEditorOnlyReferenceFilterType
+{
+	Game,
+	Propagation,
+	EditorOnly,
+};
+
 UCLASS(config=EditorPerProjectUserSettings)
 class UReferenceViewerSettings : public UObject
 {
@@ -74,8 +83,13 @@ public:
 	bool IsShowHardReferences() const;
 	void SetShowHardReferencesEnabled(bool bNewEnabled);
 
-	bool IsShowEditorOnlyReferences() const;
-	void SetShowEditorOnlyReferencesEnabled(bool bNewEnabled);
+	UE_DEPRECATED(5.5, "Use GetEditorOnlyReferenceFilterType.")
+	bool IsShowEditorOnlyReferences() const { return true; }
+	UE_DEPRECATED(5.5, "Use SetEditorOnlyReferenceFilterType.")
+	void SetShowEditorOnlyReferencesEnabled(bool bNewEnabled) {}
+
+	EEditorOnlyReferenceFilterType GetEditorOnlyReferenceFilterType() const;
+	void SetEditorOnlyReferenceFilterType(EEditorOnlyReferenceFilterType Value);
 
 	bool IsShowManagementReferences() const;
 	void SetShowManagementReferencesEnabled(bool bNewEnabled);
@@ -159,10 +173,13 @@ private:
 	UPROPERTY(config)
 	bool bIsShowHardReferences;
 	
-	/* Show/Hide EditorOnly References */
-	UPROPERTY(config)
+	UE_DEPRECATED(5.5, "Use EditorOnlyReferenceFilterType.")
+	UPROPERTY(config, meta = (DeprecatedProperty, DeprecationMessage = "Use EditorOnlyReferenceFilterType."))
 	bool bIsShowEditorOnlyReferences;
-	
+
+	UPROPERTY(config)
+	EEditorOnlyReferenceFilterType EditorOnlyReferenceFilterType;
+
 	/* Show/Hide Management Assets (i.e. PrimaryAssetIds) */
 	UPROPERTY(config)
 	bool bIsShowManagementReferences;
diff --git a/Engine/Source/Runtime/AssetRegistry/Internal/AssetRegistry/PackageReader.h b/Engine/Source/Runtime/AssetRegistry/Internal/AssetRegistry/PackageReader.h
--- a/Engine/Source/Runtime/AssetRegistry/Internal/AssetRegistry/PackageReader.h
+++ b/Engine/Source/Runtime/AssetRegistry/Internal/AssetRegistry/PackageReader.h
@@ -12,6 +12,7 @@ struct FAssetData;
 class FAssetPackageData;
 class FLinkerTables;
 class FPackageDependencyData;
+namespace UE::AssetRegistry { enum class EExtraDependencyFlags : uint32; }
 struct FGatherableTextData;
 struct FObjectFullNameAndThumbnail;
 
@@ -145,11 +146,12 @@ private:
 	bool SerializeEditorOnlyFlags(TBitArray<>& OutImportUsedInGame, TBitArray<>& OutSoftPackageUsedInGame);
 	bool SerializeSearchableNamesMap(FLinkerTables& OutSearchableNames);
 	bool SerializeAssetRegistryDependencyData(TBitArray<>& OutImportUsedInGame, TBitArray<>& OutSoftPackageUsedInGame,
-		const TArray<FObjectImport>& InImportMap, const TArray<FName>& InSoftPackageReferenceList);
+		TArray<TPair<FName, UE::AssetRegistry::EExtraDependencyFlags>>& OutExtraPackageDependencies);
 	bool SerializePackageTrailer(FAssetPackageData& PackageData);
 
 	void ApplyRelocationToImportMapAndSoftPackageReferenceList(FStringView LoadedPackageName,
-		TArray<FName>& OutSoftPackageReferenceList);
+		TArray<FName>& InOutSoftPackageReferenceList, 
+		TArray<TPair<FName, UE::AssetRegistry::EExtraDependencyFlags>>& InOutExtraPackageDependencies);
 	static void ConvertLinkerTableToPaths(FName PackageName, TArray<FObjectExport>& ExportMap,
 		TArray<FObjectImport>& ImportMap, TArray<FSoftObjectPath>& OutExports, TArray<FSoftObjectPath>& OutImports);
 
@@ -162,7 +164,7 @@ private:
 	 * Loader is the interface used to read the bytes from the package's repository. All interpretation of the bytes is
 	 * done by serializing into *this, which is also an FArchive.
 	 */
-	FArchive* Loader;
+	FArchive* Loader = nullptr;
 	FPackageFileSummary PackageFileSummary;
 	TArray<FName> NameMap;
 	TArray<FObjectImport> ImportMap;
@@ -172,9 +174,11 @@ private:
 	TArray<FSoftObjectPath> SoftObjectPathMap;
 	TArray<FGatherableTextData> GatherableTextDataMap;
 	TArray<FObjectFullNameAndThumbnail> ThumbnailMap;
-	int64 PackageFileSize;
-	int64 AssetRegistryDependencyDataOffset;
-	bool bLoaderOwner;
+	int64 PackageFileSize = 0;
+	int64 AssetRegistryDependencyDataOffset = INDEX_NONE;
+	// Defined as uint32 to avoid including AssetData.h in this header
+	uint32 AssetRegistryVersion = static_cast<uint32>(~0);
+	bool bLoaderOwner = false;
 };
 ENUM_CLASS_FLAGS(FPackageReader::EReadOptions);
 
diff --git a/Engine/Source/Runtime/AssetRegistry/Private/AssetDataGatherer.cpp b/Engine/Source/Runtime/AssetRegistry/Private/AssetDataGatherer.cpp
--- a/Engine/Source/Runtime/AssetRegistry/Private/AssetDataGatherer.cpp
+++ b/Engine/Source/Runtime/AssetRegistry/Private/AssetDataGatherer.cpp
@@ -48,7 +48,7 @@ namespace AssetDataGathererConstants
 	constexpr int32 MinSecondsToElapseBeforeCacheWrite = 30;
 	constexpr int32 MinAssetReadsBeforeCacheWrite = 100000;
 	constexpr int32 CacheShardAssetCount = 250000;
-	static constexpr uint32 CacheSerializationMagic = 0xC3A9D89B; // Versioning and integrity checking
+	static constexpr uint32 CacheSerializationMagic = 0x2779E7E0; // Versioning and integrity checking
 	static constexpr uint64 CurrentVersion = FAssetRegistryVersion::LatestVersion | (uint64(CacheSerializationMagic) << 32);
 	static const FTimespan CachePruneAge = FTimespan(7, 0, 0, 0);
 	/**
diff --git a/Engine/Source/Runtime/AssetRegistry/Private/AssetRegistry.cpp b/Engine/Source/Runtime/AssetRegistry/Private/AssetRegistry.cpp
--- a/Engine/Source/Runtime/AssetRegistry/Private/AssetRegistry.cpp
+++ b/Engine/Source/Runtime/AssetRegistry/Private/AssetRegistry.cpp
@@ -8600,8 +8600,8 @@ void FAssetRegistryImpl::SetManageReferences(const TMultiMap<FAssetIdentifier, F
 
 					SourceNode->IterateOverDependencies([&IterateFunction, SourceNode](FDependsNode* TargetNode, EDependencyCategory DependencyCategory, EDependencyProperty DependencyProperties, bool bDuplicate)
 						{
-							// Skip editor-only properties
-							if (!!(DependencyProperties & EDependencyProperty::Game))
+							// Skip editor-only, non-build dependencies. Propagate only through used-in-game or build dependencies.
+							if (EnumHasAnyFlags(DependencyProperties, EDependencyProperty::Game | EDependencyProperty::Build))
 							{
 								IterateFunction(SourceNode, TargetNode, DependencyCategory, DependencyProperties);
 							}
diff --git a/Engine/Source/Runtime/AssetRegistry/Private/DependsNode.cpp b/Engine/Source/Runtime/AssetRegistry/Private/DependsNode.cpp
--- a/Engine/Source/Runtime/AssetRegistry/Private/DependsNode.cpp
+++ b/Engine/Source/Runtime/AssetRegistry/Private/DependsNode.cpp
@@ -66,11 +66,31 @@ void IterateDependencyList(const FDependsNode::FIterateDependenciesCallback& InC
 			for (uint32 DependencyFlagBits : DependsNodeFlagsSet)
 			{
 				EDependencyProperty DependencyProperties = ByteToProperties(static_cast<uint8>(DependencyFlagBits));
-				if (EnumHasAllFlags(DependencyProperties, RequiredProperties)
-					&& !EnumHasAnyFlags(DependencyProperties, ExcludedProperties))
+				if (!EnumHasAllFlags(DependencyProperties, RequiredProperties))
 				{
-					InCallback(DependsNode, ListCategory, DependencyProperties, bDuplicate);
+					continue;
 				}
+				if (EnumHasAnyFlags(DependencyProperties, ExcludedProperties))
+				{
+					continue;
+				}
+				bool bPassesRequiredUnions = true;
+				for (EDependencyProperty RequiredUnion : SearchFlags.RequiredUnions)
+				{
+					EDependencyProperty RequiredUnionProperty = RequiredUnion & CategoryMask;
+					if (RequiredUnionProperty != EDependencyProperty::None &&
+						!EnumHasAnyFlags(DependencyProperties, RequiredUnionProperty))
+					{
+						bPassesRequiredUnions = false;
+						break;
+					}
+				}
+				if (!bPassesRequiredUnions)
+				{
+					continue;
+				}
+
+				InCallback(DependsNode, ListCategory, DependencyProperties, bDuplicate);
 				bDuplicate = true;
 			}
 		}
@@ -113,11 +133,30 @@ void IterateDependencyList(const FDependsNode::FIterateDependenciesCallback& InC
 			for (uint32 DependencyFlagBits : DependsNodeFlagsSet)
 			{
 				EDependencyProperty DependencyProperties = ByteToProperties(static_cast<uint8>(DependencyFlagBits));
-				if (EnumHasAllFlags(DependencyProperties, RequiredProperties)
-					&& !EnumHasAnyFlags(DependencyProperties, ExcludedProperties))
+				if (!EnumHasAllFlags(DependencyProperties, RequiredProperties))
+				{
+					continue;
+				}
+				if (EnumHasAnyFlags(DependencyProperties, ExcludedProperties))
+				{
+					continue;
+				}
+				bool bPassesRequiredUnions = true;
+				for (EDependencyProperty RequiredUnion : SearchFlags.RequiredUnions)
+				{
+					EDependencyProperty RequiredUnionProperty = RequiredUnion & CategoryMask;
+					if (RequiredUnionProperty != EDependencyProperty::None &&
+						!EnumHasAnyFlags(DependencyProperties, RequiredUnionProperty))
+					{
+						bPassesRequiredUnions = false;
+						break;
+					}
+				}
+				if (!bPassesRequiredUnions)
 				{
-					InCallback(DependsNode, ListCategory, DependencyProperties, bDuplicate);
+					continue;
 				}
+				InCallback(DependsNode, ListCategory, DependencyProperties, bDuplicate);
 				bDuplicate = true;
 			}
 		}
@@ -232,7 +271,8 @@ void FDependsNode::GetReferencers(TArray<FDependsNode*>& OutReferencers,
 		// If type specified, filter
 		if (Category != UE::AssetRegistry::EDependencyCategory::All
 			|| Flags.Required != UE::AssetRegistry::EDependencyProperty::None
-			|| Flags.Excluded != UE::AssetRegistry::EDependencyProperty::None)
+			|| Flags.Excluded != UE::AssetRegistry::EDependencyProperty::None
+			|| !Flags.RequiredUnions.IsEmpty())
 		{
 			Referencer->IterateOverDependencies([&bShouldAdd](const FDependsNode* InDependency,
 				UE::AssetRegistry::EDependencyCategory InCategory,
diff --git a/Engine/Source/Runtime/AssetRegistry/Private/DependsNode.h b/Engine/Source/Runtime/AssetRegistry/Private/DependsNode.h
--- a/Engine/Source/Runtime/AssetRegistry/Private/DependsNode.h
+++ b/Engine/Source/Runtime/AssetRegistry/Private/DependsNode.h
@@ -89,10 +89,29 @@ public:
 	/** Returns amount of memory used by the arrays */
 	SIZE_T GetAllocatedSize(void) const;
 
-	typedef TUniqueFunction<void(FDependsNode * Dependency,
+	typedef TUniqueFunction<void(
+		/** The other node the source node has a link to. */
+		FDependsNode* Dependency,
+		/**
+		 * The category of the dependency; dependencies are divided into high level categories,
+		 * @see EDependencyCategory.
+		 */
 		UE::AssetRegistry::EDependencyCategory Category,
+		/**
+		 * The properties of the dependency within its category. Each category has properties for dependencies in that category.
+		 * @see EDependencyProperty.
+		 */
 		UE::AssetRegistry::EDependencyProperty Properties,
-		bool bDuplicate)> FIterateDependenciesCallback;
+		/**
+		 * The source node may have multiple links to the targetnode, in different categories or with different
+		 * property combinations within the category. e.g. A package might have a Soft Game reference to another
+		 * package, but also a hard EditorOnly reference to that same package. When this occurs, and multiple links to
+		 * to the same Dependency are reported to an FIterateDependenciesCallback, all of the links to the same 
+		 * node are iterated consecutively, and bDuplicate=true for each of the reports after the first.
+		 * For the first or only occurrence of a Dependency in the iteration, bDuplicate=false.
+		 */
+		bool bDuplicate
+		)> FIterateDependenciesCallback;
 	/**
 	 * Iterate over all the dependencies of this node, optionally filtered by the target node, category and query,
 	 * and call the supplied lambda parameter on the record.
diff --git a/Engine/Source/Runtime/AssetRegistry/Private/PackageDependencyData.cpp b/Engine/Source/Runtime/AssetRegistry/Private/PackageDependencyData.cpp
--- a/Engine/Source/Runtime/AssetRegistry/Private/PackageDependencyData.cpp
+++ b/Engine/Source/Runtime/AssetRegistry/Private/PackageDependencyData.cpp
@@ -45,7 +45,8 @@ struct FSortPackageDependency
 
 void FPackageDependencyData::LoadDependenciesFromPackageHeader(FName SourcePackageName, TConstArrayView<FObjectImport> ImportMap,
 	TArray<FName>& SoftPackageReferenceList, TMap<FPackageIndex, TArray<FName>>& SearchableNames,
-	TBitArray<>& ImportUsedInGame, TBitArray<>& SoftPackageUsedInGame)
+	TBitArray<>& ImportUsedInGame, TBitArray<>& SoftPackageUsedInGame,
+	TArray<TPair<FName, UE::AssetRegistry::EExtraDependencyFlags>>& ExtraPackageDependencies)
 {
 	using namespace UE::AssetRegistry;
 
@@ -54,7 +55,7 @@ void FPackageDependencyData::LoadDependenciesFromPackageHeader(FName SourcePacka
 	for (int32 ImportIdx = 0; ImportIdx < ImportMap.Num(); ++ImportIdx)
 	{
 		FName DependencyPackageName = GetImportPackageName(ImportMap, ImportIdx);
-		EDependencyProperty DependencyProperty = EDependencyProperty::Build | EDependencyProperty::Hard;
+		EDependencyProperty DependencyProperty = EDependencyProperty::Hard;
 		DependencyProperty |= ImportUsedInGame[ImportIdx] ? EDependencyProperty::Game : EDependencyProperty::None;
 		PackageDependencies.Add({ DependencyPackageName, DependencyProperty });
 	}
@@ -70,13 +71,22 @@ void FPackageDependencyData::LoadDependenciesFromPackageHeader(FName SourcePacka
 	{
 		FName DependencyPackageName = SoftPackageReferenceList[SoftPackageIdx];
 		FAssetIdentifier AssetId(DependencyPackageName);
-		EDependencyProperty DependencyProperty = UE::AssetRegistry::EDependencyProperty::Build; // !EDependencyProperty::Hard
+		EDependencyProperty DependencyProperty = EDependencyProperty::None;
 		DependencyProperty |= (SoftPackageUsedInGame[SoftPackageIdx] ? EDependencyProperty::Game : EDependencyProperty::None);
 
 		// Don't need to remove duplicates here because SavePackage only writes unique elements into SoftPackageReferenceList
 		PackageDependencies.Add({ DependencyPackageName, DependencyProperty });
 	}
 
+	for (const TPair<FName, EExtraDependencyFlags>& Pair : ExtraPackageDependencies)
+	{
+		EDependencyProperty DependencyProperty = EDependencyProperty::None;
+		DependencyProperty |= EnumHasAnyFlags(Pair.Value, EExtraDependencyFlags::Build)
+			? EDependencyProperty::Build : EDependencyProperty::None;
+		// Don't need to remove duplicates here because SavePackage only writes unique elements into PackageBuildDependencies
+		PackageDependencies.Add({ Pair.Key, DependencyProperty });
+	}
+
 	SearchableNameDependencies.Reset(SearchableNames.Num());
 	for (const TPair<FPackageIndex, TArray<FName>>& SearchableNameList : SearchableNames)
 	{
diff --git a/Engine/Source/Runtime/AssetRegistry/Private/PackageDependencyData.h b/Engine/Source/Runtime/AssetRegistry/Private/PackageDependencyData.h
--- a/Engine/Source/Runtime/AssetRegistry/Private/PackageDependencyData.h
+++ b/Engine/Source/Runtime/AssetRegistry/Private/PackageDependencyData.h
@@ -69,7 +69,8 @@ public:
 
 	void LoadDependenciesFromPackageHeader(FName PackageName, TConstArrayView<FObjectImport> ImportMap,
 		TArray<FName>& SoftPackageReferenceList, TMap<FPackageIndex, TArray<FName>>& SearchableNames,
-		TBitArray<>& ImportUsedInGame, TBitArray<>& SoftPackageUsedInGame);
+		TBitArray<>& ImportUsedInGame, TBitArray<>& SoftPackageUsedInGame,
+		TArray<TPair<FName, UE::AssetRegistry::EExtraDependencyFlags>>& ExtraPackageDependencies);
 
 	/** Returns the amount of memory allocated by this container, not including sizeof(*this). */
 	SIZE_T GetAllocatedSize() const
diff --git a/Engine/Source/Runtime/AssetRegistry/Private/PackageReader.cpp b/Engine/Source/Runtime/AssetRegistry/Private/PackageReader.cpp
--- a/Engine/Source/Runtime/AssetRegistry/Private/PackageReader.cpp
+++ b/Engine/Source/Runtime/AssetRegistry/Private/PackageReader.cpp
@@ -85,10 +85,6 @@ static void ApplyRelocationToTagsAndValues(FAssetDataTagMap& TagsAndValues, UE::
 }
 
 FPackageReader::FPackageReader()
-	: Loader(nullptr)
-	, PackageFileSize(0)
-	, AssetRegistryDependencyDataOffset(INDEX_NONE)
-	, bLoaderOwner(false)
 {
 	this->SetIsLoading(true);
 	this->SetIsPersistent(true);
@@ -566,29 +562,28 @@ bool FPackageReader::ReadLinkerObjects(TMap<FSoftObjectPath, FObjectData>& OutEx
 	return true;
 }
 
-bool FPackageReader::SerializeAssetRegistryDependencyData(TBitArray<>& OutImportUsedInGame, TBitArray<>& OutSoftPackageUsedInGame,
-	const TArray<FObjectImport>& InImportMap, const TArray<FName>& InSoftPackageReferenceList)
+bool FPackageReader::SerializeAssetRegistryDependencyData(TBitArray<>& OutImportUsedInGame,
+	TBitArray<>& OutSoftPackageUsedInGame,
+	TArray<TPair<FName, UE::AssetRegistry::EExtraDependencyFlags>>& OutExtraPackageDependencies)
 {
-	if (AssetRegistryDependencyDataOffset == INDEX_NONE)
-	{
-		// For old package versions that did not write out the dependency flags, set default values of the flags
-		OutImportUsedInGame.Init(true, InImportMap.Num());
-		OutSoftPackageUsedInGame.Init(true, InSoftPackageReferenceList.Num());
-		return true;
-	}
+	UE::AssetRegistry::FReadPackageDataDependenciesArgs Args;
+	Args.BinaryNameAwareArchive = this;
+	Args.AssetRegistryDependencyDataOffset = AssetRegistryDependencyDataOffset;
+	Args.NumImports = ImportMap.Num();
+	Args.NumSoftPackageReferences = SoftPackageReferenceList.Num();
+	Args.PackageVersion = PackageFileSummary.GetFileVersionUE();
 
-	if (!StartSerializeSection(AssetRegistryDependencyDataOffset))
-	{
-		return false;
-	}
+	ClearError();
+	Loader->ClearError();
 
-	if (!UE::AssetRegistry::ReadPackageDataDependencies(*this, OutImportUsedInGame, OutSoftPackageUsedInGame) ||
-		OutImportUsedInGame.Num() != InImportMap.Num() ||
-		OutSoftPackageUsedInGame.Num() != InSoftPackageReferenceList.Num())
+	if (!UE::AssetRegistry::ReadPackageDataDependencies(Args))
 	{
 		UE_PACKAGEREADER_CORRUPTPACKAGE_WARNING("SerializeAssetRegistryDependencyData", PackageFilename);
 		return false;
 	}
+	OutImportUsedInGame = MoveTemp(Args.ImportUsedInGame);
+	OutSoftPackageUsedInGame = MoveTemp(Args.SoftPackageUsedInGame);
+	OutExtraPackageDependencies = MoveTemp(Args.ExtraPackageDependencies);
 	return true;
 }
 
@@ -612,15 +607,28 @@ bool FPackageReader::SerializePackageTrailer(FAssetPackageData& PackageData)
 	return true;
 }
 
-void FPackageReader::ApplyRelocationToImportMapAndSoftPackageReferenceList(FStringView LoadedPackageName, TArray<FName>& OutSoftPackageReferenceList)
+void FPackageReader::ApplyRelocationToImportMapAndSoftPackageReferenceList(FStringView LoadedPackageName, 
+	TArray<FName>& InOutSoftPackageReferenceList,
+	TArray<TPair<FName, UE::AssetRegistry::EExtraDependencyFlags>>& InOutExtraPackageDependencies)
 {
 #if WITH_EDITOR
 	UE::Package::Relocation::Private::FPackageRelocationContext RelocationArgs;
 	if (UE::Package::Relocation::Private::ShouldApplyRelocation(PackageFileSummary, LoadedPackageName, RelocationArgs))
 	{
-		UE_LOG(LogPackageRelocation, Verbose, TEXT("Detected relocated package (%.*s). The package was saved as (%s)."), LoadedPackageName.Len(), LoadedPackageName.GetData(), *PackageFileSummary.PackageName);
+		UE_LOG(LogPackageRelocation, Verbose, TEXT("Detected relocated package (%.*s). The package was saved as (%s)."),
+			LoadedPackageName.Len(), LoadedPackageName.GetData(), *PackageFileSummary.PackageName);
 		UE::Package::Relocation::Private::ApplyRelocationToObjectImportMap(RelocationArgs, ImportMap);
-		UE::Package::Relocation::Private::ApplyRelocationToNameArray(RelocationArgs, OutSoftPackageReferenceList);
+		UE::Package::Relocation::Private::ApplyRelocationToNameArray(RelocationArgs, InOutSoftPackageReferenceList);
+		for (TPair<FName, UE::AssetRegistry::EExtraDependencyFlags>& Pair : InOutExtraPackageDependencies)
+		{
+			FNameBuilder Package(Pair.Key);
+			FNameBuilder RelocatedPackageName;
+			if (UE::Package::Relocation::Private::TryRelocateReference(
+				RelocationArgs, Package.ToView(), RelocatedPackageName))
+			{
+				Pair.Key = *RelocatedPackageName;
+			}
+		}
 	}
 #endif
 }
@@ -773,16 +781,18 @@ bool FPackageReader::ReadDependencyData(FPackageDependencyData& OutDependencyDat
 
 		TBitArray<> ImportUsedInGame;
 		TBitArray<> SoftPackageUsedInGame;
-		if (!SerializeAssetRegistryDependencyData(ImportUsedInGame, SoftPackageUsedInGame, ImportMap,
-			SoftPackageReferenceList))
+		TArray<TPair<FName, UE::AssetRegistry::EExtraDependencyFlags>> ExtraPackageDependencies;
+		if (!SerializeAssetRegistryDependencyData(ImportUsedInGame, SoftPackageUsedInGame, ExtraPackageDependencies))
 		{
 			return false;
 		}
 
-		ApplyRelocationToImportMapAndSoftPackageReferenceList(PackageNameString, SoftPackageReferenceList);
+		ApplyRelocationToImportMapAndSoftPackageReferenceList(PackageNameString, SoftPackageReferenceList,
+			ExtraPackageDependencies);
 
-		OutDependencyData.LoadDependenciesFromPackageHeader(OutDependencyData.PackageName, ImportMap, SoftPackageReferenceList,
-			SearchableNames.SearchableNamesMap, ImportUsedInGame, SoftPackageUsedInGame);
+		OutDependencyData.LoadDependenciesFromPackageHeader(OutDependencyData.PackageName, ImportMap,
+			SoftPackageReferenceList, SearchableNames.SearchableNamesMap, ImportUsedInGame, SoftPackageUsedInGame,
+			ExtraPackageDependencies);
 	}
 
 	return true;
@@ -1312,18 +1322,23 @@ bool FPackageReader::SerializeEditorOnlyFlags(TBitArray<>& OutImportUsedInGame,
 		return true;
 	}
 
-	if (!StartSerializeSection(AssetRegistryDependencyDataOffset))
-	{
-		return false;
-	}
+	ClearError();
+	Loader->ClearError();
 
-	if (!UE::AssetRegistry::ReadPackageDataDependencies(*this, OutImportUsedInGame, OutSoftPackageUsedInGame) ||
-		OutImportUsedInGame.Num() != ImportMap.Num() ||
-		OutSoftPackageUsedInGame.Num() != SoftPackageReferenceList.Num())
+	UE::AssetRegistry::FReadPackageDataDependenciesArgs Args;
+	Args.BinaryNameAwareArchive = this;
+	Args.AssetRegistryDependencyDataOffset = AssetRegistryDependencyDataOffset;
+	Args.NumImports = ImportMap.Num();
+	Args.NumSoftPackageReferences = SoftPackageReferenceList.Num();
+	Args.PackageVersion = PackageFileSummary.GetFileVersionUE();
+
+	if (!UE::AssetRegistry::ReadPackageDataDependencies(Args))
 	{
 		UE_PACKAGEREADER_CORRUPTPACKAGE_WARNING("SerializeAssetRegistryDependencyData", PackageFilename);
 		return false;
 	}
+	OutImportUsedInGame = MoveTemp(Args.ImportUsedInGame);
+	OutSoftPackageUsedInGame = MoveTemp(Args.SoftPackageUsedInGame);
 
 	return true;
 }
@@ -1826,12 +1841,54 @@ namespace UE::AssetRegistry
 		return true;
 	}
 
-	// See the corresponding WriteAssetRegistryPackageData defined in SavePackageUtilities.cpp in CoreUObject module
 	bool ReadPackageDataDependencies(FArchive& BinaryArchive, TBitArray<>& OutImportUsedInGame, TBitArray<>& OutSoftPackageUsedInGame)
 	{
-		BinaryArchive << OutImportUsedInGame;
-		BinaryArchive << OutSoftPackageUsedInGame;
-		return !BinaryArchive.IsError();
+		UE_LOG(LogAssetRegistry, Error,
+			TEXT("This version of ReadPackageDataDependencies is no longer supported since it does not include enough information to know the package's AssetRegistryVersion. Read will be marked as failed."));
+		return false;
+	}
+
+	// See the corresponding WriteAssetRegistryPackageData defined in SavePackageUtilities.cpp in CoreUObject module
+	bool ReadPackageDataDependencies(FReadPackageDataDependenciesArgs& Args)
+	{
+		// Always set the output AssetRegistryVersion; in an error case it indicates to the caller whether the error
+		// was caused by a too-high version.
+		if (Args.AssetRegistryDependencyDataOffset == INDEX_NONE)
+		{
+			// For old package versions that did not write out the dependency flags, set default values of the flags
+			Args.ImportUsedInGame.Init(true, Args.NumImports);
+			Args.SoftPackageUsedInGame.Init(true, Args.NumSoftPackageReferences);
+			Args.ExtraPackageDependencies.Empty();
+			return true;
+		}
+
+		FArchive& Ar = *Args.BinaryNameAwareArchive;
+		Ar.Seek(Args.AssetRegistryDependencyDataOffset);
+		if (Ar.IsError())
+		{
+			return false;
+		}
+
+		Ar << Args.ImportUsedInGame;
+		Ar << Args.SoftPackageUsedInGame;
+		if (Args.PackageVersion >= EUnrealEngineObjectUE5Version::ASSETREGISTRY_PACKAGEBUILDDEPENDENCIES)
+		{
+			// Reinterpret ExtraPackageDependencies as an Array with integer values so we can serialize the 
+			// EExtraDependencyFlags as integers.
+			TArray<TPair<FName, uint32>>& ExtraPackageDependenciesAsIntegers = 
+				reinterpret_cast<TArray<TPair<FName, uint32>>&>(Args.ExtraPackageDependencies);
+			Ar << ExtraPackageDependenciesAsIntegers;
+		}
+		else
+		{
+			Args.ExtraPackageDependencies.Empty();
+		}
+		if (Args.ImportUsedInGame.Num() != Args.NumImports ||
+			Args.SoftPackageUsedInGame.Num() != Args.NumSoftPackageReferences)
+		{
+			return false;
+		}
+		return !Ar.IsError();
 	}
 }
 
diff --git a/Engine/Source/Runtime/AssetRegistry/Public/AssetRegistry/IAssetRegistry.h b/Engine/Source/Runtime/AssetRegistry/Public/AssetRegistry/IAssetRegistry.h
--- a/Engine/Source/Runtime/AssetRegistry/Public/AssetRegistry/IAssetRegistry.h
+++ b/Engine/Source/Runtime/AssetRegistry/Public/AssetRegistry/IAssetRegistry.h
@@ -1009,9 +1009,7 @@ protected:
 	virtual bool SetPrimaryAssetIdForObjectPath(const FSoftObjectPath& ObjectPath, FPrimaryAssetId PrimaryAssetId) = 0;
 };
 
-namespace UE
-{
-namespace AssetRegistry
+namespace UE::AssetRegistry
 {
 	enum EReadPackageDataMainErrorCode
 	{
@@ -1051,6 +1049,22 @@ namespace AssetRegistry
 	ASSETREGISTRY_API bool ReadPackageDataMain(FArchive& BinaryArchive, const FString& PackageName, const FPackageFileSummary& PackageFileSummary,
 		int64& OutDependencyDataOffset, TArray<FAssetData*>& OutAssetDataList, EReadPackageDataMainErrorCode& OutError,
 		const TArray<FObjectImport>* InImports = nullptr, const TArray<FObjectExport>* InExports = nullptr);
+	struct FReadPackageDataDependenciesArgs
+	{
+		// Required inputs, must be initialized and non-null
+		FArchive* BinaryNameAwareArchive = nullptr;
+		int64 AssetRegistryDependencyDataOffset = -1;
+		int32 NumImports = -1;
+		int32 NumSoftPackageReferences = -1;
+		FPackageFileVersion PackageVersion;
+
+		// Outputs
+		TBitArray<> ImportUsedInGame;
+		TBitArray<> SoftPackageUsedInGame;
+		TArray<TPair<FName, EExtraDependencyFlags>> ExtraPackageDependencies;
+	};
+	ASSETREGISTRY_API bool ReadPackageDataDependencies(FReadPackageDataDependenciesArgs& Args);
+	UE_DEPRECATED(5.5, "Use version that takes FReadPackageDataArgs")
 	ASSETREGISTRY_API bool ReadPackageDataDependencies(FArchive& BinaryArchive, TBitArray<>& OutImportUsedInGame, TBitArray<>& OutSoftPackageUsedInGame);
 
 	/**
@@ -1124,8 +1138,7 @@ namespace AssetRegistry
 	extern ASSETREGISTRY_API const FName Stage_ChunkInstalledSizeFName;
 	extern ASSETREGISTRY_API const FName Stage_ChunkStreamingSizeFName;
 	extern ASSETREGISTRY_API const FName Stage_ChunkOptionalSizeFName;
-} // namespace AssetRegistry
-} // namespace UE
+} // namespace UE::AssetRegistry
 
 /** Returns the filename without filepath for the DevelopmentAssetRegistry written by the cooker. */
 ASSETREGISTRY_API const TCHAR* GetDevelopmentAssetRegistryFilename();
diff --git a/Engine/Source/Runtime/Core/Public/UObject/ObjectVersion.h b/Engine/Source/Runtime/Core/Public/UObject/ObjectVersion.h
--- a/Engine/Source/Runtime/Core/Public/UObject/ObjectVersion.h
+++ b/Engine/Source/Runtime/Core/Public/UObject/ObjectVersion.h
@@ -84,6 +84,9 @@ enum class EUnrealEngineObjectUE5Version : uint32
 	// Added property tag complete type name and serialization type
 	PROPERTY_TAG_COMPLETE_TYPE_NAME,
 
+	// Changed UE::AssetRegistry::WritePackageData to include PackageBuildDependencies
+	ASSETREGISTRY_PACKAGEBUILDDEPENDENCIES,
+
 	// -----<new versions can be added before this line>-------------------------------------------------
 	// - this needs to be the last line (see note below)
 	AUTOMATIC_VERSION_PLUS_ONE,
diff --git a/Engine/Source/Runtime/CoreUObject/Private/UObject/SavePackage/SaveContext.cpp b/Engine/Source/Runtime/CoreUObject/Private/UObject/SavePackage/SaveContext.cpp
--- a/Engine/Source/Runtime/CoreUObject/Private/UObject/SavePackage/SaveContext.cpp
+++ b/Engine/Source/Runtime/CoreUObject/Private/UObject/SavePackage/SaveContext.cpp
@@ -3,13 +3,12 @@
 #include "SaveContext.h"
 
 #include "Algo/Find.h"
+#include "Algo/Unique.h"
 #include "Cooker/CookDependency.h"
 #include "Misc/ConfigCacheIni.h"
 #include "Serialization/PackageWriter.h"
 #include "UObject/UObjectGlobals.h"
 
-
-
 TArray<ESaveRealm> FSaveContext::GetHarvestedRealmsToSave()
 {
 	TArray<ESaveRealm> HarvestedContextsToSave;
@@ -294,6 +293,50 @@ EObjectMark FSaveContext::GetExcludedObjectMarksForGameRealm(const ITargetPlatfo
 	}
 }
 
+void FSaveContext::UpdateEditorRealmPackageBuildDependencies()
+{
+	using namespace UE::Cook;
+
+	PackageBuildDependencies.Empty();
+
+	// PackageBuildDependencies are only recorded for non-cooked packages
+	if (IsCooking())
+	{
+		return;
+	}
+
+#if WITH_EDITOR
+	PackageBuildDependencies.Reserve(ObjectSaveContext.CookBuildDependencies.Num());
+	for (UE::Cook::FCookDependency& CookDependency : ObjectSaveContext.CookBuildDependencies)
+	{
+		FName PackageName = NAME_None;
+		switch (CookDependency.GetType())
+		{
+		case ECookDependency::Package: [[fallthrough]];
+		case ECookDependency::TransitiveBuild:
+			PackageName = CookDependency.GetPackageName();
+			break;
+		default:
+			break;
+		}
+		if (PackageName.IsNone())
+		{
+			continue;
+		}
+		PackageBuildDependencies.Add(PackageName);
+	}
+	PackageBuildDependencies.Sort(FNameLexicalLess());
+	PackageBuildDependencies.SetNum(Algo::Unique(PackageBuildDependencies));
+
+	FHarvestedRealm& HarvestedRealm = GetHarvestedRealm(ESaveRealm::Editor);
+	TSet<FNameEntryId>& NamesReferencedFromPackageHeader = HarvestedRealm.GetNamesReferencedFromPackageHeader();
+	for (FName PackageBuildDependency : PackageBuildDependencies)
+	{
+		NamesReferencedFromPackageHeader.Add(PackageBuildDependency.GetDisplayIndex());
+	}
+#endif
+}
+
 const TCHAR* LexToString(ESaveableStatus Status)
 {
 	static_assert(static_cast<int32>(ESaveableStatus::__Count) == 10);
diff --git a/Engine/Source/Runtime/CoreUObject/Private/UObject/SavePackage/SaveContext.h b/Engine/Source/Runtime/CoreUObject/Private/UObject/SavePackage/SaveContext.h
--- a/Engine/Source/Runtime/CoreUObject/Private/UObject/SavePackage/SaveContext.h
+++ b/Engine/Source/Runtime/CoreUObject/Private/UObject/SavePackage/SaveContext.h
@@ -904,6 +904,22 @@ public:
 		return GetHarvestedRealm(ESaveRealm::Game).GetSoftPackageReferenceList();
 	}
 
+	const TArray<FName>& GetPackageBuildDependencies(ESaveRealm SaveRealm = ESaveRealm::None)
+	{
+		SaveRealm = (SaveRealm == ESaveRealm::None) ? GetCurrentHarvestingRealm() : SaveRealm;
+		if (SaveRealm == ESaveRealm::Editor)
+		{
+			return PackageBuildDependencies;
+		}
+		return EmptyList;
+	}
+
+	/**
+	 * PackageBuildDependencies are copied from CookBuildDependencies, and they are only needed
+	 * for the Editor realm.
+	*/
+	void UpdateEditorRealmPackageBuildDependencies();
+
 	const TMap<TObjectPtr<UObject>, TArray<FName>>& GetSearchableNamesObjectMap() const
 	{
 		return GetHarvestedRealm().GetSearchableNamesObjectMap();
@@ -924,6 +940,11 @@ public:
 		return GetHarvestedRealm().GetNamesReferencedFromPackageHeader();
 	}
 
+	TSet<FNameEntryId>& GetNamesReferencedFromPackageHeader()
+	{
+		return GetHarvestedRealm().GetNamesReferencedFromPackageHeader();
+	}
+
 	const TSet<FSoftObjectPath>& GetSoftObjectPathList() const
 	{
 		return GetHarvestedRealm().GetSoftObjectPathList();
@@ -1186,7 +1207,7 @@ private:
 	ESaveableStatus GetSaveableStatusNoOuter(TObjectPtr<UObject> Obj,
 		UE::SavePackageUtilities::FObjectStatus& ObjectStatus) const;
 	static EObjectMark GetExcludedObjectMarksForGameRealm(const ITargetPlatform* TargetPlatform);
-		
+
 	friend class FPackageHarvester;
 
 	// Args
@@ -1244,6 +1265,11 @@ private:
 
 	// Cache of FObjectStatus for every object encountered during the save
 	TMap<TObjectPtr<UObject>, UE::SavePackageUtilities::FObjectStatus> ObjectStatusCache;
+
+	// List of package build dependencies reported from PreSave or Serialize functions
+	TArray<FName> PackageBuildDependencies;
+	// Empty list of FNames, used for functions that need to return a reference to an empty array.
+	TArray<FName> EmptyList;
 };
 
 const TCHAR* LexToString(ESaveableStatus Status);
diff --git a/Engine/Source/Runtime/CoreUObject/Private/UObject/SavePackage/SavePackageUtilities.cpp b/Engine/Source/Runtime/CoreUObject/Private/UObject/SavePackage/SavePackageUtilities.cpp
--- a/Engine/Source/Runtime/CoreUObject/Private/UObject/SavePackage/SavePackageUtilities.cpp
+++ b/Engine/Source/Runtime/CoreUObject/Private/UObject/SavePackage/SavePackageUtilities.cpp
@@ -2299,12 +2299,29 @@ void WritePackageData(FStructuredArchiveRecord& ParentRecord, bool bIsCooking, c
 	}
 }
 
-// See the corresponding ReadPackageDataMain and ReadPackageDataDependencies defined in PackageReader.cpp in AssetRegistry module
 void WritePackageData(FStructuredArchiveRecord& ParentRecord, FArchiveCookContext* CookContext, const UPackage* Package,
 	FLinkerSave* Linker, const TSet<TObjectPtr<UObject>>& ImportsUsedInGame, const TSet<FName>& SoftPackagesUsedInGame,
 	TArray<FAssetData>* OutAssetDatas, bool bProceduralSave)
 {
-	bProceduralSave = bProceduralSave || (CookContext != nullptr);
+	FWritePackageDataArgs Args;
+	Args.ParentRecord = &ParentRecord;
+	Args.Package = Package;
+	Args.Linker = Linker;
+	Args.ImportsUsedInGame = &ImportsUsedInGame;
+	Args.SoftPackagesUsedInGame = &SoftPackagesUsedInGame;
+	Args.bProceduralSave = bProceduralSave;
+	Args.CookContext = CookContext;
+	Args.OutAssetDatas = OutAssetDatas;
+	TArray<FName> PackageBuildDependencies;
+	Args.PackageBuildDependencies = &PackageBuildDependencies;
+	WritePackageData(Args);
+}
+
+// See the corresponding ReadPackageDataMain and ReadPackageDataDependencies defined in PackageReader.cpp in AssetRegistry module
+void WritePackageData(FWritePackageDataArgs& Args)
+{
+	Args.bProceduralSave = Args.bProceduralSave || (Args.CookContext != nullptr);
+	FLinkerSave* Linker = Args.Linker;
 	IAssetRegistryInterface* AssetRegistry = IAssetRegistryInterface::GetPtr();
 
 	// To avoid large patch sizes, we have frozen cooked package format at the format before VER_UE4_ASSETREGISTRY_DEPENDENCYFLAGS
@@ -2313,22 +2330,22 @@ void WritePackageData(FStructuredArchiveRecord& ParentRecord, FArchiveCookContex
 	bool bPreDependencyFormat = false;
 	bool bWriteAssetsToPackage = true;
 	// Editor saves do a full update, but procedural saves (including cook) do not
-	bool bFullUpdate = !bProceduralSave;
+	bool bFullUpdate = !Args.bProceduralSave;
 	FCookTagList* CookTagList = nullptr;
-	if (CookContext)
+	if (Args.CookContext)
 	{
 		bPreDependencyFormat = true;
 		bWriteAssetsToPackage = false;
-		CookTagList = CookContext->GetCookTagList();
+		CookTagList = Args.CookContext->GetCookTagList();
 	}	
 
 	// WritePackageData is currently only called if not bTextFormat; we rely on that to save offsets
-	FArchive& BinaryArchive = ParentRecord.GetUnderlyingArchive();
+	FArchive& BinaryArchive = Args.ParentRecord->GetUnderlyingArchive();
 	check(!BinaryArchive.IsTextFormat());
 
 	// Store the asset registry offset in the file and enter a record for the asset registry data
 	Linker->Summary.AssetRegistryDataOffset = (int32)BinaryArchive.Tell();
-	FStructuredArchiveRecord AssetRegistryRecord = ParentRecord.EnterField(TEXT("AssetRegistry")).EnterRecord();
+	FStructuredArchiveRecord AssetRegistryRecord = Args.ParentRecord->EnterField(TEXT("AssetRegistry")).EnterRecord();
 
 	// Offset to Dependencies
 	int64 OffsetToAssetRegistryDependencyDataOffset = INDEX_NONE;
@@ -2348,10 +2365,10 @@ void WritePackageData(FStructuredArchiveRecord& ParentRecord, FArchiveCookContex
 		if (Export.Object && Export.Object->IsAsset())
 		{
 #if WITH_EDITOR
-			if (CookContext)
+			if (Args.CookContext)
 			{
 				TArray<UObject*> AdditionalObjects;
-				Export.Object->GetAdditionalAssetDataObjectsForCook(*CookContext, AdditionalObjects);
+				Export.Object->GetAdditionalAssetDataObjectsForCook(*Args.CookContext, AdditionalObjects);
 				for (UObject* Object : AdditionalObjects)
 				{
 					if (Object->IsAsset())
@@ -2367,24 +2384,24 @@ void WritePackageData(FStructuredArchiveRecord& ParentRecord, FArchiveCookContex
 
 	int32 ObjectCountInPackage = bWriteAssetsToPackage ? AssetObjects.Num() : 0;
 	FStructuredArchive::FArray AssetArray = AssetRegistryRecord.EnterArray(TEXT("TagMap"), ObjectCountInPackage);
-	FString PackageName = Package->GetName();
+	FString PackageName = Args.Package->GetName();
 
 	for (int32 ObjectIdx = 0; ObjectIdx < AssetObjects.Num(); ++ObjectIdx)
 	{
 		const UObject* Object = AssetObjects[ObjectIdx];
 
 		// Exclude the package name in the object path, we just need to know the path relative to the package we are saving
-		FString ObjectPath = Object->GetPathName(Package);
+		FString ObjectPath = Object->GetPathName(Args.Package);
 		FString ObjectClassName = Object->GetClass()->GetPathName();
 
 		FAssetRegistryTagsContextData TagsContextData(Object, EAssetRegistryTagsCaller::SavePackage);
-		TagsContextData.bProceduralSave = bProceduralSave;
+		TagsContextData.bProceduralSave = Args.bProceduralSave;
 		TagsContextData.TargetPlatform = nullptr;
-		if (CookContext)
+		if (Args.CookContext)
 		{
-			TagsContextData.TargetPlatform = CookContext->GetTargetPlatform();
-			TagsContextData.CookType = CookContext->GetCookType();
-			TagsContextData.CookingDLC = CookContext->GetCookingDLC();
+			TagsContextData.TargetPlatform = Args.CookContext->GetTargetPlatform();
+			TagsContextData.CookType = Args.CookContext->GetCookType();
+			TagsContextData.CookingDLC = Args.CookContext->GetCookingDLC();
 			TagsContextData.bWantsCookTags = TagsContextData.CookType == UE::Cook::ECookType::ByTheBook;
 		}
 		TagsContextData.bFullUpdateRequested = bFullUpdate;
@@ -2440,7 +2457,7 @@ void WritePackageData(FStructuredArchiveRecord& ParentRecord, FArchiveCookContex
 			}
 		}
 
-		if (OutAssetDatas)
+		if (Args.OutAssetDatas)
 		{
 			FAssetDataTagMap TagsAndValues;
 			for (TPair<FName, UObject::FAssetRegistryTag>& TagPair : TagsContextData.Tags)
@@ -2464,8 +2481,8 @@ void WritePackageData(FStructuredArchiveRecord& ParentRecord, FArchiveCookContex
 				ObjectPath = PackageName + TEXT(".") + ObjectPath;
 			}
 
-			OutAssetDatas->Emplace(PackageName, ObjectPath, FTopLevelAssetPath(ObjectClassName),
-				MoveTemp(TagsAndValues), Package->GetChunkIDs(), Package->GetPackageFlags());
+			Args.OutAssetDatas->Emplace(PackageName, ObjectPath, FTopLevelAssetPath(ObjectClassName),
+				MoveTemp(TagsAndValues), Args.Package->GetChunkIDs(), Args.Package->GetPackageFlags());
 		}
 	}
 	if (bPreDependencyFormat)
@@ -2481,7 +2498,7 @@ void WritePackageData(FStructuredArchiveRecord& ParentRecord, FArchiveCookContex
 		BinaryArchive << AssetRegistryDependencyDataOffset;
 		BinaryArchive.Seek(AssetRegistryDependencyDataOffset);
 	}
-	FStructuredArchiveRecord DependencyDataRecord = ParentRecord.EnterField(TEXT("AssetRegistryDependencyData")).EnterRecord();
+	FStructuredArchiveRecord DependencyDataRecord = Args.ParentRecord->EnterField(TEXT("AssetRegistryDependencyData")).EnterRecord();
 
 	// Convert the IsUsedInGame sets into a bitarray with a value per import/softpackagereference
 	TBitArray<> ImportUsedInGameBits;
@@ -2489,17 +2506,32 @@ void WritePackageData(FStructuredArchiveRecord& ParentRecord, FArchiveCookContex
 	ImportUsedInGameBits.Reserve(Linker->ImportMap.Num());
 	for (int32 ImportIndex = 0; ImportIndex < Linker->ImportMap.Num(); ++ImportIndex)
 	{
-		ImportUsedInGameBits.Add(ImportsUsedInGame.Contains(Linker->ImportMap[ImportIndex].XObject));
+		ImportUsedInGameBits.Add(Args.ImportsUsedInGame->Contains(Linker->ImportMap[ImportIndex].XObject));
 	}
 	SoftPackageUsedInGameBits.Reserve(Linker->SoftPackageReferenceList.Num());
 	for (int32 SoftPackageIndex = 0; SoftPackageIndex < Linker->SoftPackageReferenceList.Num(); ++SoftPackageIndex)
 	{
-		SoftPackageUsedInGameBits.Add(SoftPackagesUsedInGame.Contains(Linker->SoftPackageReferenceList[SoftPackageIndex]));
+		SoftPackageUsedInGameBits.Add(Args.SoftPackagesUsedInGame->Contains(
+			Linker->SoftPackageReferenceList[SoftPackageIndex]));
 	}
 
 	// Serialize the Dependency section
 	DependencyDataRecord << SA_VALUE(TEXT("ImportUsedInGame"), ImportUsedInGameBits);
 	DependencyDataRecord << SA_VALUE(TEXT("SoftPackageUsedInGame"), SoftPackageUsedInGameBits);
+
+	// Currently the only type of ExtraPackageDependencies we have are the collected build dependencies,
+	// which have both the Build and PropagateManage flags. Store them as pairs of { PackageName, EExtraDepencyFlags }
+	// even though we only have one possible value for EExtraDepencyFlags, so that we can extend the types of dependencies
+	// later without needing to change EUnrealEngineObjectUE5Version.
+	TArray<TPair<FName, uint32>> ExtraPackageDependencies;
+	ExtraPackageDependencies.Reserve(Args.PackageBuildDependencies->Num());
+	constexpr EExtraDependencyFlags BuildAndPropagate =
+		EExtraDependencyFlags::Build | EExtraDependencyFlags::PropagateManage;
+	for (FName ExtraPackageName : *Args.PackageBuildDependencies)
+	{
+		ExtraPackageDependencies.Add({ ExtraPackageName, static_cast<uint32>(BuildAndPropagate) });
+	}
+	DependencyDataRecord << SA_VALUE(TEXT("ExtraPackageDependencies"), ExtraPackageDependencies);
 }
 
 }
diff --git a/Engine/Source/Runtime/CoreUObject/Private/UObject/SavePackage2.cpp b/Engine/Source/Runtime/CoreUObject/Private/UObject/SavePackage2.cpp
--- a/Engine/Source/Runtime/CoreUObject/Private/UObject/SavePackage2.cpp
+++ b/Engine/Source/Runtime/CoreUObject/Private/UObject/SavePackage2.cpp
@@ -461,6 +461,11 @@ ESavePackageResult HarvestPackage(FSaveContext& SaveContext)
 	SaveContext.SetCustomVersions(Harvester.GetCustomVersions());
 	SaveContext.SetTransientPropertyOverrides(Harvester.ReleaseTransientPropertyOverrides());
 
+	// Contractually all CookBuildDependencies (which come from PreSave or from Serialize) must be declared by this point.
+	// Copy PackageBuildDependencies from CookBuildDependencies and implement our contract that PackageBuildDependencies
+	// are in the name map.
+	SaveContext.UpdateEditorRealmPackageBuildDependencies();
+
 	return ReturnSuccessOrCancel();
 }
 
@@ -1215,7 +1220,8 @@ ESavePackageResult BuildLinker(FSaveContext& SaveContext)
 	{
 		SCOPED_SAVETIMER(UPackage_Save_BuildNameMap);
 		const TSet<FNameEntryId>& NamesReferencedFromExportData = SaveContext.GetNamesReferencedFromExportData();
-		const TSet<FNameEntryId>& NamesReferencedFromPackageHeader = SaveContext.GetNamesReferencedFromPackageHeader();
+		TSet<FNameEntryId>& NamesReferencedFromPackageHeader = SaveContext.GetNamesReferencedFromPackageHeader();
+
 		Linker->NameMap.Reserve(NamesReferencedFromExportData.Num() + NamesReferencedFromPackageHeader.Num());
 		for (FNameEntryId NameEntryId : NamesReferencedFromExportData)
 		{
@@ -1975,10 +1981,17 @@ ESavePackageResult WritePackageHeader(FStructuredArchive::FRecord& StructuredArc
 		// Save asset registry data so the editor can search for information about assets in this package
 		SCOPED_SAVETIMER(UPackage_Save_SaveAssetRegistryData);
 		FArchiveCookData* CookData = SaveContext.GetCookData();
-		FArchiveCookContext* CookContext = CookData ? &CookData->CookContext : nullptr;
-		UE::AssetRegistry::WritePackageData(StructuredArchiveRoot, CookContext, SaveContext.GetPackage(),
-			Linker, SaveContext.GetImportsUsedInGame(), SaveContext.GetSoftPackagesUsedInGame(),
-			&SaveContext.GetSavedAssets(), SaveContext.IsProceduralSave());
+		UE::AssetRegistry::FWritePackageDataArgs WriteARArgs;
+		WriteARArgs.ParentRecord = &StructuredArchiveRoot;
+		WriteARArgs.Package = SaveContext.GetPackage();
+		WriteARArgs.Linker = Linker;
+		WriteARArgs.ImportsUsedInGame = &SaveContext.GetImportsUsedInGame();
+		WriteARArgs.SoftPackagesUsedInGame = &SaveContext.GetSoftPackagesUsedInGame();
+		WriteARArgs.PackageBuildDependencies = &SaveContext.GetPackageBuildDependencies();
+		WriteARArgs.bProceduralSave = SaveContext.IsProceduralSave();
+		WriteARArgs.CookContext = CookData ? &CookData->CookContext : nullptr;;
+		WriteARArgs.OutAssetDatas = &SaveContext.GetSavedAssets();
+		UE::AssetRegistry::WritePackageData(WriteARArgs);
 	}
 	// Save level information used by World browser
 	{
diff --git a/Engine/Source/Runtime/CoreUObject/Public/Misc/AssetRegistryInterface.h b/Engine/Source/Runtime/CoreUObject/Public/Misc/AssetRegistryInterface.h
--- a/Engine/Source/Runtime/CoreUObject/Public/Misc/AssetRegistryInterface.h
+++ b/Engine/Source/Runtime/CoreUObject/Public/Misc/AssetRegistryInterface.h
@@ -131,6 +131,12 @@ namespace UE::AssetRegistry
 		Build = 0x010,			// Return only dependencies with EDependencyProperty::Build
 		NotBuild = 0x020,		// Return only dependencies without EDependencyProperty::Build
 
+		Propagation = 0x040,    // Return only dependencies that cause propagation of manage dependencies, which means
+		// either Game or Build. Presence of the Propagation flag in a query causes the Game, NotGame, EditorOnly,
+		// Build, and NotBuild flags to be ignored in the query if present. Either Game or Build is sufficient to pass
+		// the query no matter which of those other flags are present. Hard vs Soft is still respected, and Soft Build
+		// or Soft Game dependencies will be skipped query if Hard is required.
+
 		// Manage Dependencies Only
 		Direct = 0x0400,		// Return only dependencies with EDependencyProperty::Direct
 		NotDirect = 0x0800,		// Return only dependencies without EDependencyProperty::Direct
@@ -166,8 +172,25 @@ namespace UE::AssetRegistry
 	 */
 	struct FDependencyQuery
 	{
-		UE::AssetRegistry::EDependencyProperty Required; // Only Dependencies that possess all of these properties will be returned. Note that flags specific to another EDependencyCategory are ignored when querying dependencies in a given category.
-		UE::AssetRegistry::EDependencyProperty Excluded; // Only Dependencies that possess none of these properties will be returned. Note that flags specific to another EDependencyCategory are ignored when querying dependencies in a given category.
+		/**
+		 * Only Dependencies that possess all of these properties will be returned.
+		 * Note that flags specific to another EDependencyCategory are ignored when querying dependencies in a given category.
+		 */
+		UE::AssetRegistry::EDependencyProperty Required;
+		/**
+		 * Only Dependencies that possess none of these properties will be returned.
+		 * Note that flags specific to another EDependencyCategory are ignored when querying dependencies in a given category.
+		 */
+		UE::AssetRegistry::EDependencyProperty Excluded;
+
+		/**
+		 * RequiredUnions is an intersection of unions. Each element of RequiredUnions is a set of bit flags that are
+		 * unioned: having any one of the bit flags causes that element of RequiredUnions to pass.
+		 * After pass/fail is decided for each element, they are intersected: all must pass for the total to pass.
+		 * This allows RequiredUnions to be a conjunction of disjunctions, whereas the Required field is just a conjunction
+		 * of atoms.
+		 */
+		TArray<UE::AssetRegistry::EDependencyProperty, TInlineAllocator<1>> RequiredUnions;
 
 		FDependencyQuery()
 		{
@@ -177,29 +200,71 @@ namespace UE::AssetRegistry
 
 		inline FDependencyQuery(EDependencyQuery QueryFlags)
 		{
-			Required = (!!(QueryFlags & EDependencyQuery::Hard) ? UE::AssetRegistry::EDependencyProperty::Hard : UE::AssetRegistry::EDependencyProperty::None)
-				| (!!(QueryFlags & EDependencyQuery::Game) ? UE::AssetRegistry::EDependencyProperty::Game : UE::AssetRegistry::EDependencyProperty::None)
-				| (!!(QueryFlags & EDependencyQuery::Build) ? UE::AssetRegistry::EDependencyProperty::Build : UE::AssetRegistry::EDependencyProperty::None)
+			if (!EnumHasAnyFlags(QueryFlags, EDependencyQuery::Propagation))
+			{
+				Required = (!!(QueryFlags & EDependencyQuery::Game) ? UE::AssetRegistry::EDependencyProperty::Game : UE::AssetRegistry::EDependencyProperty::None)
+					| (!!(QueryFlags & EDependencyQuery::Build) ? UE::AssetRegistry::EDependencyProperty::Build : UE::AssetRegistry::EDependencyProperty::None);
+			}
+			else
+			{
+				EnumRemoveFlags(QueryFlags, EDependencyQuery::Game | EDependencyQuery::NotGame | EDependencyQuery::Build | EDependencyQuery::NotBuild);
+				RequiredUnions.Add(UE::AssetRegistry::EDependencyProperty::Game | UE::AssetRegistry::EDependencyProperty::Build);
+			}
+			Required |= (!!(QueryFlags & EDependencyQuery::Hard) ? UE::AssetRegistry::EDependencyProperty::Hard : UE::AssetRegistry::EDependencyProperty::None)
 				| (!!(QueryFlags & EDependencyQuery::Direct) ? UE::AssetRegistry::EDependencyProperty::Direct : UE::AssetRegistry::EDependencyProperty::None);
 			Excluded = (!!(QueryFlags & EDependencyQuery::NotHard) ? UE::AssetRegistry::EDependencyProperty::Hard : UE::AssetRegistry::EDependencyProperty::None)
 				| (!!(QueryFlags & EDependencyQuery::NotGame) ? UE::AssetRegistry::EDependencyProperty::Game : UE::AssetRegistry::EDependencyProperty::None)
 				| (!!(QueryFlags & EDependencyQuery::NotBuild) ? UE::AssetRegistry::EDependencyProperty::Build : UE::AssetRegistry::EDependencyProperty::None)
 				| (!!(QueryFlags & EDependencyQuery::NotDirect) ? UE::AssetRegistry::EDependencyProperty::Direct : UE::AssetRegistry::EDependencyProperty::None);
 		}
+	};
 
-		FDependencyQuery(const FDependencyQuery& Other) = default;
-		FDependencyQuery& operator=(const FDependencyQuery& Other) = default;
+	struct FWritePackageDataArgs
+	{
+		// Required inputs, must be initialized and non-null
+		FStructuredArchiveRecord* ParentRecord = nullptr;
+		const UPackage* Package = nullptr;
+		FLinkerSave* Linker = nullptr;
+		const TSet<TObjectPtr<UObject>>* ImportsUsedInGame = nullptr;
+		const TSet<FName>* SoftPackagesUsedInGame = nullptr;
+		const TArray<FName>* PackageBuildDependencies = nullptr;
+		bool bProceduralSave = false;
+
+		// Optional inputs that may be null
+		FArchiveCookContext* CookContext = nullptr;
+
+		// Optional outputs that may be null
+		TArray<FAssetData>* OutAssetDatas = nullptr;
+	};
+	/**
+	 * Bitfield of flags written into a package's AssetRegistry DependencyData section to represent what kind of
+	 * dependency is stored for each PackageName in ExtraPackageDependencies. Values are serialized as integers;
+	 * new bits can be added as necessary, but the integer values for existing enum values may not be changed.
+	 */
+	enum class EExtraDependencyFlags : uint32
+	{
+		None					= 0,
+		Build					= 0x1,
+		PropagateManage			= 0x2,
 	};
+	ENUM_CLASS_FLAGS(EExtraDependencyFlags);
 
-	// Functions to read and write the data used by the AssetRegistry in each package; the format of this data is separate from the format of the data in the asset registry
+	/**
+	 * Writes the data used by the AssetRegistry in each package; the format of this data is separate from the
+	 * format of the data in the asset registry.
+	 * The corresponding read functions are ReadPackageDataMain and ReadPackageDataDependencies; they are are declared
+	 * in IAssetRegistry.h, in the AssetRegistry module, because they depend upon some structures defined in the
+	 * AssetRegistry module
+	 */
+	COREUOBJECT_API void WritePackageData(FWritePackageDataArgs& Args);
+	UE_DEPRECATED(5.5, "Use version that takes FWritePackageDataArgs");
 	COREUOBJECT_API void WritePackageData(FStructuredArchiveRecord& ParentRecord, FArchiveCookContext* CookContext,
 		const UPackage* Package, FLinkerSave* Linker, const TSet<TObjectPtr<UObject>>& ImportsUsedInGame,
 		const TSet<FName>& SoftPackagesUsedInGame, TArray<FAssetData>* OutAssetDatas, bool bProceduralSave);
-	UE_DEPRECATED(5.4, "Use version that takes FArchiveCookContext");
+	UE_DEPRECATED(5.4, "Use version that takes FWritePackageDataArgs");
 	COREUOBJECT_API void WritePackageData(FStructuredArchiveRecord& ParentRecord, bool bIsCooking, const UPackage* Package,
 		FLinkerSave* Linker, const TSet<TObjectPtr<UObject>>& ImportsUsedInGame, const TSet<FName>& SoftPackagesUsedInGame,
 		const ITargetPlatform* TargetPlatform, TArray<FAssetData>* OutAssetDatas);
-	// ReadPackageDataMain and ReadPackageDataDependencies are declared in IAssetRegistry.h, in the AssetRegistry module, because they depend upon some structures defined in the AssetRegistry module
 
 	namespace Private
 	{
diff --git a/Engine/Source/Runtime/CoreUObject/Public/UObject/ObjectSaveContext.h b/Engine/Source/Runtime/CoreUObject/Public/UObject/ObjectSaveContext.h
--- a/Engine/Source/Runtime/CoreUObject/Public/UObject/ObjectSaveContext.h
+++ b/Engine/Source/Runtime/CoreUObject/Public/UObject/ObjectSaveContext.h
@@ -215,6 +215,9 @@ public:
 	/**
 	 * Add the given FCookDependency to the build dependencies for the package being cook-saved. Iterative cooks will
 	 * invalidate the package and recook it if the CookDependency changes.
+	 * CookBuildDependencies of type FCookDependency::CreatePackage are also used when saving an editor package rather
+	 * than a cooked package. They are saved into the AssetRegistry dependency data for the package, as
+	 * EDependencyProperty::Build properties.
 	 */
 	COREUOBJECT_API void AddCookBuildDependency(UE::Cook::FCookDependency BuildDependency);
 	/**
diff --git a/Engine/Source/Runtime/Engine/Private/LevelInstance/LevelInstanceActor.cpp b/Engine/Source/Runtime/Engine/Private/LevelInstance/LevelInstanceActor.cpp
--- a/Engine/Source/Runtime/Engine/Private/LevelInstance/LevelInstanceActor.cpp
+++ b/Engine/Source/Runtime/Engine/Private/LevelInstance/LevelInstanceActor.cpp
@@ -10,6 +10,7 @@
 #include UE_INLINE_GENERATED_CPP_BY_NAME(LevelInstanceActor)
 
 #if WITH_EDITOR
+#include "Cooker/CookDependency.h"
 #include "UObject/ObjectSaveContext.h"
 #include "WorldPartition/LevelInstance/LevelInstanceActorDesc.h"
 #include "LevelInstance/LevelInstanceEditorPivotActor.h"
@@ -184,14 +185,23 @@ void ALevelInstance::PostLoad()
 #endif
 }
 
-#if WITH_EDITOR
+void ALevelInstance::PreSave(FObjectPreSaveContext SaveContext)
+{
+	Super::PreSave(SaveContext);
+
+	FName WorldPackageName = FName(GetWorldAsset().GetLongPackageName());
+	if (!WorldPackageName.IsNone())
+	{
+		SaveContext.AddCookBuildDependency(UE::Cook::FCookDependency::Package(WorldPackageName));
+	}
+}
+
 bool ALevelInstance::ShouldCookWorldAsset() const
 {
 	// If ALevelInstance actor gets loaded it means it needs to Cook its WorldAsset (World Partition Embedded Level Instances don't get loaded as they aren't runtime relevant)
 	// If ALevelInstnace is a template then we only need to Cook its WorldAsset if it's desired runtime behavior is to be Level Streamed
 	return !IsTemplate() || GetDesiredRuntimeBehavior() == ELevelInstanceRuntimeBehavior::LevelStreaming;
 }
-#endif
 
 void ALevelInstance::PreEditUndo()
 {
@@ -383,7 +393,7 @@ void ALevelInstance::PushLevelInstanceEditingStateToProxies(bool bInEditingState
 	LevelInstanceActorImpl.PushLevelInstanceEditingStateToProxies(bInEditingState);
 }
 
-#endif
+#endif // WITH_EDITOR
 
 #undef LOCTEXT_NAMESPACE
 
diff --git a/Engine/Source/Runtime/Engine/Public/LevelInstance/LevelInstanceActor.h b/Engine/Source/Runtime/Engine/Public/LevelInstance/LevelInstanceActor.h
--- a/Engine/Source/Runtime/Engine/Public/LevelInstance/LevelInstanceActor.h
+++ b/Engine/Source/Runtime/Engine/Public/LevelInstance/LevelInstanceActor.h
@@ -85,6 +85,7 @@ public:
 	ENGINE_API virtual void PostEditUndo() override;
 	ENGINE_API virtual void PostEditUndo(TSharedPtr<ITransactionObjectAnnotation> TransactionAnnotation) override;
 	ENGINE_API virtual void PostLoad() override;
+	ENGINE_API virtual void PreSave(FObjectPreSaveContext SaveContext) override;
 	ENGINE_API virtual bool CanEditChange(const FProperty* InProperty) const override;
 	ENGINE_API virtual void PreEditChange(FProperty* PropertyThatWillChange) override;
 	ENGINE_API virtual void PostEditChangeProperty(FPropertyChangedEvent& PropertyChangedEvent) override;
