Tests: AssetRegistry: Fixed 'System.AssetRegistry.AssetPathStrings' test failing on most platforms because of an incorrect test case.

The test case expected to get full asset paths from FAssetData, which platforms/targets built without WITH_EDITORONLY_DATA return only simplified paths. Matched the test case checks to reflect that.

#rnx

diff --git a/Engine/Source/Runtime/AssetRegistry/Private/AssetRegistryState.cpp b/Engine/Source/Runtime/AssetRegistry/Private/AssetRegistryState.cpp
--- a/Engine/Source/Runtime/AssetRegistry/Private/AssetRegistryState.cpp
+++ b/Engine/Source/Runtime/AssetRegistry/Private/AssetRegistryState.cpp
@@ -4128,20 +4128,18 @@ bool FAssetRegistryAssetPathStringsTest::RunTest(const FString& Parameters)
 	FAssetData DirectSubObjectAssetData = PathToAssetData(DirectSubObjectPathString);
 	FAssetData SubSubObjectAssetData = PathToAssetData(SubSubObjectPathString);
 
-	auto GetAssetDataPath = [](const FAssetData& AssetData)
-	{
-		TStringBuilder<FName::StringBufferSize> Buffer;
-		AssetData.AppendObjectPath(Buffer);
-		return FString(Buffer);
-	};
-
 	// Test FAssetData::AppendPath for asset data with variable length of OptionalOuterPath
 	TestEqual(TEXT("FAssetData::AppendPath() correct for top-level asset"),
 		TopLevelAssetData.GetObjectPathString(), TopLevelPathString);
+
+#if WITH_EDITORONLY_DATA
+	// These tests are only enabled when WITH_EDITORONLY_DATA is active because only
+	// then OuterPath is retained by FAssedData (see FAssetData::AppendObjectPath).
 	TestEqual(TEXT("FAssetData::AppendPath() correct for subobject asset"),
 		DirectSubObjectAssetData.GetObjectPathString(), DirectSubObjectPathString);
 	TestEqual(TEXT("FAssetData::AppendPath() correct for sub-subobject asset"),
 		SubSubObjectAssetData.GetObjectPathString(), SubSubObjectPathString);
+#endif
 	
 	return true;
 }
