Emit asset added events when loading the initial premade asset registry state in the editor

This gives consistent behavior for initial asset "discovery" between a cooked and uncooked editor

#rnx

diff --git a/Engine/Source/Runtime/AssetRegistry/Private/AssetRegistry.cpp b/Engine/Source/Runtime/AssetRegistry/Private/AssetRegistry.cpp
--- a/Engine/Source/Runtime/AssetRegistry/Private/AssetRegistry.cpp
+++ b/Engine/Source/Runtime/AssetRegistry/Private/AssetRegistry.cpp
@@ -887,6 +887,8 @@ void FAssetRegistryImpl::LoadPremadeAssetRegistry(Impl::FEventContext& EventCont
 	SCOPED_BOOT_TIMING("LoadPremadeAssetRegistry");
 	UE_SCOPED_ENGINE_ACTIVITY("Loading premade asset registry");
 
+	const bool bEmitAssetEvents = GIsEditor;
+
 	if (SerializationOptions.bSerializeAssetRegistry)
 	{
 		SCOPED_BOOT_TIMING("LoadPremadeAssetRegistry_Main");
@@ -896,17 +898,31 @@ void FAssetRegistryImpl::LoadPremadeAssetRegistry(Impl::FEventContext& EventCont
 			{
 				State = MoveTemp(ARState);
 				CachePathsFromState(EventContext, State);
+				if (bEmitAssetEvents)
+				{
+					State.EnumerateAllAssets([&EventContext](const FAssetData& AssetData)
+					{
+						EventContext.AssetEvents.Emplace(AssetData, UE::AssetRegistry::Impl::FEventContext::EEvent::Added);
+					});
+				}
 			}
 			else if (State.GetNumAssets() < ARState.GetNumAssets())
 			{
 				FAssetRegistryState ExistingState = MoveTemp(State);
 				State = MoveTemp(ARState);
 				CachePathsFromState(EventContext, State);
+				if (bEmitAssetEvents)
+				{
+					State.EnumerateAllAssets([&EventContext](const FAssetData& AssetData)
+					{
+						EventContext.AssetEvents.Emplace(AssetData, UE::AssetRegistry::Impl::FEventContext::EEvent::Added);
+					});
+				}
 				AppendState(EventContext, ExistingState);
 			}
 			else
 			{
-				AppendState(EventContext, ARState, FAssetRegistryState::EInitializationMode::OnlyUpdateNew);
+				AppendState(EventContext, ARState, FAssetRegistryState::EInitializationMode::OnlyUpdateNew, bEmitAssetEvents);
 			}
 			UpdatePersistentMountPoints();
 			State.bCookedGlobalAssetRegistryState = true;
@@ -939,9 +955,9 @@ void FAssetRegistryImpl::LoadPremadeAssetRegistry(Impl::FEventContext& EventCont
 					 * The main state will often already include the DLC/plugin assets and is often in a development mode where the plugin state will not be.
 					 * If we update the existing assets in those cases it will be causing a lost of tags and values that are needed for the editor systems.
 					 */
-					AppendState(EventContext, PluginState,  FAssetRegistryState::EInitializationMode::OnlyUpdateNew);
+					AppendState(EventContext, PluginState, FAssetRegistryState::EInitializationMode::OnlyUpdateNew, bEmitAssetEvents);
 #else
-					AppendState(EventContext, PluginState);
+					AppendState(EventContext, PluginState, FAssetRegistryState::EInitializationMode::Append, bEmitAssetEvents);
 #endif
 				}
 			}
@@ -5491,15 +5507,7 @@ void UAssetRegistryImpl::AppendState(const FAssetRegistryState& InState)
 	{
 		LLM_SCOPE(ELLMTag::AssetRegistry);
 		UE::AssetRegistry::FInterfaceWriteScopeLock InterfaceScopeLock(InterfaceLock);
-		GuardedData.AppendState(EventContext, InState);
-	}
-
-	// AppendState does not create the AssetAdded events; create them here and then broadcast
-	TArray<FAssetData> TempAssets;
-	InState.GetAllAssets(TSet<FName>(), TempAssets, true /* bARFiltering */);
-	for (FAssetData& AssetData : TempAssets)
-	{
-		EventContext.AssetEvents.Emplace(MoveTemp(AssetData), UE::AssetRegistry::Impl::FEventContext::EEvent::Added);
+		GuardedData.AppendState(EventContext, InState, FAssetRegistryState::EInitializationMode::Append, /*bEmitAssetEvents*/true);
 	}
 	Broadcast(EventContext);
 }
@@ -5507,8 +5515,12 @@ void UAssetRegistryImpl::AppendState(const FAssetRegistryState& InState)
 namespace UE::AssetRegistry
 {
 
-void FAssetRegistryImpl::AppendState(Impl::FEventContext& EventContext, const FAssetRegistryState& InState, FAssetRegistryState::EInitializationMode Mode)
+void FAssetRegistryImpl::AppendState(Impl::FEventContext& EventContext, const FAssetRegistryState& InState, 
+	FAssetRegistryState::EInitializationMode Mode, bool bEmitAssetEvents)
 {
+	FAssetRegistryAppendResult LocalAppendResult;
+	FAssetRegistryAppendResult* AppendResultPtr = bEmitAssetEvents ? &LocalAppendResult : nullptr;
+
 	// @note Using this define to identify cooked editor for now. We might want to change the name later if we find more differences.
 	State.InitializeFromExisting(
 		InState,
@@ -5517,9 +5529,22 @@ void FAssetRegistryImpl::AppendState(Impl::FEventContext& EventContext, const FA
 #else
 		SerializationOptions,
 #endif
-		Mode);
+		Mode,
+		AppendResultPtr);
 
 	CachePathsFromState(EventContext, InState);
+
+	if (AppendResultPtr)
+	{
+		for (const FAssetData* AssetData : AppendResultPtr->AddedAssets)
+		{
+			EventContext.AssetEvents.Emplace(*AssetData, UE::AssetRegistry::Impl::FEventContext::EEvent::Added);
+		}
+		for (const FAssetData* AssetData : AppendResultPtr->UpdatedAssets)
+		{
+			EventContext.AssetEvents.Emplace(*AssetData, UE::AssetRegistry::Impl::FEventContext::EEvent::Updated);
+		}
+	}
 }
 
 void FAssetRegistryImpl::CachePathsFromState(Impl::FEventContext& EventContext, const FAssetRegistryState& InState)
diff --git a/Engine/Source/Runtime/AssetRegistry/Private/AssetRegistryImpl.h b/Engine/Source/Runtime/AssetRegistry/Private/AssetRegistryImpl.h
--- a/Engine/Source/Runtime/AssetRegistry/Private/AssetRegistryImpl.h
+++ b/Engine/Source/Runtime/AssetRegistry/Private/AssetRegistryImpl.h
@@ -248,7 +248,7 @@ public:
 		const TArray<FString>& InFilePaths, UE::AssetRegistry::EScanFlags ScanFlags);
 	void Serialize(FArchive& Ar, Impl::FEventContext& EventContext);
 	void AppendState(Impl::FEventContext& EventContext, const FAssetRegistryState& InState,
-		FAssetRegistryState::EInitializationMode Mode = FAssetRegistryState::EInitializationMode::Append);
+		FAssetRegistryState::EInitializationMode Mode = FAssetRegistryState::EInitializationMode::Append, bool bEmitAssetEvents = false);
 	void GetAllocatedSize(bool bLogDetailed, SIZE_T& StateSize, SIZE_T& StaticSize, SIZE_T& SearchSize) const;
 	bool IsLoadingAssets() const;
 	void SetManageReferences(const TMultiMap<FAssetIdentifier, FAssetIdentifier>& ManagerMap,
diff --git a/Engine/Source/Runtime/AssetRegistry/Private/AssetRegistryState.cpp b/Engine/Source/Runtime/AssetRegistry/Private/AssetRegistryState.cpp
--- a/Engine/Source/Runtime/AssetRegistry/Private/AssetRegistryState.cpp
+++ b/Engine/Source/Runtime/AssetRegistry/Private/AssetRegistryState.cpp
@@ -366,7 +366,8 @@ void FAssetRegistryState::InitializeFromExisting(const FAssetDataMap& AssetDataM
 	const TMap<FAssetIdentifier, FDependsNode*>& DependsNodeMap, 
 	const TMap<FName, FAssetPackageData*>& AssetPackageDataMap,
 	const FAssetRegistrySerializationOptions& Options,
-	EInitializationMode InInitializationMode)
+	EInitializationMode InInitializationMode,
+	FAssetRegistryAppendResult* OutAppendResult)
 {
 	if (InInitializationMode == EInitializationMode::Rebuild)
 	{
@@ -408,13 +409,22 @@ void FAssetRegistryState::InitializeFromExisting(const FAssetDataMap& AssetDataM
 		{
 			FAssetData NewData(AssetData);
 			NewData.TagsAndValues = FAssetDataTagMapSharedView(MoveTemp(LocalTagsAndValues));
-			UpdateAssetData(ExistingData, MoveTemp(NewData));
+			bool bIsModified = false;
+			UpdateAssetData(ExistingData, MoveTemp(NewData), &bIsModified);
+			if (OutAppendResult && bIsModified)
+			{
+				OutAppendResult->UpdatedAssets.Add(ExistingData);
+			}
 		}
 		else
 		{
 			FAssetData* NewData = new FAssetData(AssetData);
 			NewData->TagsAndValues = FAssetDataTagMapSharedView(MoveTemp(LocalTagsAndValues));
 			AddAssetData(NewData);
+			if (OutAppendResult)
+			{
+				OutAppendResult->AddedAssets.Add(NewData);
+			}
 		}
 	}
 
diff --git a/Engine/Source/Runtime/AssetRegistry/Public/AssetRegistry/AssetRegistryState.h b/Engine/Source/Runtime/AssetRegistry/Public/AssetRegistry/AssetRegistryState.h
--- a/Engine/Source/Runtime/AssetRegistry/Public/AssetRegistry/AssetRegistryState.h
+++ b/Engine/Source/Runtime/AssetRegistry/Public/AssetRegistry/AssetRegistryState.h
@@ -144,6 +144,12 @@ struct FAssetRegistryPruneOptions
 	TSet<FPrimaryAssetType> RemoveDependenciesWithoutPackagesKeepPrimaryAssetTypes;
 };
 
+struct FAssetRegistryAppendResult
+{
+	TArray<const FAssetData*> AddedAssets;
+	TArray<const FAssetData*> UpdatedAssets;
+};
+
 /**
  * The state of an asset registry, this is used internally by IAssetRegistry to represent the disk cache,
  * and is also accessed directly to save/load cooked caches.
@@ -531,7 +537,7 @@ public:
 	ASSETREGISTRY_API void Reset();
 
 	void InitializeFromExisting(const FAssetRegistryState& Existing, const FAssetRegistrySerializationOptions& Options,
-		EInitializationMode InitializationMode = EInitializationMode::Rebuild);
+		EInitializationMode InitializationMode = EInitializationMode::Rebuild, FAssetRegistryAppendResult* OutAppendResult = nullptr);
 
 	/** 
 	 * Prunes an asset cache, this removes asset data, nodes, and package data that isn't needed. 
@@ -630,7 +636,8 @@ private:
 		const TMap<FAssetIdentifier, FDependsNode*>& DependsNodeMap,
 		const TMap<FName, FAssetPackageData*>& AssetPackageDataMap,
 		const FAssetRegistrySerializationOptions& Options,
-		EInitializationMode InitializationMode = EInitializationMode::Rebuild);
+		EInitializationMode InitializationMode = EInitializationMode::Rebuild,
+		FAssetRegistryAppendResult* OutAppendResult = nullptr);
 
 	template<class Archive>
 	void Load(Archive&& Ar, const FAssetRegistryHeader& Header, const FAssetRegistryLoadOptions& Options);
@@ -1055,8 +1062,9 @@ inline int32 FAssetRegistryState::GetNumPackages() const
 
 inline void FAssetRegistryState::InitializeFromExisting(const FAssetRegistryState& Existing,
 	const FAssetRegistrySerializationOptions& Options,
-	EInitializationMode InitializationMode)
+	EInitializationMode InitializationMode,
+	FAssetRegistryAppendResult* OutAppendResult)
 {
 	InitializeFromExisting(Existing.CachedAssets, Existing.CachedDependsNodes,
-		Existing.CachedPackageData, Options, InitializationMode);
-}
\ No newline at end of file
+		Existing.CachedPackageData, Options, InitializationMode, OutAppendResult);
+}
