AssetRegistry: Add option Engine.ini:[AssetRegistry]:BlockingInitialLoad to make the initial scan blocking which reduces total scan time at the cost of making the time to interactivity longer.
#rnx

diff --git a/Engine/Config/BaseEditorPerProjectUserSettings.ini b/Engine/Config/BaseEditorPerProjectUserSettings.ini
--- a/Engine/Config/BaseEditorPerProjectUserSettings.ini
+++ b/Engine/Config/BaseEditorPerProjectUserSettings.ini
@@ -991,3 +991,10 @@ SyntaxColors=(Comment=(R=0.258183,G=0.539479,B=0.068478,A=1.000000),Error=(R=0.8
 FontSize=9
 SelectionColor=(R=0.051269,G=0.095307,B=0.158961,A=1.000000)
 SortMethod=None
+
+[AssetRegistry]
+; Only used in editor. If true, the cpu calculations of the scan will not only run on the game thread, but will run
+; in a blocking fashion, blocking editor realtime interactivity until the scan is complete rather than making the editor
+; interactable as soon as the minimum set is reached. This will reduce the time necessary to complete the scan, at the cost
+; of lengthening the time until the editor becomes interactable.
+BlockingInitialLoad=false
diff --git a/Engine/Config/BaseEngine.ini b/Engine/Config/BaseEngine.ini
--- a/Engine/Config/BaseEngine.ini
+++ b/Engine/Config/BaseEngine.ini
@@ -2819,6 +2819,7 @@ AspectRatioAxisConstraint=AspectRatio_MaintainYFOV
 +DisablePythonForCommandlet=CookShaders
 
 [AssetRegistry]
+; See also the AssetRegistry section in BaseEditorPerProjectUserSettings, for user-specific options
 ; Default Path Filters not scanned by the asset registry when searching for all assets
 ; These are long package paths
 ;+BlacklistPackagePathScanFilters="/Game/__ExternalActors__"
@@ -2866,6 +2867,12 @@ bFilterSearchableNames=true
 ; Set EditorGameScansAR=false to turn it off.
 EditorGameScansAR=true
 
+; Performance settings for the AssetRegistryScan
+; Only used in editor. In editor the assetregistry scan by default does IO on a background thread and also does cpu
+; calculations on the loaded data on the background thread. Setting GatherOnGameThreadOnly=true still does IO in the
+; background, but forces the calculations to run on the GameThread
+GatherOnGameThreadOnly=false
+
 ; Commandlets by default do not SearchAllAssets because many do not need it. A commandlet can be opted in
 ; using CommandletsUsingAR. Format is the shortname of the commandlet, optionally without the "Commandlet" suffix.
 +CommandletsUsingAR=DataValidation
diff --git a/Engine/Source/Runtime/AssetRegistry/Private/AssetDataGatherer.cpp b/Engine/Source/Runtime/AssetRegistry/Private/AssetDataGatherer.cpp
--- a/Engine/Source/Runtime/AssetRegistry/Private/AssetDataGatherer.cpp
+++ b/Engine/Source/Runtime/AssetRegistry/Private/AssetDataGatherer.cpp
@@ -85,12 +85,6 @@ static FAutoConsoleVariableRef CVarIgnoreEmptyDirectories(
 	bIgnoreEmptyDirectories,
 	TEXT("If true, completely empty leaf directories are ignored by the asset registry while scanning"));
 
-bool bTickGatherOnGTOnly = false;
-static FAutoConsoleVariableRef CVarTickGatherOnGTOnly(
-	TEXT("AssetRegistry.TickGatherOnGTOnly"),
-	bTickGatherOnGTOnly,
-	TEXT("If true, TickGatherer will only be called from the game thread."));
-
 void LexFromString(EFeatureEnabledReadWrite& OutValue, FStringView Text)
 {
 	Text.TrimStartAndEndInline();
@@ -3614,6 +3608,7 @@ FAssetDataGatherer::FAssetDataGatherer(const TArray<FString>& InLongPackageNames
 	, bFirstTickAfterIdle(true)
 	, bFinishedInitialDiscovery(false)
 	, bIsInitialSearchCompleted(false)
+	, bGatherOnGameThreadOnly(false)
 	, WaitBatchCount(-1)
 	, LastCacheSaveNumUncachedAssetFiles(0)
 	, CacheInUseCount(0)
@@ -3724,7 +3719,7 @@ uint32 FAssetDataGatherer::Run()
 			}
 
 			UE::AssetRegistry::Impl::EGatherStatus Status = UE::AssetRegistry::Impl::EGatherStatus::Complete;
-			if (bLocalIdle && !UE::AssetDataGather::Private::bTickGatherOnGTOnly 
+			if (bLocalIdle && !IsGatherOnGameThreadOnly()
 				&& (IsProcessingPaused.load(std::memory_order_relaxed) == 0))
 			{
 				IAssetRegistry& Registry = IAssetRegistry::GetChecked();
diff --git a/Engine/Source/Runtime/AssetRegistry/Private/AssetDataGatherer.h b/Engine/Source/Runtime/AssetRegistry/Private/AssetDataGatherer.h
--- a/Engine/Source/Runtime/AssetRegistry/Private/AssetDataGatherer.h
+++ b/Engine/Source/Runtime/AssetRegistry/Private/AssetDataGatherer.h
@@ -247,6 +247,9 @@ public:
 	void ResumeProcessing() { IsProcessingPaused.fetch_sub(1, std::memory_order_relaxed); }
 	bool IsProcessingPauseRequested() const { return IsProcessingPaused.load(std::memory_order_relaxed) != 0; }
 
+	void SetGatherOnGameThreadOnly(bool bValue);
+	bool IsGatherOnGameThreadOnly() const;
+
 private:
 	enum class ETickResult
 	{
@@ -472,6 +475,7 @@ private:
 	bool bFinishedInitialDiscovery;
 	/** True if OnInitialSearchCompleted has been called. */
 	std::atomic<bool> bIsInitialSearchCompleted;
+	std::atomic<bool> bGatherOnGameThreadOnly;
 
 	// Variable section for variables that are read/writable only within TickLock.
 
@@ -507,3 +511,19 @@ private:
 	/** Packages can be marked for retry up until bInitialPluginsLoaded is set. After it is set, we retry them once. */
 	bool bFlushedRetryFiles;
 };
+
+
+///////////////////////////////////////////////////////
+// Inline implementations
+///////////////////////////////////////////////////////
+
+
+inline void FAssetDataGatherer::SetGatherOnGameThreadOnly(bool bValue)
+{
+	bGatherOnGameThreadOnly.store(bValue, std::memory_order_relaxed);
+}
+
+inline bool FAssetDataGatherer::IsGatherOnGameThreadOnly() const
+{
+	return bGatherOnGameThreadOnly.load(std::memory_order_relaxed);
+}
diff --git a/Engine/Source/Runtime/AssetRegistry/Private/AssetRegistry.cpp b/Engine/Source/Runtime/AssetRegistry/Private/AssetRegistry.cpp
--- a/Engine/Source/Runtime/AssetRegistry/Private/AssetRegistry.cpp
+++ b/Engine/Source/Runtime/AssetRegistry/Private/AssetRegistry.cpp
@@ -160,6 +160,7 @@ namespace UE::AssetRegistry
 namespace UE::AssetRegistry::Impl
 {
 	/** The max time to spend in UAssetRegistryImpl::Tick */
+	constexpr float MaxSecondsPerFrameToUseInBlockingInitialLoad = 5.0f;
 	float MaxSecondsPerFrame = 0.04f;
 	static FAutoConsoleVariableRef CVarAssetRegistryMaxSecondsPerFrame(
 		TEXT("AssetRegistry.MaxSecondsPerFrame"),
@@ -974,6 +975,7 @@ void FAssetRegistryImpl::Initialize(Impl::FInitializeContext& Context)
 
 	bInitialSearchStarted = false;
 	bInitialSearchCompleted.store(true, std::memory_order_relaxed);
+	UpdateMaxSecondsPerFrame();
 	GatherStatus = Impl::EGatherStatus::TickActiveGatherActive;
 	PerformanceMode = Impl::EPerformanceMode::MostlyStatic;
 
@@ -1832,6 +1834,7 @@ bool FAssetRegistryImpl::TryConstructGathererIfNeeded()
 
 	bool bAsyncGatherEnabled = !IsRunningGame() && !IsRunningDedicatedServer();
 	GlobalGatherer = MakeUnique<FAssetDataGatherer>(PathsDenyList, ContentSubPathsDenyList, bAsyncGatherEnabled, *this);
+	UpdateMaxSecondsPerFrame();
 
 	// Read script packages if all initial plugins have been loaded, otherwise do nothing; we wait for the callback.
 	ELoadingPhase::Type LoadingPhase = IPluginManager::Get().GetLastCompletedLoadingPhase();
@@ -1946,6 +1949,7 @@ void FAssetRegistryImpl::SearchAllAssets(Impl::FEventContext& EventContext,
 		InitialSearchStartTime = FPlatformTime::Seconds();
 		bInitialSearchStarted = true;
 		bInitialSearchCompleted.store(false, std::memory_order_relaxed);
+		UpdateMaxSecondsPerFrame();
 	}
 
 	FAssetDataGatherer& Gatherer = *GlobalGatherer;
@@ -4631,6 +4635,8 @@ void UAssetRegistryImpl::Tick(float DeltaTime)
 	}
 
 	bool bInterruptedOrShouldProcessDeferredEvents = false;
+	float LocalMaxSecondsPerFrame = UE::AssetRegistry::Impl::MaxSecondsPerFrame;
+
 	do
 	{
 		bInterruptedOrShouldProcessDeferredEvents = false;
@@ -4656,7 +4662,8 @@ void UAssetRegistryImpl::Tick(float DeltaTime)
 			{
 				// Tick the Gatherer
 				UE::AssetRegistry::Impl::FTickContext TickContext(EventContext, InheritanceContext);
-				TickContext.InterruptionContext.SetLimitedTickTime(TickStartTime, UE::AssetRegistry::Impl::MaxSecondsPerFrame);
+				LocalMaxSecondsPerFrame = GuardedData.MaxSecondsPerFrame;
+				TickContext.InterruptionContext.SetLimitedTickTime(TickStartTime, LocalMaxSecondsPerFrame);
 				TickContext.bHandleCompletion = true;
 				TickContext.bHandleDeferred = true;
 				Status = GuardedData.TickGatherer(TickContext);
@@ -4680,7 +4687,7 @@ void UAssetRegistryImpl::Tick(float DeltaTime)
 		if (!bInterruptedOrShouldProcessDeferredEvents)
 		{
 			UE::AssetRegistry::Impl::FInterruptionContext InterruptionContext;
-			InterruptionContext.SetLimitedTickTime(TickStartTime, UE::AssetRegistry::Impl::MaxSecondsPerFrame);
+			InterruptionContext.SetLimitedTickTime(TickStartTime, LocalMaxSecondsPerFrame);
 			ProcessLoadedAssetsToUpdateCache(EventContext, Status, InterruptionContext);
 			bInterruptedOrShouldProcessDeferredEvents = bInterruptedOrShouldProcessDeferredEvents
 				|| InterruptionContext.WasInterrupted();
@@ -4692,7 +4699,7 @@ void UAssetRegistryImpl::Tick(float DeltaTime)
 			Broadcast(EventContext, true /* bAllowFileLoadedEvent */);
 		}
 	} while ((bInterruptedOrShouldProcessDeferredEvents || Status == UE::AssetRegistry::Impl::EGatherStatus::WaitingForEvents) &&
-		(TickStartTime < 0 || (FPlatformTime::Seconds() - TickStartTime) <= UE::AssetRegistry::Impl::MaxSecondsPerFrame));
+		(TickStartTime < 0 || (FPlatformTime::Seconds() - TickStartTime) <= LocalMaxSecondsPerFrame));
 }
 
 namespace UE::AssetRegistry
@@ -4721,6 +4728,36 @@ bool FAssetRegistryImpl::ClassRequiresGameThreadProcessing(const UClass* Class)
 	return true;
 }
 
+void FAssetRegistryImpl::UpdateMaxSecondsPerFrame()
+{
+	float NewMaxSecondsPerFrame = UE::AssetRegistry::Impl::MaxSecondsPerFrame;
+#if WITH_EDITOR
+	bool bGatherOnGameThreadOnly = false;
+	GConfig->GetBool(TEXT("AssetRegistry"), TEXT("GatherOnGameThreadOnly"), bGatherOnGameThreadOnly, GEngineIni);
+
+	if (bInitialSearchStarted && !bInitialSearchCompleted)
+	{
+		bool bBlockingInitialLoad;
+		GConfig->GetBool(TEXT("AssetRegistry"), TEXT("BlockingInitialLoad"), bBlockingInitialLoad, GEditorPerProjectIni);
+		if (bBlockingInitialLoad)
+		{
+			bGatherOnGameThreadOnly = true;
+			NewMaxSecondsPerFrame = UE::AssetRegistry::Impl::MaxSecondsPerFrameToUseInBlockingInitialLoad;
+			if (MaxSecondsPerFrame < NewMaxSecondsPerFrame)
+			{
+				UE_LOG(LogAssetRegistry, Display,
+					TEXT("EditorPerProjectUserSettings.ini:[AssetRegistry]:BlockingInitialLoad=true, setting AssetRegistry load to blocking. The editor will not be interactive until the initial scan completes."));
+			}
+		}
+	}
+	if (GlobalGatherer)
+	{
+		GlobalGatherer->SetGatherOnGameThreadOnly(bGatherOnGameThreadOnly);
+	}
+#endif
+	MaxSecondsPerFrame = NewMaxSecondsPerFrame;
+}
+
 Impl::EGatherStatus FAssetRegistryImpl::TickGatherer(Impl::FTickContext& TickContext)
 {
 	TRACE_CPUPROFILER_EVENT_SCOPE_STR_CONDITIONAL("FAssetRegistryImpl::TickGatherer", !IsInitialSearchCompleted());
@@ -5168,6 +5205,7 @@ void FAssetRegistryImpl::OnInitialSearchCompleted(Impl::FEventContext& EventCont
 	EventContext.bFileLoadedEventBroadcast = true;
 	
 	bInitialSearchCompleted.store(true, std::memory_order_relaxed);
+	UpdateMaxSecondsPerFrame();
 }
 
 void FAssetRegistryImpl::LogSearchDiagnostics(double StartTime)
diff --git a/Engine/Source/Runtime/AssetRegistry/Private/AssetRegistryImpl.h b/Engine/Source/Runtime/AssetRegistry/Private/AssetRegistryImpl.h
--- a/Engine/Source/Runtime/AssetRegistry/Private/AssetRegistryImpl.h
+++ b/Engine/Source/Runtime/AssetRegistry/Private/AssetRegistryImpl.h
@@ -514,6 +514,8 @@ private:
 	 */
 	TUniquePtr<FAssetDataGatherer>& AccessGlobalGatherer() { return GlobalGatherer; }
 
+	void UpdateMaxSecondsPerFrame();
+
 private:
 
 	/** Internal state of the cached asset registry */
@@ -607,6 +609,7 @@ private:
 
 	bool bForceCompletionEvenIfPostLoadsFail = false;
 	bool bProcessedAnyAssetsAfterRetryDeferred = true;
+	float MaxSecondsPerFrame = 0.04f;
 
 	/** List of all class names derived from Blueprint (including Blueprint itself) */
 	TSet<FTopLevelAssetPath> ClassGeneratorNames;
