AssetRegistry: Improve BlockingInitialLoad by letting the gatherer still run on the background thread, and implement the handshake to pause it when the gamethread takes over the gather during tick. This allows us to still do the background gather during initial boot, before we reach the point where the gamethread is ready to tick it.
#rnx

diff --git a/Engine/Source/Runtime/AssetRegistry/Private/AssetRegistry.cpp b/Engine/Source/Runtime/AssetRegistry/Private/AssetRegistry.cpp
--- a/Engine/Source/Runtime/AssetRegistry/Private/AssetRegistry.cpp
+++ b/Engine/Source/Runtime/AssetRegistry/Private/AssetRegistry.cpp
@@ -975,6 +975,10 @@ void FAssetRegistryImpl::Initialize(Impl::FInitializeContext& Context)
 
 	bInitialSearchStarted = false;
 	bInitialSearchCompleted.store(true, std::memory_order_relaxed);
+#if WITH_EDITOR
+	SetGameThreadTakeOverGatherEachTick(false);
+#endif
+
 	UpdateMaxSecondsPerFrame();
 	GatherStatus = Impl::EGatherStatus::TickActiveGatherActive;
 	PerformanceMode = Impl::EPerformanceMode::MostlyStatic;
@@ -1076,7 +1080,7 @@ void FAssetRegistryImpl::Initialize(Impl::FInitializeContext& Context)
 	InitRedirectors(Context.Events, Context.InheritanceContext, Context.bRedirectorsNeedSubscribe);
 
 #if WITH_EDITOR
-	// Makae sure first call to LoadCalculatedDependencies builds the Gatherer list. At that point Classes should be loaded.
+	// Make sure first call to LoadCalculatedDependencies builds the Gatherer list. At that point Classes should be loaded.
 	bRegisteredDependencyGathererClassesDirty = true;
 #endif
 }
@@ -4643,7 +4647,45 @@ void UAssetRegistryImpl::Tick(float DeltaTime)
 
 		UE::AssetRegistry::Impl::FEventContext EventContext;
 
-		if (GatheredDataProcessingLock.TryLock())
+		bool bHasEnteredGatheredDataProcessingLock = false;
+#if WITH_EDITOR
+		bool bTakeOverGather = GuardedData.IsGameThreadTakeOverGatherEachTick();
+		if (!bTakeOverGather)
+#endif
+		{
+			// When we are not trying to block on the gather, we allow the background thread to keep working on
+			// tickgatherer, and we only enter the lock and tickgatherer here on the game thread if the background
+			// thread is not already in the lock
+			bHasEnteredGatheredDataProcessingLock = GatheredDataProcessingLock.TryLock();
+		}
+#if WITH_EDITOR
+		else
+		{
+			// When we want to block on the gather results, we take over TickGatherer from the background thread.
+			// For the GlobalGatherer's side of this race, see TickOnBackgroundThread and the code that calls it in
+			// FAssetDataGatherer::Run.
+			{
+				// First we use an FInterfaceWriteScopeLock with the default High Priority to register ourselves as
+				// waiting on the InterfaceLock.
+				UE::AssetRegistry::FInterfaceWriteScopeLock InterfaceScopeLock(InterfaceLock);
+				// The GlobalGatherer will see that we are waiting on entering the lock and will exit the lock as soon
+				// as possible to allow us to take it. After we take the lock, it will race with us to reenter the
+				// GatheredDataProcessingLock and then enter the InterfaceLock, and will block on the InterfaceLock as
+				// long as we are still holding it.
+				// By requesting pause We tell the GlobalGatherer to leave the GatheredDataProcessingLock and not try
+				// to reenter it until we request resume.
+				GuardedData.RequestPauseBackgroundProcessing();
+				// We drop the InterfaceLock to allow the globalgatherer to continue on if it is waiting on it.
+			}
+			// After dropping the InterfaceLock, we block on the GatheredDataProcessingLock, waiting for the
+			// GlobalGatherer to notice that backgroundprocessing is paused and get out of both of the locks.
+			GatheredDataProcessingLock.Lock();
+			bHasEnteredGatheredDataProcessingLock = true;
+			// We unpause after we finish ticking
+		}
+#endif
+
+		if (bHasEnteredGatheredDataProcessingLock)
 		{
 			LLM_SCOPE(ELLMTag::AssetRegistry);
 			UE::AssetRegistry::FInterfaceWriteScopeLock InterfaceScopeLock(InterfaceLock);
@@ -4674,6 +4716,15 @@ void UAssetRegistryImpl::Tick(float DeltaTime)
 				// Skip the TickGather to deal with the DeferredEvents first
 				bInterruptedOrShouldProcessDeferredEvents = true;
 			}
+			
+#if WITH_EDITOR
+			if (bTakeOverGather)
+			{
+				// As soon as we execute this unpause, the globalgatherer can race to reenter the locks
+				// but it will block entering GatheredDataProcessingLock until we unlock it on the next line.
+				GuardedData.RequestResumeBackgroundProcessing();
+			}
+#endif
 			GatheredDataProcessingLock.Unlock();
 		}
 		else
@@ -4734,6 +4785,7 @@ void FAssetRegistryImpl::UpdateMaxSecondsPerFrame()
 #if WITH_EDITOR
 	bool bGatherOnGameThreadOnly = false;
 	GConfig->GetBool(TEXT("AssetRegistry"), TEXT("GatherOnGameThreadOnly"), bGatherOnGameThreadOnly, GEngineIni);
+	bool bLocalGameThreadTakeOverGatherEachTick = false;
 
 	if (bInitialSearchStarted && !bInitialSearchCompleted)
 	{
@@ -4741,7 +4793,7 @@ void FAssetRegistryImpl::UpdateMaxSecondsPerFrame()
 		GConfig->GetBool(TEXT("AssetRegistry"), TEXT("BlockingInitialLoad"), bBlockingInitialLoad, GEditorPerProjectIni);
 		if (bBlockingInitialLoad)
 		{
-			bGatherOnGameThreadOnly = true;
+			bLocalGameThreadTakeOverGatherEachTick = true;
 			NewMaxSecondsPerFrame = UE::AssetRegistry::Impl::MaxSecondsPerFrameToUseInBlockingInitialLoad;
 			if (MaxSecondsPerFrame < NewMaxSecondsPerFrame)
 			{
@@ -4754,6 +4806,8 @@ void FAssetRegistryImpl::UpdateMaxSecondsPerFrame()
 	{
 		GlobalGatherer->SetGatherOnGameThreadOnly(bGatherOnGameThreadOnly);
 	}
+	SetGameThreadTakeOverGatherEachTick(bLocalGameThreadTakeOverGatherEachTick);
+
 #endif
 	MaxSecondsPerFrame = NewMaxSecondsPerFrame;
 }
diff --git a/Engine/Source/Runtime/AssetRegistry/Private/AssetRegistryImpl.h b/Engine/Source/Runtime/AssetRegistry/Private/AssetRegistryImpl.h
--- a/Engine/Source/Runtime/AssetRegistry/Private/AssetRegistryImpl.h
+++ b/Engine/Source/Runtime/AssetRegistry/Private/AssetRegistryImpl.h
@@ -293,6 +293,14 @@ public:
 	/** Consume any results from the gatherer and return its status */
 	Impl::EGatherStatus TickGatherer(Impl::FTickContext& TickContext);
 
+	/**
+	 * At some times during editor startup, depending on config settings, we might want to block
+	 * on the gather, and not tick the rest of the editor. When we want that to happen, we set
+	 * this flag. This function can be read outside of the lock (we need to do that during Tick).
+	 */
+	bool IsGameThreadTakeOverGatherEachTick() const;
+	void SetGameThreadTakeOverGatherEachTick(bool bValue);
+
 	/** Send a log message with the search statistics. 
 	 *  StartTime is used to report wall clock search time in the case of background scan
 	 */
@@ -648,6 +656,7 @@ private:
 	 * AssetRegistry's lock from another thread.
 	 */
 	uint32 BackgroundTickInterruptionsCount = 0;
+	std::atomic<bool> bGameThreadTakeOverGatherEachTick;
 
 	/** A map of per asset class dependency gatherer called in LoadCalculatedDependencies */
 	TMultiMap<FTopLevelAssetPath, UE::AssetDependencyGatherer::Private::FRegisteredAssetDependencyGatherer*> RegisteredDependencyGathererClasses;
@@ -862,6 +871,16 @@ inline uint32& FAssetRegistryImpl::GetBackgroundTickInterruptionsCount()
 {
 	return BackgroundTickInterruptionsCount;
 }
+
+inline bool FAssetRegistryImpl::IsGameThreadTakeOverGatherEachTick() const
+{
+	return bGameThreadTakeOverGatherEachTick.load(std::memory_order_relaxed);
+}
+
+inline void FAssetRegistryImpl::SetGameThreadTakeOverGatherEachTick(bool bValue)
+{
+	bGameThreadTakeOverGatherEachTick.store(bValue, std::memory_order_relaxed);
+}
 #endif
 
 
