[AutoRTFM] Update TRWScopeLockWithPriority in AssetRegistry to allow for transaction safety.

The high-priority lock tracking mechanism uses an atomic, which now has proper abort handling.

diff --git a/Engine/Source/Runtime/AssetRegistry/Private/AssetRegistry.cpp b/Engine/Source/Runtime/AssetRegistry/Private/AssetRegistry.cpp
--- a/Engine/Source/Runtime/AssetRegistry/Private/AssetRegistry.cpp
+++ b/Engine/Source/Runtime/AssetRegistry/Private/AssetRegistry.cpp
@@ -11,6 +11,7 @@
 #include "AssetRegistryPrivate.h"
 #include "Async/Async.h"
 #include "Async/ParallelFor.h"
+#include "AutoRTFM/AutoRTFM.h"
 #include "Blueprint/BlueprintSupport.h"
 #include "DependsNode.h"
 #include "GenericPlatform/GenericPlatformChunkInstall.h"
@@ -32,6 +33,7 @@
 #include "Misc/ScopeExit.h"
 #include "Misc/ScopeRWLock.h"
 #include "Misc/TrackedActivity.h"
+#include "Misc/TransactionallySafeRWScopeLock.h"
 #include "AssetRegistry/PackageReader.h"
 #include "ProfilingDebugging/MiscTrace.h"
 #include "Serialization/ArrayReader.h"
@@ -192,12 +194,51 @@ namespace UE::AssetRegistry::Impl
 
 namespace UE::AssetRegistry
 {
+
+/** 
+ * Tracks high-priority waiters on a lock, in an thread- and transactionally-safe manner.
+ */
+class FScopedHighPriorityWaitTracker
+{
+public:
+	FScopedHighPriorityWaitTracker(std::atomic<int32>* InCounter, Private::ELockPriority InPriority)
+		: Counter((InPriority == Private::PriorityHigh) ? InCounter : nullptr)
+	{
+		if (Counter)
+		{
+			UE_AUTORTFM_OPEN
+			{
+				Counter->fetch_add(1, std::memory_order_relaxed);
+			};
+			AutoRTFM::PushOnAbortHandler(this, [this]
+			{
+				Counter->fetch_sub(1, std::memory_order_relaxed);
+			});
+		}
+	}
+
+	~FScopedHighPriorityWaitTracker()
+	{
+		if (Counter)
+		{
+			UE_AUTORTFM_OPEN
+			{
+				Counter->fetch_sub(1, std::memory_order_relaxed);
+			};
+			AutoRTFM::PopOnAbortHandler(this);
+		}
+	}
+
+private:
+	std::atomic<int32>* Counter = nullptr;
+};
+
 /** Keeps a FRWLock read-locked while this scope lives */
 /** This is almost a clone of the existing FReadScopeLock and similar types
-	however this adds an extra flag to help the background processing thread 
+	however this adds an extra flag to help the background processing thread
 	know when a higher priority thread would like to gain access to the protected
 	data */
-template<typename TScopeLockType> class TRWScopeLockWithPriority
+template <typename TScopeLockType> class TRWScopeLockWithPriority
 {
 public:
 	UE_NODISCARD_CTOR explicit TRWScopeLockWithPriority(Private::FRWLockWithPriority& InLock, 
@@ -205,15 +246,8 @@ public:
 	: Lock(InLock)
 	, Priority(InPriority)
 	{
-		if (Priority == Private::PriorityHigh)
-		{
-			Lock.HighPriorityWaitersCount.fetch_add(1, std::memory_order_relaxed);
-		}
+		FScopedHighPriorityWaitTracker Tracker(&Lock.HighPriorityWaitersCount, Priority);
 		GuardWrapper.Emplace(Lock);
-		if (Priority == Private::PriorityHigh)
-		{
-			Lock.HighPriorityWaitersCount.fetch_sub(1, std::memory_order_relaxed);
-		}
 	}
 
 	TOptional<TScopeLockType> GuardWrapper;
@@ -230,15 +264,8 @@ public:
 		, Priority(InPriority)
 		, LockType(InLockType)
 	{
-		if (Priority == Private::PriorityHigh)
-		{
-			Lock.HighPriorityWaitersCount.fetch_add(1, std::memory_order_relaxed);
-		}
+		FScopedHighPriorityWaitTracker Tracker(&Lock.HighPriorityWaitersCount, Priority);
 		GuardWrapper.Emplace(Lock, LockType);
-		if (Priority == Private::PriorityHigh)
-		{
-			Lock.HighPriorityWaitersCount.fetch_sub(1, std::memory_order_relaxed);
-		}
 	}
 
 	// NOTE: As the name suggests, this function should be used with caution. 
@@ -250,21 +277,16 @@ public:
 		if (LockType == SLT_ReadOnly)
 		{
 			GuardWrapper.Reset();
-			if (Priority == Private::PriorityHigh)
-			{
-				Lock.HighPriorityWaitersCount.fetch_add(1, std::memory_order_relaxed);
-			}
-			GuardWrapper.Emplace(Lock, SLT_Write);
-			if (Priority == Private::PriorityHigh)
 			{
-				Lock.HighPriorityWaitersCount.fetch_sub(1, std::memory_order_relaxed);
+				FScopedHighPriorityWaitTracker Tracker(&Lock.HighPriorityWaitersCount, Priority);
+				GuardWrapper.Emplace(Lock, SLT_Write);
 			}
 			LockType = SLT_Write;
 		}
 	}
 
 	UE::AssetRegistry::Private::FRWLockWithPriority& Lock;
-	TOptional<FRWScopeLock> GuardWrapper;
+	TOptional<FTransactionallySafeRWScopeLock> GuardWrapper;
 	UE::AssetRegistry::Private::ELockPriority Priority;
 	FRWScopeLockType LockType;
 };
diff --git a/Engine/Source/Runtime/AssetRegistry/Private/AssetRegistryImpl.h b/Engine/Source/Runtime/AssetRegistry/Private/AssetRegistryImpl.h
--- a/Engine/Source/Runtime/AssetRegistry/Private/AssetRegistryImpl.h
+++ b/Engine/Source/Runtime/AssetRegistry/Private/AssetRegistryImpl.h
@@ -10,11 +10,15 @@
 #include "Containers/RingBuffer.h"
 #include "HAL/CriticalSection.h"
 #include "Math/NumericLimits.h"
+#include "Misc/TransactionallySafeRWLock.h"
+#include "Misc/TransactionallySafeRWScopeLock.h"
 #include "ModuleDescriptor.h"
 #include "PackageDependencyData.h"
 #include "Templates/Function.h"
 #include "Templates/UniquePtr.h"
 
+#include <atomic>
+
 // Disable premade asset registry until the target specifies it in UBT
 #ifndef ASSETREGISTRY_ENABLE_PREMADE_REGISTRY_IN_EDITOR
 #define ASSETREGISTRY_ENABLE_PREMADE_REGISTRY_IN_EDITOR 0
@@ -33,9 +37,9 @@ namespace UE::AssetRegistry::Impl { struct FInitializeContext; }
 namespace UE::AssetRegistry::Impl { struct FScanPathContext; }
 namespace UE::AssetRegistry::Impl { struct FTickContext; }
 namespace UE::AssetRegistry::Impl { struct FClassInheritanceContext; }
-namespace UE::AssetRegistry { template<typename TScopeLockType> class TRWScopeLockWithPriority; }
-namespace UE::AssetRegistry { using FInterfaceReadScopeLock = TRWScopeLockWithPriority<FReadScopeLock>; }
-namespace UE::AssetRegistry { using FInterfaceWriteScopeLock = TRWScopeLockWithPriority<FWriteScopeLock>; }
+namespace UE::AssetRegistry { template <typename TScopeLockType> class TRWScopeLockWithPriority; }
+namespace UE::AssetRegistry { using FInterfaceReadScopeLock = TRWScopeLockWithPriority<FTransactionallySafeReadScopeLock>; }
+namespace UE::AssetRegistry { using FInterfaceWriteScopeLock = TRWScopeLockWithPriority<FTransactionallySafeWriteScopeLock>; }
 namespace UE::AssetRegistry { using FInterfaceRWScopeLock = class FRWScopeLockWithPriority; }
 
 #if WITH_EDITOR
@@ -84,10 +88,18 @@ namespace Private
 	 *  to work correctly it must be used with the associated TPriorityScopeLock types 
 	 *  (TScopeLockWithPriority and FRWScopeLockWithPriority)
 	 */
-	class FRWLockWithPriority : public FRWLock
+	class FRWLockWithPriority : public FTransactionallySafeRWLock
 	{
 	public:
-		bool HasWaiters() const { return HighPriorityWaitersCount.load(std::memory_order_relaxed) > 0; }
+		bool HasWaiters() const 
+		{
+			bool Result;
+			UE_AUTORTFM_OPEN
+			{
+				Result = HighPriorityWaitersCount.load(std::memory_order_relaxed) > 0; 
+			};
+			return Result;
+		}
 	
 	private:
 		friend FInterfaceReadScopeLock;
diff --git a/Engine/Source/Runtime/Core/Public/Misc/TransactionallySafeRWLock.h b/Engine/Source/Runtime/Core/Public/Misc/TransactionallySafeRWLock.h
--- a/Engine/Source/Runtime/Core/Public/Misc/TransactionallySafeRWLock.h
+++ b/Engine/Source/Runtime/Core/Public/Misc/TransactionallySafeRWLock.h
@@ -36,7 +36,7 @@
 	//   WriteUnlock     -------------
 	//   -----------     ReadLock      <- Invalid because the transaction can potentially observe side
 	//                                    effects of the open-threads writes!
-struct FTransactionallySafeRWLockDefinition final
+struct FTransactionallySafeRWLockDefinition
 {
 	// Always open because the constructor arguments will create the underlying critical section.
 	UE_AUTORTFM_ALWAYS_OPEN
