Asset Registry - Make sure loaded assets are from a monitored path before processing them
* Temp objects loaded through an instancing context were previously processed by the asset registry

diff --git a/Engine/Source/Runtime/AssetRegistry/Private/AssetRegistry.cpp b/Engine/Source/Runtime/AssetRegistry/Private/AssetRegistry.cpp
--- a/Engine/Source/Runtime/AssetRegistry/Private/AssetRegistry.cpp
+++ b/Engine/Source/Runtime/AssetRegistry/Private/AssetRegistry.cpp
@@ -7744,6 +7744,20 @@ namespace UE::AssetRegistry
 
 void FAssetRegistryImpl::AddLoadedAssetToProcess(const UObject& AssetLoaded)
 {
+	// Make sure the loaded asset is from a monitored path
+	{
+		FString LocalPath;
+		if (!FPackageName::TryConvertLongPackageNameToFilename(AssetLoaded.GetPackage()->GetName(), LocalPath))
+		{
+			return;
+		}
+
+		if (!GlobalGatherer->IsMonitored(LocalPath))
+		{
+			return;
+		}
+	}
+
 	LoadedAssetsToProcess.Add(&AssetLoaded);
 }
 
