AssetRegistry: Add missing null check; GlobalGatherer can be null in some configurations.
#rnx

diff --git a/Engine/Source/Runtime/AssetRegistry/Private/AssetRegistry.cpp b/Engine/Source/Runtime/AssetRegistry/Private/AssetRegistry.cpp
--- a/Engine/Source/Runtime/AssetRegistry/Private/AssetRegistry.cpp
+++ b/Engine/Source/Runtime/AssetRegistry/Private/AssetRegistry.cpp
@@ -7745,6 +7745,7 @@ namespace UE::AssetRegistry
 void FAssetRegistryImpl::AddLoadedAssetToProcess(const UObject& AssetLoaded)
 {
 	// Make sure the loaded asset is from a monitored path
+	if (GlobalGatherer.IsValid())
 	{
 		FString LocalPath;
 		if (!FPackageName::TryConvertLongPackageNameToFilename(AssetLoaded.GetPackage()->GetName(), LocalPath))
