Speed up UpdateRedirectCollector

The expensive part is GetRedirectedObjectPath, which we can do in parallel. In the content worker with a 32 core setup, this reduces the time taken from 2.8s to 1.2s. The sad part is that GetRedirectedObjectPath is taking a lock under the hood, so there is still a lot of contention (in particular in StaticFind).

#rnx

diff --git a/Engine/Source/Runtime/AssetRegistry/Private/AssetRegistry.cpp b/Engine/Source/Runtime/AssetRegistry/Private/AssetRegistry.cpp
--- a/Engine/Source/Runtime/AssetRegistry/Private/AssetRegistry.cpp
+++ b/Engine/Source/Runtime/AssetRegistry/Private/AssetRegistry.cpp
@@ -7845,18 +7845,35 @@ void FAssetRegistryImpl::PushProcessLoadedAssetsBatch(Impl::FEventContext& Event
 void FAssetRegistryImpl::UpdateRedirectCollector()
 {
 	// Look for all redirectors in the AssetRegistry
+	TArray<FSoftObjectPath> Paths;
 	State.EnumerateAssetsByClassPathName(UE::AssetRegistry::GetClassPathObjectRedirector(),
-		[this](const FAssetData* AssetData)
+		[this, &Paths](const FAssetData* AssetData)
 		{
-			FSoftObjectPath Source = AssetData->GetSoftObjectPath();
-			FSoftObjectPath Destination = GetRedirectedObjectPath(Source, nullptr, nullptr, /*bNeedsScanning*/ false);
-
-			if (Destination != Source)
-			{
-				GRedirectCollector.AddAssetPathRedirection(Source, Destination);
-			}
+			Paths.Add(AssetData->GetSoftObjectPath());
 			return true; // Keep iterating
 		});
+
+	TArray<FSoftObjectPath> RedirectedPaths;
+	RedirectedPaths.SetNumZeroed(Paths.Num());
+	ParallelFor(Paths.Num(), [&RedirectedPaths, &Paths, this](int32 Index)
+	{
+		const FSoftObjectPath& Source = Paths[Index];
+		FSoftObjectPath Destination = GetRedirectedObjectPath(Source, nullptr, nullptr, /*bNeedsScanning*/ false);
+		if (Destination != Source)
+		{
+			RedirectedPaths[Index] = Destination;
+		}
+	});
+
+	for (int32_t AssetIndex = 0; AssetIndex < Paths.Num(); AssetIndex++)
+	{
+		FSoftObjectPath& Destination = RedirectedPaths[AssetIndex];
+		if (RedirectedPaths[AssetIndex].IsNull())
+		{
+			continue;
+		}
+		GRedirectCollector.AddAssetPathRedirection(Paths[AssetIndex], Destination);
+	}
 }
 
 }
