Add PackageLocation to FAssetPackageData to allow the AssetRegistry to inform where AssetData is coming from. This is used to accelerate DoesPackageExist calls when the asset registry has been initialized to alleviate the need to read from disk if the registry already knows about the data. In particular, cooked editor cases are optimized since previously the asset registry cache was ignored if the IOStore contained any data at all.

Opening a large project with a cooked editor now sees DoesPackageExist time decrease by ~8x, from ~17s to ~2s

- FAssetPackageData is set to have a package location of IoDispatcher when cooking and is set to Filesystem when found on disk as part of the AssetRegistry filesystem asset gathering
- Lookups for FAssetPackageData while the registry is scanning are treated as having a package location of Unknown unless the already scanned asset is located as on disk/Filesystem
- Added a new `IsGathering` method to IAssetRegistry which expands upon IsLoadingAssets(). IsLoadingAssets() only reports when the asset registry is searching for assets initially. Subsequent scans for assets will not be detected, whereas IsGathering will return true whenever the asset registry is scanning for assets (such as after a new mount point has been added)

#rnx

diff --git a/Engine/Plugins/Experimental/EditorTelemetry/Source/Private/EditorTelemetry.cpp b/Engine/Plugins/Experimental/EditorTelemetry/Source/Private/EditorTelemetry.cpp
--- a/Engine/Plugins/Experimental/EditorTelemetry/Source/Private/EditorTelemetry.cpp
+++ b/Engine/Plugins/Experimental/EditorTelemetry/Source/Private/EditorTelemetry.cpp
@@ -958,6 +958,7 @@ void FEditorTelemetry::StartSession()
 				{ TEXT("DurationSeconds"), Data.Duration },
 				{ TEXT("InitialSearchStarted"), Data.bInitialSearchStarted },
 				{ TEXT("InitialSearchCompleted"), Data.bInitialSearchCompleted },
+				{ TEXT("AdditionalMountSearchInProgress"), Data.bAdditionalMountSearchInProgress },
 			});
 		});
 		Router.OnTelemetry<FGatherTelemetry>([this](const FGatherTelemetry& Data){
diff --git a/Engine/Source/Editor/UnrealEd/Private/CookOnTheFlyServer.cpp b/Engine/Source/Editor/UnrealEd/Private/CookOnTheFlyServer.cpp
--- a/Engine/Source/Editor/UnrealEd/Private/CookOnTheFlyServer.cpp
+++ b/Engine/Source/Editor/UnrealEd/Private/CookOnTheFlyServer.cpp
@@ -10788,16 +10788,24 @@ void UCookOnTheFlyServer::CookByTheBookFinishedInternal()
 		{
 			UE_SCOPED_HIERARCHICAL_COOKTIMER(GeneratorPreSave);
 
-			// Add the package hashes to the relevant AssetPackageDatas.
-			// PackageHashes are gated by requiring UPackage::WaitForAsyncFileWrites(), which is called above.
+			// Make changes to all AssetPackageDatas to record data from the cook
+			// PackageHashes are guaranteed finished by UPackage::WaitForAsyncFileWrites(), which is called above.
 			TMap<FName, TRefCountPtr<FPackageHashes>>& AllPackageHashes = SaveContext.PackageWriter->GetPackageHashes();
 			for (TPair<FName, TRefCountPtr<FPackageHashes>>& HashSet : AllPackageHashes)
 			{
 				FAssetPackageData* AssetPackageData = Generator.GetAssetPackageData(HashSet.Key);
-				TRefCountPtr<FPackageHashes>& PackageHashes = HashSet.Value;
 
+				// Add the package hashes to the relevant AssetPackageDatas.
+				TRefCountPtr<FPackageHashes>& PackageHashes = HashSet.Value;
 				AssetPackageData->CookedHash = PackageHashes->PackageHash;
 				Move(AssetPackageData->ChunkHashes, PackageHashes->ChunkHashes);
+
+				// Mark that the Assets in the cooked files are from the IoDispatcher.
+				// This assumes that all cooked files are loaded by the IoDispatcher, which is the case for the primary 
+				// supported workflow, but it is possible that a developer could choose to stage to loose files, in which case
+				// IoDispatcher would be incorrect.
+				// TODO: Modify the package location in the cooked asset registry when staging to loose files.
+				AssetPackageData->SetPackageLocation(FPackageName::EPackageLocationFilter::IoDispatcher);
 			}
 
 			Generator.PreSave(CookedPackageNames);
diff --git a/Engine/Source/Editor/UnrealEd/Private/Cooker/CookGenerationHelper.cpp b/Engine/Source/Editor/UnrealEd/Private/Cooker/CookGenerationHelper.cpp
--- a/Engine/Source/Editor/UnrealEd/Private/Cooker/CookGenerationHelper.cpp
+++ b/Engine/Source/Editor/UnrealEd/Private/Cooker/CookGenerationHelper.cpp
@@ -1302,6 +1302,8 @@ void FGenerationHelper::FinishGeneratedPlatformSave(FPackageData& PackageData,
 	OutGeneratedResult.AssetPackageData.FileVersionUE = GPackageFileUEVersion;
 	OutGeneratedResult.AssetPackageData.FileVersionLicenseeUE = GPackageFileLicenseeUEVersion;
 	OutGeneratedResult.AssetPackageData.SetIsLicenseeVersion(FEngineVersion::Current().IsLicenseeVersion());
+	// Currently assumes IoDispatcher however should we cook to loose files in the future this would need updating
+	OutGeneratedResult.AssetPackageData.SetPackageLocation(FPackageName::EPackageLocationFilter::IoDispatcher);
 	OutGeneratedResult.AssetPackageData.Extension = FPackagePath::ParseExtension(
 		WriteToString<256>(PackageData.GetFileName()));
 	OutGeneratedResult.AssetPackageData.SetPackageSavedHash(Info->PackageHash);
diff --git a/Engine/Source/Runtime/AssetRegistry/Private/AssetDataGatherer.cpp b/Engine/Source/Runtime/AssetRegistry/Private/AssetDataGatherer.cpp
--- a/Engine/Source/Runtime/AssetRegistry/Private/AssetDataGatherer.cpp
+++ b/Engine/Source/Runtime/AssetRegistry/Private/AssetDataGatherer.cpp
@@ -1702,6 +1702,12 @@ void FAssetDataDiscovery::OnInitialSearchCompleted()
 	Cache.Shutdown();
 }
 
+void FAssetDataDiscovery::OnAdditionalMountSearchCompleted()
+{
+	// After the initial search completed, OnInitialSearchCompleted cleared 
+	// out the data for the cache so we no longer have enough data to save it.
+}
+
 void FAssetDataDiscovery::TickInternal(bool bTickAll)
 {
 	LLM_SCOPE(ELLMTag::AssetRegistry);
@@ -3608,6 +3614,7 @@ FAssetDataGatherer::FAssetDataGatherer(const TArray<FString>& InLongPackageNames
 	, bFirstTickAfterIdle(true)
 	, bFinishedInitialDiscovery(false)
 	, bIsInitialSearchCompleted(false)
+	, bIsAdditionalMountSearchInProgress(false)
 	, bGatherOnGameThreadOnly(false)
 	, WaitBatchCount(-1)
 	, LastCacheSaveNumUncachedAssetFiles(0)
@@ -3680,6 +3687,16 @@ void FAssetDataGatherer::OnInitialSearchCompleted()
 	bIsInitialSearchCompleted.store(true, std::memory_order_relaxed);
 }
 
+void FAssetDataGatherer::OnAdditionalMountSearchCompleted()
+{
+	if (Discovery)
+	{
+		Discovery->OnAdditionalMountSearchCompleted();
+	}
+
+	bIsAdditionalMountSearchInProgress.store(false, std::memory_order_relaxed);
+}
+
 void FAssetDataGatherer::StartAsync()
 {
 	if (bAsyncEnabled && !Thread)
@@ -3730,7 +3747,11 @@ uint32 FAssetDataGatherer::Run()
 			if ((Status != UE::AssetRegistry::Impl::EGatherStatus::TickActiveGatherActive)
 				&& (Status != UE::AssetRegistry::Impl::EGatherStatus::TickActiveGatherIdle))
 			{
-				TRACE_CPUPROFILER_EVENT_SCOPE_STR_CONDITIONAL("FAssetDataGatherer Sleep", !bIsInitialSearchCompleted.load(std::memory_order_relaxed));
+				const bool bInitialSearchCompleted = bIsInitialSearchCompleted.load(std::memory_order_relaxed);
+				const bool bAdditionalMountSearchInProgress = bIsAdditionalMountSearchInProgress.load(std::memory_order_relaxed);
+				const bool bShouldLog = !bInitialSearchCompleted || bAdditionalMountSearchInProgress;
+
+				TRACE_CPUPROFILER_EVENT_SCOPE_STR_CONDITIONAL("FAssetDataGatherer Sleep", bShouldLog);
 				FPlatformProcess::Sleep(bLocalIdle ? IdleSleepTime : PausedSleepTime);
 			}
 		}
diff --git a/Engine/Source/Runtime/AssetRegistry/Private/AssetDataGatherer.h b/Engine/Source/Runtime/AssetRegistry/Private/AssetDataGatherer.h
--- a/Engine/Source/Runtime/AssetRegistry/Private/AssetDataGatherer.h
+++ b/Engine/Source/Runtime/AssetRegistry/Private/AssetDataGatherer.h
@@ -83,6 +83,7 @@ public:
 	virtual ~FAssetDataGatherer();
 
 	void OnInitialSearchCompleted();
+	void OnAdditionalMountSearchCompleted();
 
 	// Controlling Async behavior
 
@@ -250,6 +251,11 @@ public:
 	void SetGatherOnGameThreadOnly(bool bValue);
 	bool IsGatherOnGameThreadOnly() const;
 
+	/*
+	* Mark that the gatherer is in the process of handling an additional search.
+	*/
+	void SetIsAdditionalMountSearchInProgress(bool bIsInProgress);
+
 private:
 	enum class ETickResult
 	{
@@ -475,6 +481,8 @@ private:
 	bool bFinishedInitialDiscovery;
 	/** True if OnInitialSearchCompleted has been called. */
 	std::atomic<bool> bIsInitialSearchCompleted;
+	/** True if we have begun discovering files after the initial search. */
+	std::atomic<bool> bIsAdditionalMountSearchInProgress;
 	std::atomic<bool> bGatherOnGameThreadOnly;
 
 	// Variable section for variables that are read/writable only within TickLock.
@@ -527,3 +535,8 @@ inline bool FAssetDataGatherer::IsGatherOnGameThreadOnly() const
 {
 	return bGatherOnGameThreadOnly.load(std::memory_order_relaxed);
 }
+
+inline void FAssetDataGatherer::SetIsAdditionalMountSearchInProgress(bool bIsInProgress)
+{
+	bIsAdditionalMountSearchInProgress.store(bIsInProgress, std::memory_order_relaxed);
+}
\ No newline at end of file
diff --git a/Engine/Source/Runtime/AssetRegistry/Private/AssetDataGathererPrivate.h b/Engine/Source/Runtime/AssetRegistry/Private/AssetDataGathererPrivate.h
--- a/Engine/Source/Runtime/AssetRegistry/Private/AssetDataGathererPrivate.h
+++ b/Engine/Source/Runtime/AssetRegistry/Private/AssetDataGathererPrivate.h
@@ -608,6 +608,7 @@ public:
 	void WaitForIdle(double EndTimeSeconds);
 	bool IsIdle() const;
 	void OnInitialSearchCompleted();
+	void OnAdditionalMountSearchCompleted();
 
 	/** Optionally set some scan properties for the given paths and then wait for their scans to finish. */
 	void SetPropertiesAndWait(TArrayView<FPathExistence> QueryPaths, bool bAddToAllowList, bool bForceRescan,
diff --git a/Engine/Source/Runtime/AssetRegistry/Private/AssetRegistry.cpp b/Engine/Source/Runtime/AssetRegistry/Private/AssetRegistry.cpp
--- a/Engine/Source/Runtime/AssetRegistry/Private/AssetRegistry.cpp
+++ b/Engine/Source/Runtime/AssetRegistry/Private/AssetRegistry.cpp
@@ -997,6 +997,7 @@ void FAssetRegistryImpl::Initialize(Impl::FInitializeContext& Context)
 
 	bInitialSearchStarted = false;
 	bInitialSearchCompleted.store(true, std::memory_order_relaxed);
+	bAdditionalMountSearchInProgress.store(false, std::memory_order_relaxed);
 #if WITH_EDITOR
 	SetGameThreadTakeOverGatherEachTick(false);
 #endif
@@ -1971,12 +1972,19 @@ void FAssetRegistryImpl::SearchAllAssets(Impl::FEventContext& EventContext,
 	{
 		return;
 	}
+
 	if (!bInitialSearchStarted)
 	{
 		InitialSearchStartTime = FPlatformTime::Seconds();
 		bInitialSearchStarted = true;
 		bInitialSearchCompleted.store(false, std::memory_order_relaxed);
-		UpdateMaxSecondsPerFrame();
+		UpdateMaxSecondsPerFrame(); // MaxSecondsPerFrame only depends on initial search
+	}
+	else if (!IsGathering())
+	{
+		AdditionalMountSearchStartTime = FPlatformTime::Seconds();
+		bAdditionalMountSearchInProgress.store(true, std::memory_order_relaxed);
+		GlobalGatherer->SetIsAdditionalMountSearchInProgress(true);
 	}
 
 	FAssetDataGatherer& Gatherer = *GlobalGatherer;
@@ -2043,7 +2051,7 @@ void UAssetRegistryImpl::WaitForCompletion()
 	}
 
 #if WITH_EDITOR
-	if (bInitialSearchStarted && !bInitialSearchCompleted)
+	if (IsGathering())
 	{
 		// If we do need to wait, then tick the DirectoryWatcher so we have the most up to date information.
 		// This is also important because we ignore rescan events from the directory watcher if they are sent
@@ -2186,7 +2194,7 @@ void UAssetRegistryImpl::WaitForPackage(const FString& PackageName)
 	{
 		LLM_SCOPE(ELLMTag::AssetRegistry);
 		UE::AssetRegistry::FInterfaceWriteScopeLock InterfaceScopeLock(InterfaceLock);
-		if (GuardedData.IsLoadingAssets())
+		if (GuardedData.IsGathering())
 		{
 			FString LocalPath;
 			if (FPackageName::TryConvertLongPackageNameToFilename(PackageName, LocalPath))
@@ -2830,12 +2838,12 @@ PRAGMA_ENABLE_DEPRECATION_WARNINGS;
 UE::AssetRegistry::EExists UAssetRegistryImpl::TryGetAssetByObjectPath(const FSoftObjectPath& ObjectPath, FAssetData& OutAssetData) const
 {
 	UE::AssetRegistry::FInterfaceReadScopeLock InterfaceScopeLock(InterfaceLock);
-	bool bAssetRegistryReady = GuardedData.IsInitialSearchStarted() && GuardedData.IsInitialSearchCompleted();
+	const bool bAssetRegistryHasFullKnowledge = GuardedData.IsSearchAllAssets() && !GuardedData.IsGathering();
 	const FAssetRegistryState& State = GuardedData.GetState();
 	const FAssetData* FoundData = State.GetAssetByObjectPath(ObjectPath);
 	if (!FoundData)
 	{
-		if (!bAssetRegistryReady)
+		if (!bAssetRegistryHasFullKnowledge)
 		{
 			return UE::AssetRegistry::EExists::Unknown;
 		}
@@ -2854,17 +2862,25 @@ UE::AssetRegistry::EExists UAssetRegistryImpl::TryGetAssetPackageData(const FNam
 UE::AssetRegistry::EExists UAssetRegistryImpl::TryGetAssetPackageData(const FName PackageName, FAssetPackageData& OutAssetPackageData, FName& OutCorrectCasePackageName) const
 {
 	UE::AssetRegistry::FInterfaceReadScopeLock InterfaceScopeLock(InterfaceLock);
-	bool bAssetRegistryReady = GuardedData.IsInitialSearchStarted() && GuardedData.IsInitialSearchCompleted();
+	const bool bLocalIsGathering = GuardedData.IsGathering();
+	const bool bAssetRegistryHasFullKnowledge = GuardedData.IsSearchAllAssets() && !bLocalIsGathering;
 	const FAssetRegistryState& State = GuardedData.GetState();
 	const FAssetPackageData* FoundData = State.GetAssetPackageData(PackageName, OutCorrectCasePackageName);
 	if (!FoundData)
 	{
-		if (!bAssetRegistryReady)
+		if (!bAssetRegistryHasFullKnowledge)
 		{
 			return UE::AssetRegistry::EExists::Unknown;
 		}
 		return UE::AssetRegistry::EExists::DoesNotExist;
 	}
+	// Currently when we cook we mark FAssetPackageData as being from the IoDispatcher however that isn't true until the content
+	// is staged. The correct location will be determined once the registry scan is completed, but until then we should ignore any 
+	// IoDispatcher located results since they might be stale and about to be overwritten by filesystem data found during scanning.
+	else if (bLocalIsGathering && FoundData->GetPackageLocation() == FPackageName::EPackageLocationFilter::IoDispatcher)
+	{
+		return UE::AssetRegistry::EExists::Unknown;
+	}
 	OutAssetPackageData = *FoundData;
 	return UE::AssetRegistry::EExists::Exists;
 }
@@ -3314,7 +3330,7 @@ FSoftObjectPath UAssetRegistryImpl::GetRedirectedObjectPath(const FSoftObjectPat
 	// In that case, we can skip further scanning while looking for a redirected path
 	{
 		UE::AssetRegistry::FInterfaceReadScopeLock InterfaceScopeLock(InterfaceLock);
-		if (GuardedData.IsSearchAllAssets() && GuardedData.IsInitialSearchCompleted())
+		if (GuardedData.IsSearchAllAssets() && !GuardedData.IsGathering())
 		{
 			return GuardedData.GetRedirectedObjectPath(ObjectPath, nullptr, nullptr, /*bNeedsScanning*/ false);
 		}
@@ -4181,6 +4197,7 @@ void UAssetRegistryImpl::ScanPathsSynchronousInternal(const TArray<FString>& InD
 
 	bool bInitialSearchStarted;
 	bool bInitialSearchCompleted;
+	bool bAdditionalMountSearchInProgress;
 	{
 		LLM_SCOPE(ELLMTag::AssetRegistry);
 		UE::AssetRegistry::FInterfaceWriteScopeLock InterfaceScopeLock(InterfaceLock);
@@ -4188,6 +4205,7 @@ void UAssetRegistryImpl::ScanPathsSynchronousInternal(const TArray<FString>& InD
 
 		bInitialSearchStarted = GuardedData.IsInitialSearchStarted();
 		bInitialSearchCompleted = GuardedData.IsInitialSearchCompleted();
+		bAdditionalMountSearchInProgress = GuardedData.IsAdditionalMountSearchInProgress();
 		// make sure any outstanding async preload is complete
 		GuardedData.ConditionalLoadPremadeAssetRegistry(*this, EventContext, InterfaceScopeLock);
 		GuardedData.ScanPathsSynchronous(Context);
@@ -4227,6 +4245,7 @@ void UAssetRegistryImpl::ScanPathsSynchronousInternal(const TArray<FString>& InD
 	Telemetry.Duration = Duration;
 	Telemetry.bInitialSearchStarted = bInitialSearchStarted;
 	Telemetry.bInitialSearchCompleted = bInitialSearchCompleted;
+	Telemetry.bAdditionalMountSearchInProgress = bAdditionalMountSearchInProgress;
 	FTelemetryRouter::Get().ProvideTelemetry(Telemetry);
 	UE_LOG(LogAssetRegistry, Verbose, TEXT("ScanPathsSynchronous completed scanning %s to find %d assets in %0.4f seconds"), *PathsString,
 		Context.NumFoundAssets, Duration);
@@ -4303,7 +4322,6 @@ void UAssetRegistryImpl::AssetDeleted(UObject* DeletedAsset)
 	{
 		UPackage* DeletedObjectPackage = DeletedAsset->GetOutermost();
 		bool bIsEmptyPackage = DeletedObjectPackage != nullptr && UPackage::IsEmptyPackage(DeletedObjectPackage, DeletedAsset);
-		bool bInitialSearchCompleted = false;
 
 		bool bShouldSkipAsset;
 		{
@@ -4324,13 +4342,11 @@ void UAssetRegistryImpl::AssetDeleted(UObject* DeletedAsset)
 					MetaData->ClearFlags(RF_Standalone);
 				}
 			}
-			bInitialSearchCompleted = GuardedData.IsInitialSearchCompleted();
 			bShouldSkipAsset = GuardedData.ShouldSkipAsset(DeletedAsset);
 		}
 
 #if WITH_EDITOR
-		if (bInitialSearchCompleted && FAssetData::IsRedirector(DeletedAsset))
-
+		if (!IsGathering() && FAssetData::IsRedirector(DeletedAsset))
 		{
 			// Need to remove from GRedirectCollector
 			GRedirectCollector.RemoveAssetPathRedirection(FSoftObjectPath::ConstructFromObject(DeletedAsset));
@@ -4564,6 +4580,11 @@ bool UAssetRegistryImpl::IsLoadingAssets() const
 	return GuardedData.IsLoadingAssets();
 }
 
+bool UAssetRegistryImpl::IsGathering() const
+{
+	return GuardedData.IsGathering();
+}
+
 namespace UE::AssetRegistry
 {
 
@@ -4572,6 +4593,11 @@ bool FAssetRegistryImpl::IsLoadingAssets() const
 	return !IsInitialSearchCompleted();
 }
 
+bool FAssetRegistryImpl::IsGathering() const
+{
+	return !IsInitialSearchCompleted() || IsAdditionalMountSearchInProgress();
+}
+
 }
 
 UE::AssetRegistry::Impl::EGatherStatus UAssetRegistryImpl::TickOnBackgroundThread()
@@ -4810,7 +4836,9 @@ void FAssetRegistryImpl::UpdateMaxSecondsPerFrame()
 	GConfig->GetBool(TEXT("AssetRegistry"), TEXT("GatherOnGameThreadOnly"), bGatherOnGameThreadOnly, GEngineIni);
 	bool bLocalGameThreadTakeOverGatherEachTick = false;
 
-	if (bInitialSearchStarted && !bInitialSearchCompleted)
+	// Skip this block even if we're still gathering as part of AdditionalMountSearch; 
+	// this block applies only during initial search
+	if (bInitialSearchStarted && !IsInitialSearchCompleted())
 	{
 		bool bBlockingInitialLoad;
 		GConfig->GetBool(TEXT("AssetRegistry"), TEXT("BlockingInitialLoad"), bBlockingInitialLoad, GEditorPerProjectIni);
@@ -4837,7 +4865,7 @@ void FAssetRegistryImpl::UpdateMaxSecondsPerFrame()
 
 Impl::EGatherStatus FAssetRegistryImpl::TickGatherer(Impl::FTickContext& TickContext)
 {
-	TRACE_CPUPROFILER_EVENT_SCOPE_STR_CONDITIONAL("FAssetRegistryImpl::TickGatherer", !IsInitialSearchCompleted());
+	TRACE_CPUPROFILER_EVENT_SCOPE_STR_CONDITIONAL("FAssetRegistryImpl::TickGatherer", IsGathering());
 
 	using namespace UE::AssetRegistry::Impl;
 	FEventContext& EventContext = TickContext.EventContext;
@@ -5243,6 +5271,15 @@ Impl::EGatherStatus FAssetRegistryImpl::TickGatherer(Impl::FTickContext& TickCon
 				OnInitialSearchCompleted(EventContext);
 			}
 		}
+		else if (IsAdditionalMountSearchInProgress())
+		{
+			bool bCanCompleteAdditionalMountSearch = bLocalIsInGameThread && TickContext.bHandleCompletion;
+			if (bCanCompleteAdditionalMountSearch)
+			{
+				RecordTimer(); // OnAdditionalMountSearchComplete reads data set by RecordTimer
+				OnAdditionalMountSearchCompleted(EventContext);
+			}
+		}
 	}
 
 	UpdateStatus();
@@ -5285,6 +5322,24 @@ void FAssetRegistryImpl::OnInitialSearchCompleted(Impl::FEventContext& EventCont
 	UpdateMaxSecondsPerFrame();
 }
 
+void FAssetRegistryImpl::OnAdditionalMountSearchCompleted(Impl::FEventContext& EventContext)
+{
+#if WITH_EDITOR
+	// update redirectors
+	UpdateRedirectCollector();
+#endif
+
+	// Handle any deferred loading operations
+	SetPerformanceMode(Impl::EPerformanceMode::MostlyStatic);
+
+	LogSearchDiagnostics(AdditionalMountSearchStartTime);
+	TRACE_END_REGION(TEXT("Asset Registry Scan")); // Matching TRACE_BEGIN_REGION in OnContentPathMounted
+
+	GlobalGatherer->OnAdditionalMountSearchCompleted();
+
+	bAdditionalMountSearchInProgress.store(false, std::memory_order_relaxed);
+}
+
 void FAssetRegistryImpl::LogSearchDiagnostics(double StartTime)
 {
 	FAssetGatherDiagnostics Diagnostics = GlobalGatherer->GetDiagnostics();
@@ -7326,11 +7381,13 @@ void UAssetRegistryImpl::OnDirectoryChanged(const TArray<FFileChangeData>& FileC
 	UE::AssetRegistry::Impl::FEventContext EventContext;
 	bool bInitialSearchStarted;
 	bool bInitialSearchCompleted;
+	bool bAdditionalMountSearchInProgress;
 	{
 		LLM_SCOPE(ELLMTag::AssetRegistry);
 		UE::AssetRegistry::FInterfaceWriteScopeLock InterfaceScopeLock(InterfaceLock);
 		bInitialSearchStarted = GuardedData.IsInitialSearchStarted();
 		bInitialSearchCompleted = GuardedData.IsInitialSearchCompleted();
+		bAdditionalMountSearchInProgress = GuardedData.IsAdditionalMountSearchInProgress();
 		UE::AssetRegistry::Impl::FClassInheritanceContext InheritanceContext;
 		UE::AssetRegistry::Impl::FClassInheritanceBuffer InheritanceBuffer;
 		GetInheritanceContextWithRequiredLock(InterfaceScopeLock, InheritanceContext, InheritanceBuffer);
@@ -7343,6 +7400,7 @@ void UAssetRegistryImpl::OnDirectoryChanged(const TArray<FFileChangeData>& FileC
 		FPlatformTime::Seconds() - StartTime,
 		bInitialSearchStarted,
 		bInitialSearchCompleted,
+		bAdditionalMountSearchInProgress,
 	});
 }
 
@@ -7359,7 +7417,7 @@ void FAssetRegistryImpl::OnDirectoryChanged(Impl::FEventContext& EventContext,
 	{
 		if (FileChangesProcessed[FileIdx].Action == FFileChangeData::FCA_RescanRequired)
 		{
-			if (bInitialSearchStarted && !IsInitialSearchCompleted())
+			if (bInitialSearchStarted && !IsInitialSearchCompleted()) // Initial search only
 			{
 				// Ignore rescan request during initial scan as it is probably caused by the scan itself
 				UE_LOG(LogAssetRegistry, Log, TEXT("FAssetRegistry ignoring rescan request for %s during startup"), *FileChangesProcessed[FileIdx].Filename);
@@ -8066,6 +8124,14 @@ void FAssetRegistryImpl::OnContentPathMounted(Impl::FEventContext& EventContext,
 		}
 		else
 		{
+			if (!IsGathering())
+			{
+				TRACE_BEGIN_REGION(TEXT("Asset Registry Scan")); // Matching TRACE_END_REGION in OnAdditionalMountSearchCompleted
+				AdditionalMountSearchStartTime = FPlatformTime::Seconds();
+				bAdditionalMountSearchInProgress.store(true, std::memory_order_relaxed);
+				GlobalGatherer->SetIsAdditionalMountSearchInProgress(true);
+			}
+
 			GlobalGatherer->AddMountPoint(FileSystemPath, InAssetPath);
 			GlobalGatherer->SetIsOnAllowList(FileSystemPath, true);
 		}
diff --git a/Engine/Source/Runtime/AssetRegistry/Private/AssetRegistry.h b/Engine/Source/Runtime/AssetRegistry/Private/AssetRegistry.h
--- a/Engine/Source/Runtime/AssetRegistry/Private/AssetRegistry.h
+++ b/Engine/Source/Runtime/AssetRegistry/Private/AssetRegistry.h
@@ -202,6 +202,7 @@ public:
 	virtual FFileLoadProgressUpdatedEvent& OnFileLoadProgressUpdated() override;
 
 	virtual bool IsLoadingAssets() const override;
+	virtual bool IsGathering() const override;
 	virtual bool ShouldUpdateDiskCacheAfterLoad() const override
 	{
 #if WITH_EDITORONLY_DATA
diff --git a/Engine/Source/Runtime/AssetRegistry/Private/AssetRegistryImpl.h b/Engine/Source/Runtime/AssetRegistry/Private/AssetRegistryImpl.h
--- a/Engine/Source/Runtime/AssetRegistry/Private/AssetRegistryImpl.h
+++ b/Engine/Source/Runtime/AssetRegistry/Private/AssetRegistryImpl.h
@@ -263,6 +263,7 @@ public:
 		FAssetRegistryState::EInitializationMode Mode = FAssetRegistryState::EInitializationMode::Append, bool bEmitAssetEvents = false);
 	void GetAllocatedSize(bool bLogDetailed, SIZE_T& StateSize, SIZE_T& StaticSize, SIZE_T& SearchSize) const;
 	bool IsLoadingAssets() const;
+	bool IsGathering() const;
 	void SetManageReferences(const TMultiMap<FAssetIdentifier, FAssetIdentifier>& ManagerMap,
 		bool bClearExisting, UE::AssetRegistry::EDependencyCategory RecurseType,
 		TSet<FDependsNode*>& ExistingManagedNodes, IAssetRegistry::ShouldSetManagerPredicate ShouldSetManager);
@@ -376,6 +377,7 @@ public:
 	bool IsTempCachingEnabled() const { return bIsTempCachingEnabled; }
 	bool IsTempCachingAlwaysEnabled() const { return bIsTempCachingAlwaysEnabled; }
 	bool IsInitialSearchStarted() const { return bInitialSearchStarted; }
+	bool IsAdditionalMountSearchInProgress() const { return bAdditionalMountSearchInProgress.load(std::memory_order_relaxed); }
 	bool IsSearchAllAssets() const { return bSearchAllAssets; }
 	Impl::FClassInheritanceBuffer& GetTempCachedInheritanceBuffer() { return TempCachedInheritanceBuffer; }
 	uint64 GetSavedGeneratorClassesVersionNumber() const { return SavedGeneratorClassesVersionNumber; }
@@ -520,6 +522,7 @@ private:
 	/** Add MountPoints of all AssetDatas currently registered in this->State to the list of PersistentMountPoints. */
 	void UpdatePersistentMountPoints();
 	void OnInitialSearchCompleted(Impl::FEventContext& EventContext);
+	void OnAdditionalMountSearchCompleted(Impl::FEventContext& EventContext);
 
 	/** This function is not called and, currently, will always return 'true'. See the commented out call in
 	 *	FAssetDataGatherer::TickInternal for further details on how it could be used. 
@@ -600,12 +603,18 @@ private:
 
 	/** Time the initial async search was started */
 	double InitialSearchStartTime = 0.0f;
+	/** Time the additional mount async search was started */
+	double AdditionalMountSearchStartTime = 0.0f;
 	/** Flag to indicate if we used an initial async search */
 	bool bInitialSearchStarted;
 	/** Flag to indicate if the initial background search has completed. All access are relaxed because
 	 *  the actual search data can only be accessed under a proper lock.
 	 */
 	std::atomic<bool> bInitialSearchCompleted;
+	/** Flag to indicate if an additional mount background search has started after the initial search.
+	 * All access are relaxed because the actual search data can only be accessed under a proper lock.
+	 */
+	std::atomic<bool> bAdditionalMountSearchInProgress;
 	/**
 	 * Flag to indicate PreloadingComplete; finishing the background search is blocked until preloading complete
 	 * because preloading can add assets.
diff --git a/Engine/Source/Runtime/AssetRegistry/Private/AssetRegistryState.cpp b/Engine/Source/Runtime/AssetRegistry/Private/AssetRegistryState.cpp
--- a/Engine/Source/Runtime/AssetRegistry/Private/AssetRegistryState.cpp
+++ b/Engine/Source/Runtime/AssetRegistry/Private/AssetRegistryState.cpp
@@ -442,17 +442,24 @@ void FAssetRegistryState::InitializeFromExisting(const FAssetDataMap& AssetDataM
 			if (Pair.Value)
 			{
 				// Only add if also in asset data map, or script package
+				FAssetPackageData* NewData = nullptr;
 				if (bIsScriptPackage)
 				{
 					ScriptPackages.Add(Pair.Key);
-
-					FAssetPackageData* NewData = CreateOrGetAssetPackageData(Pair.Key);
-					*NewData = *Pair.Value;
+					NewData = CreateOrGetAssetPackageData(Pair.Key);
 				}
 				else if (CachedAssetsByPackageName.Find(Pair.Key))
 				{
-					FAssetPackageData* NewData = CreateOrGetAssetPackageData(Pair.Key);
+					NewData = CreateOrGetAssetPackageData(Pair.Key);
+				}
+
+				if (NewData)
+				{
+					// Add the new location to any existing location as it's possible we 
+					// have the same content available from more than one location.
+					FPackageName::EPackageLocationFilter OriginalLocation = NewData->GetPackageLocation();
 					*NewData = *Pair.Value;
+					NewData->SetPackageLocation(FPackageName::EPackageLocationFilter(uint8(NewData->GetPackageLocation()) | uint8(OriginalLocation)));
 				}
 			}
 		}
diff --git a/Engine/Source/Runtime/AssetRegistry/Private/PackageReader.cpp b/Engine/Source/Runtime/AssetRegistry/Private/PackageReader.cpp
--- a/Engine/Source/Runtime/AssetRegistry/Private/PackageReader.cpp
+++ b/Engine/Source/Runtime/AssetRegistry/Private/PackageReader.cpp
@@ -756,6 +756,11 @@ bool FPackageReader::ReadDependencyData(FPackageDependencyData& OutDependencyDat
 		PackageData.SetIsLicenseeVersion(PackageFileSummary.SavedByEngineVersion.IsLicenseeVersion());
 		PackageData.Extension = FPackagePath::ParseExtension(PackageFilename);
 
+		// Add the filesystem location to any existing location as it's 
+        // possible we have the same content available from more than one location.
+		PackageData.SetPackageLocation(FPackageName::EPackageLocationFilter(
+			uint8(FPackageName::EPackageLocationFilter::FileSystem) | uint8(PackageData.GetPackageLocation())));
+
 		if (!SerializeImportedClasses(ImportMap, PackageData.ImportedClasses))
 		{
 			return false;
diff --git a/Engine/Source/Runtime/AssetRegistry/Public/AssetRegistry/AssetRegistryTelemetry.h b/Engine/Source/Runtime/AssetRegistry/Public/AssetRegistry/AssetRegistryTelemetry.h
--- a/Engine/Source/Runtime/AssetRegistry/Public/AssetRegistry/AssetRegistryTelemetry.h
+++ b/Engine/Source/Runtime/AssetRegistry/Public/AssetRegistry/AssetRegistryTelemetry.h
@@ -40,6 +40,8 @@ namespace UE::Telemetry::AssetRegistry
         bool bInitialSearchStarted = false;
         // Whether the main background async gather was completed at this time
         bool bInitialSearchCompleted = false;
+		// Whether an additional main background async gather was started at this time
+		bool bAdditionalMountSearchInProgress = false;
     };
 
     // Information about the initial asset registry scan triggered when launching the process. 
@@ -81,5 +83,7 @@ namespace UE::Telemetry::AssetRegistry
         bool bInitialSearchStarted = false;
         // Whether the main background async gather was completed at this time
         bool bInitialSearchCompleted = false;
+		// Whether an additional main background async gather was started at this time
+		bool bAdditionalMountSearchInProgress = false;
     };
 }
\ No newline at end of file
diff --git a/Engine/Source/Runtime/AssetRegistry/Public/AssetRegistry/IAssetRegistry.h b/Engine/Source/Runtime/AssetRegistry/Public/AssetRegistry/IAssetRegistry.h
--- a/Engine/Source/Runtime/AssetRegistry/Public/AssetRegistry/IAssetRegistry.h
+++ b/Engine/Source/Runtime/AssetRegistry/Public/AssetRegistry/IAssetRegistry.h
@@ -856,9 +856,17 @@ public:
 	DECLARE_TS_MULTICAST_DELEGATE_OneParam( FFileLoadProgressUpdatedEvent, const FFileLoadProgressUpdateData& /*ProgressUpdateData*/ );
 	virtual FFileLoadProgressUpdatedEvent& OnFileLoadProgressUpdated() = 0;
 
-	/** Returns true if the asset registry is currently loading files and does not yet know about all assets */
+	/** Returns true if the asset registry is currently loading files and does not yet know about all assets.
+	* This is a legacy function that only returns true when the asset registry is performing its initial search for assets.
+	* Prefer using IsGathering() to know if the asset registry is currently loading files.
+	*/
 	UFUNCTION(BlueprintCallable, Category = "AssetRegistry")
 	virtual bool IsLoadingAssets() const = 0;
+	/** Returns true if the asset registry is currently loading files and does not yet know about all assets.
+	* IsGathering() can be used to indicate if the asset registry is scanning during the initial search as well as any
+	* subsequent search the registry may perform (e.g. when a new mount point is added)
+	*/
+	virtual bool IsGathering() const = 0;
 
 	/** If true, the AssetRegistry updates its on-disk information for an Asset whenever that Asset loads. */
 	virtual bool ShouldUpdateDiskCacheAfterLoad() const = 0;
diff --git a/Engine/Source/Runtime/CoreUObject/Private/Misc/PackageName.cpp b/Engine/Source/Runtime/CoreUObject/Private/Misc/PackageName.cpp
--- a/Engine/Source/Runtime/CoreUObject/Private/Misc/PackageName.cpp
+++ b/Engine/Source/Runtime/CoreUObject/Private/Misc/PackageName.cpp
@@ -1971,31 +1971,43 @@ FPackageName::EPackageLocationFilter FPackageName::DoesPackageExistEx(const FPac
 
 #if WITH_EDITOR
 	IAssetRegistryInterface* AssetRegistry = IAssetRegistryInterface::GetPtr();
-
-	// Todo: The AssetRegistry currently cannot determine if a package comes from the Filesystem 
-	// or cooked content so we avoid registry lookups since we can't provide a reliable Location
-	if (AssetRegistry && ((uint8)Filter & (uint8)EPackageLocationFilter::FileSystem) && (!FIoDispatcher::IsInitialized() || !FIoDispatcherInternal::HasPackageData()))
+	if (AssetRegistry)
 	{
-		FName PackageName = PackagePath.GetPackageFName();
 		FName CorrectCasePackageName;
 		FAssetPackageData AssetPackageData;
-		if (AssetRegistry->TryGetAssetPackageData(PackageName, AssetPackageData, CorrectCasePackageName) == UE::AssetRegistry::EExists::Exists)
+		FName PackageName = PackagePath.GetPackageFName();
+		UE::AssetRegistry::EExists ExistsStatus = AssetRegistry->TryGetAssetPackageData(PackageName, AssetPackageData, CorrectCasePackageName);
+		if (ExistsStatus == UE::AssetRegistry::EExists::Exists)
 		{
-			if (OutPackagePath)
+			EPackageLocationFilter PackageLocation = AssetPackageData.GetPackageLocation();
+
+			// If the registry doesn't have this information at all, we can't trust it (e.g. reading older cache version)
+			if (PackageLocation != EPackageLocationFilter::None)
 			{
-				if (AssetPackageData.Extension != EPackageExtension::Unspecified && AssetPackageData.Extension != EPackageExtension::Custom)
+				PackageLocation = EPackageLocationFilter((uint8)Filter & (uint8)PackageLocation);
+				if (!OutPackagePath)
 				{
-					*OutPackagePath = FPackagePath::FromPackageNameUnchecked(MoveTemp(CorrectCasePackageName));
-					OutPackagePath->SetHeaderExtension(AssetPackageData.Extension);
-					return EPackageLocationFilter::FileSystem;
+					return PackageLocation;
+				}
+
+				*OutPackagePath = FPackagePath::FromPackageNameUnchecked(MoveTemp(CorrectCasePackageName));
+				if ((uint8)PackageLocation & (uint8)EPackageLocationFilter::FileSystem)
+				{
+					// If an OutPackagePath was provided but a custom or unspecified extension was used we must use the disk
+					// scanning code futher down to allow the PackageResourceManager to handle the unspecified and custom extension cases
+					if (AssetPackageData.Extension != EPackageExtension::Unspecified && AssetPackageData.Extension != EPackageExtension::Custom)
+					{
+						OutPackagePath->SetHeaderExtension(AssetPackageData.Extension);
+						return PackageLocation;
+					}
+					// Intentionally do not return a EPackageLocationFilter if an OutPackagePath was provided but a custom or 
+					// unspecified extension was used. In such cases we must use the disk scanning code futher down to allow the 
+					// PackageResourceManager to handle the unspecified and custom extension cases
+				}
+				else
+				{
+					return PackageLocation;
 				}
-				// Intentionally do not return a EPackageLocationFilter if an OutPackagePath was provided but a custom or 
-				// unspecified extension was used. In such cases we must use the disk scanning code futher down to allow the 
-				// PackageResourceManager to handle the unspecified and custom extension cases
-			}
-			else
-			{
-				return EPackageLocationFilter::FileSystem;
 			}
 		}
 	}
diff --git a/Engine/Source/Runtime/CoreUObject/Public/AssetRegistry/AssetData.h b/Engine/Source/Runtime/CoreUObject/Public/AssetRegistry/AssetData.h
--- a/Engine/Source/Runtime/CoreUObject/Public/AssetRegistry/AssetData.h
+++ b/Engine/Source/Runtime/CoreUObject/Public/AssetRegistry/AssetData.h
@@ -101,6 +101,7 @@ struct FAssetRegistryVersion
 		RemoveAssetPathFNames,				// Asset bundles are serialized as FTopLevelAssetPath instead of FSoftObjectPath, deprecated FAssetData::ObjectPath	
 		AddedHeader,						// Added header with bFilterEditorOnlyData flag 
 		AssetPackageDataHasExtension,		// Added Extension to AssetPackageData.
+		AssetPackageDataHasPackageLocation,	// Added PackageLocation to AssetPackageData.
 
 		// -----<new versions can be added above this line>-------------------------------------------------
 		VersionPlusOne,
@@ -1014,13 +1015,34 @@ public:
 	COREUOBJECT_API void NetworkWrite(FCbWriter& Writer) const;
 	COREUOBJECT_API bool TryNetworkRead(FCbFieldView Field);
 
+	COREUOBJECT_API void SetPackageLocation(FPackageName::EPackageLocationFilter PackageLocation)
+	{	
+		// Clear existing location flags
+		Flags = Flags & ~FLAG_PACKAGE_LOCATION_MASK;
+
+		static_assert((uint8)FPackageName::EPackageLocationFilter::MaxBit <= 2, "We are storing PackageLocation in two bits");
+		check(PackageLocation <= FPackageName::EPackageLocationFilter::ValidMask);
+		Flags |= uint32(PackageLocation) << FLAG_PACKAGE_LOCATION_TO_ENUM_SHIFT_SIZE;
+	}
+
+	COREUOBJECT_API FPackageName::EPackageLocationFilter GetPackageLocation() const
+	{
+		return FPackageName::EPackageLocationFilter((Flags & FLAG_PACKAGE_LOCATION_MASK) >> FLAG_PACKAGE_LOCATION_TO_ENUM_SHIFT_SIZE);
+	}
+
 private:
 	FORCEINLINE void SerializeForCacheInternal(FArchive& Ar, FAssetPackageData& PackageData, FAssetRegistryVersion::Type Version);
 
-	enum
+	enum : uint32
 	{
-		FLAG_LICENSEE_VERSION			= 1 << 0,
-		FLAG_HAS_VIRTUALIZED_PAYLOADS	= 1 << 1
+		FLAG_LICENSEE_VERSION						= 1u << 0,
+		FLAG_HAS_VIRTUALIZED_PAYLOADS				= 1u << 1,
+
+		FLAG_PACKAGE_LOCATION_IODISPATCHER			= 1u << 30,
+		FLAG_PACKAGE_LOCATION_FILESYSTEM			= 1u << 31,
+
+		FLAG_PACKAGE_LOCATION_TO_ENUM_SHIFT_SIZE	= 30u,
+		FLAG_PACKAGE_LOCATION_MASK					= uint32(FPackageName::EPackageLocationFilter::ValidMask) << FLAG_PACKAGE_LOCATION_TO_ENUM_SHIFT_SIZE
 	};
 };
 PRAGMA_ENABLE_DEPRECATION_WARNINGS
diff --git a/Engine/Source/Runtime/CoreUObject/Public/Misc/PackageName.h b/Engine/Source/Runtime/CoreUObject/Public/Misc/PackageName.h
--- a/Engine/Source/Runtime/CoreUObject/Public/Misc/PackageName.h
+++ b/Engine/Source/Runtime/CoreUObject/Public/Misc/PackageName.h
@@ -414,7 +414,13 @@ public:
 		None = 0,
 		IoDispatcher = 1,
 		FileSystem = 2,
-		Any = 0xFF, // special filter to find if it exists anywhere at all, and won't need to check both. in this case, as soon as one is found, DoesPackageExistEx will return true
+
+		MaxBit = FileSystem,
+		ValidMask = (1 << MaxBit) - 1,
+
+		// Any allows finding the PackageData if it exists anywhere at all, and won't need to check both. When using Any,
+		// as soon as the package is found in Any location, DoesPackageExistEx will return true
+		Any = 0xFF,
 	};
 
 	/**
