Asset Registry: added OnAssetCollisionEvent that can be used to resolve asset collisions based on external logic unavailable to the asset registry code.

diff --git a/Engine/Source/Runtime/AssetRegistry/Private/AssetRegistry.cpp b/Engine/Source/Runtime/AssetRegistry/Private/AssetRegistry.cpp
--- a/Engine/Source/Runtime/AssetRegistry/Private/AssetRegistry.cpp
+++ b/Engine/Source/Runtime/AssetRegistry/Private/AssetRegistry.cpp
@@ -1124,8 +1124,17 @@ void FAssetRegistryImpl::Initialize(Impl::FInitializeContext& Context)
 #endif
 }
 
-#if WITH_EDITOR
+}
 
+IAssetRegistry::FAssetCollisionEvent& UE::AssetRegistry::FAssetRegistryImpl::OnAssetCollision_Private()
+{
+	return AssetCollisionEvent;
+}
+
+namespace UE::AssetRegistry
+{
+
+#if WITH_EDITOR
 void FAssetRegistryImpl::RebuildAssetDependencyGathererMapIfNeeded()
 {
 	if (!bRegisteredDependencyGathererClassesDirty)
@@ -6289,11 +6298,15 @@ bool IsPathMounted(const FString& Path, const TSet<FString>& MountPointsNoTraili
 
 FAssetData* FAssetRegistryImpl::ResolveAssetIdCollision(FAssetData& A, FAssetData& B)
 {
+	FAssetData* ResolvedAsset = nullptr;
+	AssetCollisionEvent.Broadcast(A, B, ResolvedAsset);
+	check(!ResolvedAsset || ResolvedAsset == &A || ResolvedAsset == &B);
+
 	// We could use file age to try to guess which file is correct:
 	// FPackageName::InternalDoesPackageExistEx() to get the filename, and IFileManager::GetFileAgeSeconds
 	// But that would vary from machine to machine based on when the files were synced.
 	// So instead just pick one using an arbitrary deterministic process: alphabetical order
-	FAssetData* Keep = A.PackageName.LexicalLess(B.PackageName) ? &A : &B;
+	FAssetData* Keep = ResolvedAsset ? ResolvedAsset : (A.PackageName.LexicalLess(B.PackageName) ? &A : &B);
 	FAssetData* Discard = Keep == &A ? &B : &A;
 
 	FString PackageNameB = B.PackageName.ToString();
@@ -9369,6 +9382,11 @@ UAssetRegistryImpl::FPathRemovedEvent& UAssetRegistryImpl::OnPathRemoved()
 	return PathRemovedEvent;
 }
 
+IAssetRegistry::FAssetCollisionEvent& UAssetRegistryImpl::OnAssetCollision_Private()
+{
+	return GuardedData.OnAssetCollision_Private();
+}
+
 UAssetRegistryImpl::FAssetAddedEvent& UAssetRegistryImpl::OnAssetAdded()
 {
 	return AssetAddedEvent;
diff --git a/Engine/Source/Runtime/AssetRegistry/Private/AssetRegistry.h b/Engine/Source/Runtime/AssetRegistry/Private/AssetRegistry.h
--- a/Engine/Source/Runtime/AssetRegistry/Private/AssetRegistry.h
+++ b/Engine/Source/Runtime/AssetRegistry/Private/AssetRegistry.h
@@ -163,6 +163,8 @@ public:
 
 	virtual void PackageDeleted(UPackage* DeletedPackage) override;
 
+	virtual IAssetRegistry::FAssetCollisionEvent& OnAssetCollision_Private() override;
+
 	DECLARE_DERIVED_EVENT( UAssetRegistryImpl, IAssetRegistry::FAssetAddedEvent, FAssetAddedEvent);
 	virtual FAssetAddedEvent& OnAssetAdded() override;
 
diff --git a/Engine/Source/Runtime/AssetRegistry/Private/AssetRegistryImpl.h b/Engine/Source/Runtime/AssetRegistry/Private/AssetRegistryImpl.h
--- a/Engine/Source/Runtime/AssetRegistry/Private/AssetRegistryImpl.h
+++ b/Engine/Source/Runtime/AssetRegistry/Private/AssetRegistryImpl.h
@@ -242,6 +242,9 @@ public:
 	/** Construct the AssetRegistryImpl, including initial scans if applicable. */
 	void Initialize(Impl::FInitializeContext& Context);
 
+	/** Event for when duplicated assets are found and needs to be resolved */
+	IAssetRegistry::FAssetCollisionEvent& OnAssetCollision_Private();
+
 	// Helpers for functions of the same name from UAssetRegistryImpl
 
 
@@ -428,6 +431,9 @@ public:
 #endif
 private:
 
+	/** The delegate to execute when an asset collision is detected */
+	IAssetRegistry::FAssetCollisionEvent AssetCollisionEvent;
+
 	/**
 	 * Enumerate all AssetData in memory and in the State that match the given path.
 	 * The returned AssetData for objects present in memory have no AssetRegistryTags calculated
diff --git a/Engine/Source/Runtime/AssetRegistry/Public/AssetRegistry/IAssetRegistry.h b/Engine/Source/Runtime/AssetRegistry/Public/AssetRegistry/IAssetRegistry.h
--- a/Engine/Source/Runtime/AssetRegistry/Public/AssetRegistry/IAssetRegistry.h
+++ b/Engine/Source/Runtime/AssetRegistry/Public/AssetRegistry/IAssetRegistry.h
@@ -770,6 +770,13 @@ public:
 
 	DECLARE_TS_MULTICAST_DELEGATE_OneParam( FAssetsEvent, TConstArrayView<FAssetData> );
 
+	/** 
+	* EXPERIMENTAL: event for when duplicated assets are found and needs to be resolved, this is experimental as it's a counter measure for
+	* a situation that shouldn't happen. Also, DO NOT CALL AssetRegistry functions from delegate subscribers; doing so will create a deadlock.
+	*/
+	DECLARE_TS_MULTICAST_DELEGATE_ThreeParams( FAssetCollisionEvent, FAssetData&, FAssetData&, FAssetData*& );
+	virtual FAssetCollisionEvent& OnAssetCollision_Private() = 0;
+
 	/** Event for when assets are added to the registry */
 	DECLARE_TS_MULTICAST_DELEGATE_OneParam( FAssetAddedEvent, const FAssetData& );
 	virtual FAssetAddedEvent& OnAssetAdded() = 0;
