IncrementalCook: AssetRegistry: Mark TObjectPtr dereferences during GetAssetRegistryTags calls during AssetRegistry queries as not dependencies coming from the package that made the query. Some objects that have GetAssetRegistryTags called do not pass the query and are not used for any reason in the source package making the query. If callers have a build dependency on the AssetRegistry query, they should instead add FCookDepedency::AssetRegistryQuery.
#rnx

diff --git a/Engine/Source/Runtime/AssetRegistry/Private/AssetRegistry.cpp b/Engine/Source/Runtime/AssetRegistry/Private/AssetRegistry.cpp
--- a/Engine/Source/Runtime/AssetRegistry/Private/AssetRegistry.cpp
+++ b/Engine/Source/Runtime/AssetRegistry/Private/AssetRegistry.cpp
@@ -2545,6 +2545,7 @@ void EnumerateMemoryAssetsHelper(const FARCompiledFilter& InFilter, TSet<FName>&
 {
 	checkf(IsInGameThread(), TEXT("Enumerating in-memory assets can only be done on the game thread; it uses non-threadsafe UE::AssetRegistry::Filtering globals."));
 	bOutStopIteration = false;
+	UE_TRACK_REFERENCING_OPNAME_SCOPED(PackageAccessTrackingOps::NAME_ResetContext);
 
 	// Skip assets that were loaded for diffing
 	const uint32 FilterWithoutPackageFlags = InFilter.WithoutPackageFlags | PKG_ForDiffing;
