AssetRegistry: Downgrade warning message to verbose when we detect a background scan result is from an unmounted plugin. This case can legitimately occur in farm jobs that unmount plugins and we do not want a spurious warning in that case.
Extend the handling for the unmounted background scan results to also reject the AssetPackageData and DependencyData; previously we were correctly rejecting AssetData and Paths but were missing the AssetPackageData.
Improve performance between the 3 types of data that need to be rejected by sharing the calculation of the TSet<FString> MountPoints.
#rnx

diff --git a/Engine/Source/Runtime/AssetRegistry/Private/AssetRegistry.cpp b/Engine/Source/Runtime/AssetRegistry/Private/AssetRegistry.cpp
--- a/Engine/Source/Runtime/AssetRegistry/Private/AssetRegistry.cpp
+++ b/Engine/Source/Runtime/AssetRegistry/Private/AssetRegistry.cpp
@@ -4899,6 +4899,7 @@ Impl::EGatherStatus FAssetRegistryImpl::TickGatherer(Impl::FTickContext& TickCon
 	using namespace UE::AssetRegistry::Impl;
 	FEventContext& EventContext = TickContext.EventContext;
 	FInterruptionContext& InOutInterruptionContext = TickContext.InterruptionContext;
+	TOptional<TSet<FString>> MountPointsForVerifyAfterGather;
 	bool bLocalIsInGameThread = IsInGameThread();
 
 	EGatherStatus OutStatus = EGatherStatus::Complete;
@@ -5033,7 +5034,8 @@ Impl::EGatherStatus FAssetRegistryImpl::TickGatherer(Impl::FTickContext& TickCon
 	if (BackgroundResults.Paths.Num())
 	{
 		LazyStartTimer();
-		PathDataGathered(EventContext, BackgroundResults.Paths, InOutInterruptionContext);
+		PathDataGathered(EventContext, BackgroundResults.Paths, InOutInterruptionContext,
+			MountPointsForVerifyAfterGather);
 	}
 	if (InOutInterruptionContext.ShouldExitEarly())
 	{
@@ -5043,7 +5045,7 @@ Impl::EGatherStatus FAssetRegistryImpl::TickGatherer(Impl::FTickContext& TickCon
 	}
 
 	auto RunAssetSearchDataGathered = [this, &EventContext, &TickContext, &LazyStartTimer,
-		&InOutInterruptionContext]
+		&InOutInterruptionContext, &MountPointsForVerifyAfterGather]
 	(TMultiMap<FName, TUniquePtr<FAssetData>>& InAssetResults,
 		TMultiMap<FName, TUniquePtr<FAssetData>>& OutDeferredAssetResults)
 		{
@@ -5063,10 +5065,11 @@ Impl::EGatherStatus FAssetRegistryImpl::TickGatherer(Impl::FTickContext& TickCon
 				}
 
 				AssetSearchDataGathered(EventContext, InAssetResults, OutDeferredAssetResults,
-					InOutInterruptionContext);
+					InOutInterruptionContext, MountPointsForVerifyAfterGather);
 			}
 		};
-	auto RunDependencyDataGathered = [this, &bLocalIsInGameThread, &LazyStartTimer, &InOutInterruptionContext]
+	auto RunDependencyDataGathered = [this, &bLocalIsInGameThread, &LazyStartTimer, &InOutInterruptionContext,
+		&MountPointsForVerifyAfterGather]
 	(TMultiMap<FName, FPackageDependencyData>& DependenciesToProcess,
 		TMultiMap<FName, FPackageDependencyData>& OutDeferredDependencies,
 		TSet<FName>* OutPackagesNeedingDependencyCalculation)
@@ -5077,7 +5080,7 @@ Impl::EGatherStatus FAssetRegistryImpl::TickGatherer(Impl::FTickContext& TickCon
 				LazyStartTimer();
 
 				DependencyDataGathered(DependenciesToProcess, OutDeferredDependencies,
-					OutPackagesNeedingDependencyCalculation, InOutInterruptionContext);
+					OutPackagesNeedingDependencyCalculation, InOutInterruptionContext, MountPointsForVerifyAfterGather);
 			}
 		};
 
@@ -5463,6 +5466,7 @@ void FAssetRegistryImpl::TickGatherPackage(Impl::FEventContext& EventContext, co
 	DeferredDependenciesForGameThread.MultiFind(PackageFName, PackageDependencyDatas);
 	DeferredDependenciesForGameThread.Remove(PackageFName);
 
+	TOptional<TSet<FString>> MountPointsForVerifyAfterGather;
 	if (PackageAssets.Num() > 0)
 	{
 		LazyStartTimer();
@@ -5481,7 +5485,7 @@ void FAssetRegistryImpl::TickGatherPackage(Impl::FEventContext& EventContext, co
 		TMultiMap<FName, TUniquePtr<FAssetData>> DeferredPackageAssetsMap;
 		UE::AssetRegistry::Impl::FInterruptionContext InterruptionContext;
 		AssetSearchDataGathered(EventContext, PackageAssetsMap, DeferredPackageAssetsMap,
-			InterruptionContext);
+			InterruptionContext, MountPointsForVerifyAfterGather);
 		if (DeferredPackageAssetsMap.Num())
 		{
 			UE_LOG(LogAssetRegistry, Warning, TEXT("Attempted to add package '%s' to the registry before its UClass was available. \
@@ -5514,7 +5518,7 @@ Could not execute PostLoadAssetRegistryTags. We will try again later. Until then
 #endif
 		UE::AssetRegistry::Impl::FInterruptionContext InterruptionContext;
 		DependencyDataGathered(PackageDependencyDatasMap, DeferredDependenciesForGameThread,
-			OutPackagesNeedingDependencyCalculation, InterruptionContext);
+			OutPackagesNeedingDependencyCalculation, InterruptionContext, MountPointsForVerifyAfterGather);
 	}
 }
 
@@ -6284,11 +6288,12 @@ void FAssetRegistryImpl::ScanPathsSynchronous(Impl::FScanPathContext& Context)
 		bForceCompletionEvenIfPostLoadsFail = true;
 
 		int32 OriginalNumDeferredAssetsForGameThread = DeferredAssetsForGameThread.Num();
+		TOptional<TSet<FString>> MountPointsForVerifyAfterGather;
 
 		// We don't call AssetsFoundCallback here because even for deferred assets it will already have been called.
 		// We pass DeferredAssetsForGameThread as the OutDeferred parameter, but we expect nothing will be deferred.
 		AssetSearchDataGathered(Context.EventContext, CollectedDeferredAssets, DeferredAssetsForGameThread,
-			TickContext.InterruptionContext);
+			TickContext.InterruptionContext, MountPointsForVerifyAfterGather);
 		// All of the assets we collected should have been processed or deferred.
 		ensure(CollectedDeferredAssets.Num() == 0);
 
@@ -6327,6 +6332,18 @@ void FAssetRegistryImpl::ScanPathsSynchronous(Impl::FScanPathContext& Context)
 namespace Utils
 {
 
+void InitializeMountPoints(TOptional<TSet<FString>>& MountPoints)
+{
+	if (MountPoints.IsSet())
+	{
+		return;
+	}
+	MountPoints.Emplace();
+	TArray<FString> MountPointsArray;
+	FPackageName::QueryRootContentPaths(MountPointsArray, /*bIncludeReadOnlyRoots=*/ true, /*bWithoutLeadingSlashes*/ false, /*WithoutTrailingSlashes=*/ true);
+	MountPoints->Append(MoveTemp(MountPointsArray));
+}
+
 bool IsPathMounted(const FString& Path, const TSet<FString>& MountPointsNoTrailingSlashes, FString& StringBuffer)
 {
 	const int32 SecondSlash = Path.Len() > 1 ? Path.Find(TEXT("/"), ESearchCase::CaseSensitive, ESearchDir::FromStart, 1) : INDEX_NONE;
@@ -6636,21 +6653,19 @@ bool FAssetRegistryImpl::ShouldSkipGatheredAsset(FAssetData& AssetData)
 void FAssetRegistryImpl::AssetSearchDataGathered(Impl::FEventContext& EventContext,
 	TMultiMap<FName, TUniquePtr<FAssetData>>& AssetResults,
 	TMultiMap<FName, TUniquePtr<FAssetData>>& OutDeferredAssetResults,
-	Impl::FInterruptionContext& InOutInterruptionContext)
+	Impl::FInterruptionContext& InOutInterruptionContext,
+	TOptional<TSet<FString>>& MountPointsForVerifyAfterGather)
 {
 	TRACE_CPUPROFILER_EVENT_SCOPE(AssetSearchDataGathered);
 
 	// Refreshes ClassGeneratorNames if out of date due to module load
 	CollectCodeGeneratorClasses();
 
-	TSet<FString> MountPoints;
 	FString PackagePathString;
 	FString PackageRoot;
 	if (AssetResults.Num() > 0 && bVerifyMountPointAfterGather)
 	{
-		TArray<FString> MountPointsArray;
-		FPackageName::QueryRootContentPaths(MountPointsArray, /*bIncludeReadOnlyRoots=*/ true, /*bWithoutLeadingSlashes*/ false, /*WithoutTrailingSlashes=*/ true);
-		MountPoints.Append(MoveTemp(MountPointsArray));
+		Utils::InitializeMountPoints(MountPointsForVerifyAfterGather);
 	}
 
 #if WITH_EDITOR
@@ -6694,11 +6709,11 @@ void FAssetRegistryImpl::AssetSearchDataGathered(Impl::FEventContext& EventConte
 		if (bVerifyMountPointAfterGather)
 		{
 			PackagePath.ToString(PackagePathString);
-			if (!Utils::IsPathMounted(PackagePathString, MountPoints, PackageRoot))
+			if (!Utils::IsPathMounted(PackagePathString, *MountPointsForVerifyAfterGather, PackageRoot))
 			{
-				UE_LOG(LogAssetRegistry, Warning,
+				UE_LOG(LogAssetRegistry, Display,
 					TEXT("AssetRegistry: An asset has been loaded with an invalid mount point: '%s', Mount Point: '%s'. Ignoring the asset."),
-					*BackgroundResult->GetObjectPathString(), *PackagePathString)
+					*BackgroundResult->GetObjectPathString(), *PackagePathString);
 				continue;
 			}
 		}
@@ -6781,17 +6796,15 @@ void FAssetRegistryImpl::AssetSearchDataGathered(Impl::FEventContext& EventConte
 	}
 }
 
-void FAssetRegistryImpl::PathDataGathered(Impl::FEventContext& EventContext, TRingBuffer<FString>& PathResults, Impl::FInterruptionContext& InOutInterruptionContext)
+void FAssetRegistryImpl::PathDataGathered(Impl::FEventContext& EventContext, TRingBuffer<FString>& PathResults,
+	Impl::FInterruptionContext& InOutInterruptionContext, TOptional<TSet<FString>>& MountPointsForVerifyAfterGather)
 {
 	TRACE_CPUPROFILER_EVENT_SCOPE(PathDataGathered);
 
-	TSet<FString> MountPoints;
 	FString PackageRoot;
 	if (PathResults.Num() > 0 && bVerifyMountPointAfterGather)
 	{
-		TArray<FString> MountPointsArray;
-		FPackageName::QueryRootContentPaths(MountPointsArray, /*bIncludeReadOnlyRoots=*/ true, /*bWithoutLeadingSlashes*/ false, /*WithoutTrailingSlashes=*/ true);
-		MountPoints.Append(MoveTemp(MountPointsArray));
+		Utils::InitializeMountPoints(MountPointsForVerifyAfterGather);
 	}
 	
 	CachedPathTree.EnsureAdditionalCapacity(PathResults.Num());
@@ -6801,13 +6814,15 @@ void FAssetRegistryImpl::PathDataGathered(Impl::FEventContext& EventContext, TRi
 		FString Path = PathResults.PopFrontValue();
 
 		// Skip stale results caused by mount then unmount of a path within short period.
-		if (!bVerifyMountPointAfterGather || Utils::IsPathMounted(Path, MountPoints, PackageRoot))
+		if (!bVerifyMountPointAfterGather || Utils::IsPathMounted(Path, *MountPointsForVerifyAfterGather, PackageRoot))
 		{
 			AddAssetPath(EventContext, FName(*Path));
 		}
 		else
 		{
-			UE_LOG(LogAssetRegistry, Warning, TEXT("AssetRegistry: A path has been loaded with an invalid mount point: '%s'"), *Path)
+			UE_LOG(LogAssetRegistry, Display,
+				TEXT("AssetRegistry: A path has been loaded with an invalid mount point: '%s', Mount Point: '%s'. Ignoring the path."),
+				*Path, *PackageRoot);
 		}
 
 		// Check to see if we have run out of time in this tick
@@ -6820,7 +6835,8 @@ void FAssetRegistryImpl::PathDataGathered(Impl::FEventContext& EventContext, TRi
 
 void FAssetRegistryImpl::DependencyDataGathered(TMultiMap<FName, FPackageDependencyData>& DependsResults, 
 	TMultiMap<FName, FPackageDependencyData>& OutDeferredDependencyResults, 
-	TSet<FName>* OutPackagesNeedingDependencyCalculation, Impl::FInterruptionContext& InOutInterruptionContext)
+	TSet<FName>* OutPackagesNeedingDependencyCalculation, Impl::FInterruptionContext& InOutInterruptionContext,
+	TOptional<TSet<FString>>& MountPointsForVerifyAfterGather)
 {
 	TRACE_CPUPROFILER_EVENT_SCOPE(DependencyDataGathered);
 	using namespace UE::AssetRegistry;
@@ -6831,6 +6847,9 @@ void FAssetRegistryImpl::DependencyDataGathered(TMultiMap<FName, FPackageDepende
 	TMap<FName, FName> CachedDepToRedirect;
 	bool bInterrupted = false;
 	int64 IterationCounter = 0;
+	FString PackageNameString;
+	FString PackageRoot;
+
 	for (TMultiMap<FName, FPackageDependencyData>::TIterator Iter(DependsResults); Iter && !bInterrupted; ++Iter)
 	{
 		ON_SCOPE_EXIT
@@ -6854,6 +6873,19 @@ void FAssetRegistryImpl::DependencyDataGathered(TMultiMap<FName, FPackageDepende
 		FPackageDependencyData Result = MoveTemp(Iter.Value());
 		Iter.RemoveCurrent();
 
+		if (bVerifyMountPointAfterGather)
+		{
+			Utils::InitializeMountPoints(MountPointsForVerifyAfterGather);
+			Result.PackageName.ToString(PackageNameString);
+			if (!Utils::IsPathMounted(PackageNameString, *MountPointsForVerifyAfterGather, PackageRoot))
+			{
+				UE_LOG(LogAssetRegistry, Display,
+					TEXT("AssetRegistry: DependencyData has been loaded with an invalid mount point: '%s', Mount Point: '%s'. Ignoring the DependencyData."),
+					*PackageNameString, *PackageRoot);
+				continue;
+			}
+		}
+
 		checkf(!GIsEditor || Result.bHasPackageData, TEXT("We rely on PackageData being read for every gathered Asset in the editor."));
 		if (Result.bHasPackageData)
 		{
diff --git a/Engine/Source/Runtime/AssetRegistry/Private/AssetRegistryImpl.h b/Engine/Source/Runtime/AssetRegistry/Private/AssetRegistryImpl.h
--- a/Engine/Source/Runtime/AssetRegistry/Private/AssetRegistryImpl.h
+++ b/Engine/Source/Runtime/AssetRegistry/Private/AssetRegistryImpl.h
@@ -8,8 +8,10 @@
 #include "AssetRegistry/IAssetRegistry.h"
 #include "AssetRegistry/PathTree.h"
 #include "Containers/RingBuffer.h"
+#include "Containers/Set.h"
 #include "HAL/CriticalSection.h"
 #include "Math/NumericLimits.h"
+#include "Misc/Optional.h"
 #include "Misc/TransactionallySafeRWLock.h"
 #include "Misc/TransactionallySafeRWScopeLock.h"
 #include "ModuleDescriptor.h"
@@ -461,7 +463,8 @@ private:
 	void AssetSearchDataGathered(Impl::FEventContext& EventContext,
 		TMultiMap<FName, TUniquePtr<FAssetData>>& AssetResults, 
 		TMultiMap<FName, TUniquePtr<FAssetData>>& OutDeferredResults,\
-		Impl::FInterruptionContext& InOutInterruptionContext);
+		Impl::FInterruptionContext& InOutInterruptionContext,
+		TOptional<TSet<FString>>& MountPointsForVerifyAfterGather);
 	/** Validate assets gathered from disk before adding them to the AssetRegistry. */
 	bool ShouldSkipGatheredAsset(FAssetData& AssetData);
 
@@ -469,12 +472,15 @@ private:
 	 * Called every tick when data is retrieved by the background path search.
 	 * If TickStartTime is < 0, the entire list of gathered assets will be cached. Also used in sychronous searches
 	 */
-	void PathDataGathered(Impl::FEventContext& EventContext, TRingBuffer<FString>& PathResults, Impl::FInterruptionContext& InOutInterruptionContext);
+	void PathDataGathered(Impl::FEventContext& EventContext, TRingBuffer<FString>& PathResults,
+		Impl::FInterruptionContext& InOutInterruptionContext,
+		TOptional<TSet<FString>>& MountPointsForVerifyAfterGather);
 
 	/** Called every tick when data is retrieved by the background dependency search */
 	void DependencyDataGathered(TMultiMap<FName, FPackageDependencyData>& DependsResults, 
-		TMultiMap<FName, FPackageDependencyData>& OutDeferredDependencyResults,	TSet<FName>* OutPackagesNeedingDependencyCalculation, 
-		Impl::FInterruptionContext& InOutInterruptionContext);
+		TMultiMap<FName, FPackageDependencyData>& OutDeferredDependencyResults,
+		TSet<FName>* OutPackagesNeedingDependencyCalculation, Impl::FInterruptionContext& InOutInterruptionContext,
+		TOptional<TSet<FString>>& MountPointsForVerifyAfterGather);
 
 	/** Called every tick when data is retrieved by the background search for cooked packages that do not contain asset data */
 	void CookedPackageNamesWithoutAssetDataGathered(Impl::FEventContext& EventContext, 
@@ -858,6 +864,9 @@ void PrioritizeAssetInstall(const FAssetData& AssetData);
  */
 bool ReadAssetFile(FPackageReader& PackageReader, IAssetRegistry::FLoadPackageRegistryData& InOutData);
 
+/** If InOutMountPoints is not already set, sets it and reads mount points into it from FPackageName. */
+void InitializeMountPoints(TOptional<TSet<FString>>& InOutMountPoints);
+
 /**
  * Returns true if path belongs to one of the mount points provided
  *
