Asset Registry Helper: add UAssetRegistryHelpers::GetAssetsWithOuterForPaths that returns all the corresponding assets from the provided paths that are outered to the provided outer name.

diff --git a/Engine/Source/Runtime/AssetRegistry/Private/AssetRegistryHelpers.cpp b/Engine/Source/Runtime/AssetRegistry/Private/AssetRegistryHelpers.cpp
--- a/Engine/Source/Runtime/AssetRegistry/Private/AssetRegistryHelpers.cpp
+++ b/Engine/Source/Runtime/AssetRegistry/Private/AssetRegistryHelpers.cpp
@@ -5,6 +5,7 @@
 #include "AssetRegistry/IAssetRegistry.h"
 #include "AssetRegistry.h"
 #include "Algo/AnyOf.h"
+#include "Algo/RemoveIf.h"
 #include "Blueprint/BlueprintSupport.h"
 #include "Misc/PackageName.h"
 #include "Modules/ModuleManager.h"
@@ -332,3 +333,27 @@ void UAssetRegistryHelpers::FixupRedirectedAssetPath(FName& InOutAssetPath)
 	FixupRedirectedAssetPath(SoftObjectPath);
 	InOutAssetPath = FName(*SoftObjectPath.ToString());
 }
+
+#if WITH_EDITOR
+TArray<FAssetData> UAssetRegistryHelpers::GetAssetsWithOuterForPaths(const TArray<FName>& InPackagePaths, FName InOuterPath, bool bInRecursivePaths, bool bInIncludeOnlyOnDiskAssets)
+{
+	FARFilter Filter;
+	Filter.PackagePaths = InPackagePaths;
+	Filter.bRecursivePaths = bInRecursivePaths;
+	Filter.bIncludeOnlyOnDiskAssets = bInIncludeOnlyOnDiskAssets;
+
+	TArray<FString> PathsToScan;
+	PathsToScan.Reserve(Filter.PackagePaths.Num());
+	Algo::Transform(Filter.PackagePaths, PathsToScan, [](FName InPackagePath) { return InPackagePath.ToString(); });
+
+	TArray<FAssetData> Assets;
+	IAssetRegistry& AssetRegistry = IAssetRegistry::GetChecked();
+	AssetRegistry.ScanSynchronous(PathsToScan, TArray<FString>());
+	AssetRegistry.GetAssets(Filter, Assets);
+
+	Assets.SetNum(Algo::RemoveIf(Assets, [InOuterPath](const FAssetData& InAssetData) { return InAssetData.GetOptionalOuterPathName() != InOuterPath; }));
+	Assets.Sort();
+
+	return Assets;
+}
+#endif
\ No newline at end of file
diff --git a/Engine/Source/Runtime/AssetRegistry/Public/AssetRegistry/AssetRegistryHelpers.h b/Engine/Source/Runtime/AssetRegistry/Public/AssetRegistry/AssetRegistryHelpers.h
--- a/Engine/Source/Runtime/AssetRegistry/Public/AssetRegistry/AssetRegistryHelpers.h
+++ b/Engine/Source/Runtime/AssetRegistry/Public/AssetRegistry/AssetRegistryHelpers.h
@@ -176,4 +176,15 @@ public:
 	 * @param	InOutSoftObjectPath	Asset path to resolve [In/Out].
 	 */
 	static ASSETREGISTRY_API void FixupRedirectedAssetPath(FSoftObjectPath& InOutSoftObjectPath);
+
+#if WITH_EDITOR
+	/**
+	 * Returns all the corresponding assets from the provided paths that are outered to the provided outer name.
+	 * @param	InPackagePaths List of package paths to scan for assets.
+	 * @param	InOuterPath The outer path to filter asset with.
+	 * @param	bInRecursivePaths If we should scan recursively or not.
+	 * @param	bInIncludeOnlyOnDiskAsset If we should only include on disk assets.
+	 */
+	 static ASSETREGISTRY_API TArray<FAssetData> GetAssetsWithOuterForPaths(const TArray<FName>& InPackagePaths, FName InOuterPath, bool bInRecursivePaths = false, bool bInIncludeOnlyOnDiskAsset = false);
+#endif
 };
