AssetRegistry: Improve performance of WaitOnPackage by moving the early exit if !IsGathering outside of the critical section, so that it does not contend with thread threads reading from the AssetRegistry, if gathering has already completed. We are able to do this because we changed IsGatthering to use atomics.
#rnx

diff --git a/Engine/Source/Runtime/AssetRegistry/Private/AssetRegistry.cpp b/Engine/Source/Runtime/AssetRegistry/Private/AssetRegistry.cpp
--- a/Engine/Source/Runtime/AssetRegistry/Private/AssetRegistry.cpp
+++ b/Engine/Source/Runtime/AssetRegistry/Private/AssetRegistry.cpp
@@ -2216,18 +2216,25 @@ void UAssetRegistryImpl::WaitForPackage(const FString& PackageName)
 {
 	TRACE_CPUPROFILER_EVENT_SCOPE(UAssetRegistryImpl::WaitForPackage);
 
+	if (!GuardedData.IsGathering())
+	{
+		// IsGathering uses relaxed memory order, so it is possible that another thread has just finished
+		// marking IsGathering=false and has not finished writing the data it gathered. But that's not
+		// a problem, because to read that data, the caller will need to enter the critical section which was
+		// held by the other thread that wrote the data, and entering a critical section waits on a memory fence.
+		return;
+	}
+	FString LocalPath;
+	if (!FPackageName::TryConvertLongPackageNameToFilename(PackageName, LocalPath))
+	{
+		return;
+	}
+
 	UE::AssetRegistry::Impl::FEventContext EventContext;
 	{
 		LLM_SCOPE(ELLMTag::AssetRegistry);
 		UE::AssetRegistry::FInterfaceWriteScopeLock InterfaceScopeLock(InterfaceLock);
-		if (GuardedData.IsGathering())
-		{
-			FString LocalPath;
-			if (FPackageName::TryConvertLongPackageNameToFilename(PackageName, LocalPath))
-			{
-				GuardedData.TickGatherPackage(EventContext, PackageName, LocalPath);
-			}
-		}
+		GuardedData.TickGatherPackage(EventContext, PackageName, LocalPath);
 	}
 	Broadcast(EventContext);
 }
@@ -3357,9 +3364,10 @@ FSoftObjectPath UAssetRegistryImpl::GetRedirectedObjectPath(const FSoftObjectPat
 {
 	// Fast path, if a full registry scan was triggered & has completed
 	// In that case, we can skip further scanning while looking for a redirected path
+	if (!GuardedData.IsGathering())
 	{
 		UE::AssetRegistry::FInterfaceReadScopeLock InterfaceScopeLock(InterfaceLock);
-		if (GuardedData.IsSearchAllAssets() && !GuardedData.IsGathering())
+		if (GuardedData.IsSearchAllAssets())
 		{
 			return GuardedData.GetRedirectedObjectPath(ObjectPath, nullptr, nullptr, /*bNeedsScanning*/ false);
 		}
