[Insights] Fixed issues re tracing the "Asset Registry Scan" timing region (begin events being traced multiple times without a corresponding end event).

diff --git a/Engine/Source/Runtime/AssetRegistry/Private/AssetRegistry.cpp b/Engine/Source/Runtime/AssetRegistry/Private/AssetRegistry.cpp
--- a/Engine/Source/Runtime/AssetRegistry/Private/AssetRegistry.cpp
+++ b/Engine/Source/Runtime/AssetRegistry/Private/AssetRegistry.cpp
@@ -1991,7 +1991,6 @@ namespace UE::AssetRegistry
 void FAssetRegistryImpl::SearchAllAssets(Impl::FEventContext& EventContext,
 	Impl::FClassInheritanceContext& InheritanceContext, bool bSynchronousSearch)
 {
-	TRACE_BEGIN_REGION(TEXT("Asset Registry Scan"));
 	EventContext.bScanStartedEventBroadcast = true;
 
 	if (!TryConstructGathererIfNeeded())
@@ -2001,6 +2000,7 @@ void FAssetRegistryImpl::SearchAllAssets(Impl::FEventContext& EventContext,
 
 	if (!bInitialSearchStarted)
 	{
+		TRACE_BEGIN_REGION(TEXT("Asset Registry Scan"));
 		InitialSearchStartTime = FPlatformTime::Seconds();
 		bInitialSearchStarted = true;
 		bInitialSearchCompleted.store(false, std::memory_order_relaxed);
@@ -2008,6 +2008,7 @@ void FAssetRegistryImpl::SearchAllAssets(Impl::FEventContext& EventContext,
 	}
 	else if (!IsGathering())
 	{
+		TRACE_BEGIN_REGION(TEXT("Asset Registry - Additional Mount Search")); // Matching TRACE_END_REGION in OnAdditionalMountSearchCompleted
 		AdditionalMountSearchStartTime = FPlatformTime::Seconds();
 		bAdditionalMountSearchInProgress.store(true, std::memory_order_relaxed);
 		GlobalGatherer->SetIsAdditionalMountSearchInProgress(true);
@@ -5373,7 +5374,7 @@ void FAssetRegistryImpl::OnAdditionalMountSearchCompleted(Impl::FEventContext& E
 	SetPerformanceMode(Impl::EPerformanceMode::MostlyStatic);
 
 	LogSearchDiagnostics(AdditionalMountSearchStartTime);
-	TRACE_END_REGION(TEXT("Asset Registry Scan")); // Matching TRACE_BEGIN_REGION in OnContentPathMounted
+	TRACE_END_REGION(TEXT("Asset Registry - Additional Mount Search")); // Matching TRACE_BEGIN_REGION in OnContentPathMounted
 
 	GlobalGatherer->OnAdditionalMountSearchCompleted();
 
@@ -8198,7 +8199,7 @@ void FAssetRegistryImpl::OnContentPathMounted(Impl::FEventContext& EventContext,
 		{
 			if (!IsGathering())
 			{
-				TRACE_BEGIN_REGION(TEXT("Asset Registry Scan")); // Matching TRACE_END_REGION in OnAdditionalMountSearchCompleted
+				TRACE_BEGIN_REGION(TEXT("Asset Registry - Additional Mount Search")); // Matching TRACE_END_REGION in OnAdditionalMountSearchCompleted
 				AdditionalMountSearchStartTime = FPlatformTime::Seconds();
 				bAdditionalMountSearchInProgress.store(true, std::memory_order_relaxed);
 				GlobalGatherer->SetIsAdditionalMountSearchInProgress(true);
