Fix bug in FAssetRegistryState::InitializeFromExistingAndPrune : when creating a FAssetData, check if the asset is a root asset or not and use the appropriate constructor.

#rnx

diff --git a/Engine/Source/Runtime/AssetRegistry/Private/AssetRegistryState.cpp b/Engine/Source/Runtime/AssetRegistry/Private/AssetRegistryState.cpp
--- a/Engine/Source/Runtime/AssetRegistry/Private/AssetRegistryState.cpp
+++ b/Engine/Source/Runtime/AssetRegistry/Private/AssetRegistryState.cpp
@@ -268,8 +268,17 @@ void FAssetRegistryState::InitializeFromExistingAndPrune(const FAssetRegistrySta
 		FAssetRegistryState::FilterTags(AssetData.TagsAndValues, NewTagsAndValues,
 			Options.CookFilterlistTagsByClass.Find(AssetData.AssetClassPath), Options);
 
-		FAssetData* NewAssetData = new FAssetData(AssetData.PackageName, AssetData.PackagePath, AssetData.AssetName,
-			AssetData.AssetClassPath, NewTagsAndValues, AssetData.GetChunkIDs(), AssetData.PackageFlags);
+		FAssetData* NewAssetData = nullptr;
+		if (AssetData.IsTopLevelAsset())
+		{
+			NewAssetData = new FAssetData(AssetData.PackageName, AssetData.PackagePath, AssetData.AssetName,
+				AssetData.AssetClassPath, NewTagsAndValues, AssetData.GetChunkIDs(), AssetData.PackageFlags);
+		}
+		else
+		{
+			NewAssetData = new FAssetData(AssetData.PackageName.ToString(), AssetData.GetObjectPathString(), 
+				AssetData.AssetClassPath, NewTagsAndValues, AssetData.GetChunkIDs(), AssetData.PackageFlags);
+		}
 
 		NewAssetData->TaggedAssetBundles = AssetData.TaggedAssetBundles;
 
