AssetRegistry: Gatherer: Fix deadlock when loading the gatherer's file cache. We're running out of threads; the body of the ParallelFor is waiting on a UE::Tasks::Launch to complete, while still holding its ParallelFor thread. If too many ParallelFor threads are kicked at once and there are no more threads when all of them launch their async task, then none can continue and they deadlock.
Refactor the existing ParallelFor to instead request more ParallelFor threads, and work on the async tasks on the extra ParallelFor threads. The ParallelFor will convert those async tasks into serial tasks if we run out of threads.
#rnx

diff --git a/Engine/Source/Runtime/AssetRegistry/Private/AssetDataGatherer.cpp b/Engine/Source/Runtime/AssetRegistry/Private/AssetDataGatherer.cpp
--- a/Engine/Source/Runtime/AssetRegistry/Private/AssetDataGatherer.cpp
+++ b/Engine/Source/Runtime/AssetRegistry/Private/AssetDataGatherer.cpp
@@ -68,6 +68,13 @@ namespace AssetDataGathererConstants
 	 * Capping the number of threads allows those threads to be used in other async operations in Engine startup.
 	 */
 	static int32 GARGatherThreads = 60;
+
+	/**
+	 * The maximum number of the gatherer's cache files that can be loaded in parallel in LoadCacheFiles.
+	 * The amount of memory used when reading the cache rises proportional to this value.
+	 * Other than the memory cost, it should be set as high as possible to load the cache faster.
+	 */
+	static int32 GARGatherCacheParallelism = 16;
 }
 
 namespace UE::AssetDataGather::Private
@@ -272,6 +279,15 @@ void FPreloadSettings::Initialize()
 #else
 	bPreloadCache = false;
 #endif
+
+	FParse::Value(FCommandLine::Get(), TEXT("-ARDiscoverThreads="), AssetDataGathererConstants::GARDiscoverThreads);
+	FParse::Value(FCommandLine::Get(), TEXT("-ARDiscoverMinBatchSize="), AssetDataGathererConstants::GARDiscoverMinBatchSize);
+	FParse::Value(FCommandLine::Get(), TEXT("-ARGatherThreads="), AssetDataGathererConstants::GARGatherThreads);
+	FParse::Value(FCommandLine::Get(), TEXT("-ARGatherCacheParallelism="), AssetDataGathererConstants::GARGatherCacheParallelism);
+	AssetDataGathererConstants::GARDiscoverThreads = FMath::Max(0, AssetDataGathererConstants::GARDiscoverThreads);
+	AssetDataGathererConstants::GARDiscoverMinBatchSize = FMath::Max(1, AssetDataGathererConstants::GARDiscoverMinBatchSize);
+	AssetDataGathererConstants::GARGatherThreads = FMath::Max(0, AssetDataGathererConstants::GARGatherThreads);
+	AssetDataGathererConstants::GARGatherCacheParallelism = FMath::Max(1, AssetDataGathererConstants::GARGatherCacheParallelism);
 }
 bool FPreloadSettings::IsGatherCacheReadEnabled() const
 {
@@ -1567,12 +1583,7 @@ FAssetDataDiscovery::FAssetDataDiscovery(const TArray<FString>& InLongPackageNam
 
 	PriorityDataUpdated->Trigger();
 
-	FParse::Value(FCommandLine::Get(), TEXT("-ARDiscoverThreads="), AssetDataGathererConstants::GARDiscoverThreads);
-	FParse::Value(FCommandLine::Get(), TEXT("-ARDiscoverMinBatchSize="), AssetDataGathererConstants::GARDiscoverMinBatchSize);
-	FParse::Value(FCommandLine::Get(), TEXT("-ARGatherThreads="), AssetDataGathererConstants::GARGatherThreads);
-	AssetDataGathererConstants::GARDiscoverThreads = FMath::Max(0, AssetDataGathererConstants::GARDiscoverThreads);
-	AssetDataGathererConstants::GARDiscoverMinBatchSize = FMath::Max(1, AssetDataGathererConstants::GARDiscoverMinBatchSize);
-	AssetDataGathererConstants::GARGatherThreads = FMath::Max(0, AssetDataGathererConstants::GARGatherThreads);
+	GPreloadSettings.Initialize();
 }
 
 FAssetDataDiscovery::~FAssetDataDiscovery()
@@ -3499,7 +3510,7 @@ private:
 class FMemoryMappedFile
 {
 public:
-	FMemoryMappedFile(const TCHAR* Path)
+	explicit FMemoryMappedFile(const TCHAR* Path)
 	{
 		FOpenMappedResult Result = FPlatformFileManager::Get().GetPlatformFile().OpenMappedEx(Path);
 		if (!Result.HasError())
@@ -3630,13 +3641,13 @@ FAssetDataGatherer::FAssetDataGatherer(const TArray<FString>& InLongPackageNames
 	using namespace UE::AssetDataGather::Private;
 
 	int32 NumGatherThreads = FMath::Max(1, FTaskGraphInterface::Get().GetNumWorkerThreads());
+	GPreloadSettings.Initialize();
 	if (AssetDataGathererConstants::GARGatherThreads > 0)
 	{
 		NumGatherThreads = FMath::Min(NumGatherThreads, AssetDataGathererConstants::GARGatherThreads);
 	}
 	TickInternalBatchSize = NumGatherThreads * AssetDataGathererConstants::SingleThreadFilesPerBatch;
 
-	GPreloadSettings.Initialize();
 	constexpr bool bEditorExecutable = WITH_EDITOR;
 	bGatherAssetPackageData = bEditorExecutable || GPreloadSettings.IsForceDependsGathering();
 	bGatherDependsData = GPreloadSettings.IsGatherDependsData();
@@ -5048,17 +5059,15 @@ TArray<FCachePayload> LoadCacheFiles(TConstArrayView<FString> InCacheFilenames)
 	TRACE_CPUPROFILER_EVENT_SCOPE(LoadCacheFiles);
 	TArray<FCachePayload> Results;
 	Results.AddDefaulted(InCacheFilenames.Num());
-	ParallelFor(InCacheFilenames.Num(), [&InCacheFilenames, &Results](int32 Index)
-	{
-		const FString& CacheFilename = InCacheFilenames[Index];
-		auto DoLoad = [&](FArchive& ChecksummingReader)
+
+	auto DoLoad = [](FArchive& ChecksummingReader)
 		{
 			// We are already using AsyncTasks for our own ParallelFor, passing non-zero NumAsyncWorkers into
 			// FAssetRegistryReader will cause our current task to block waiting for another task to run
 			// the name batch loading. That can cause a deadlock due to running out of task threads
 			// if we don't have a large number of threads, so disable the parallelism.
 			constexpr int32 NumAsyncWorkers = 0;
-			
+
 			// The discovery cache is always serialized with a fixed format.
 			// We discard it before this point if it's not the latest version, and it always includes editor-only data.
 			FAssetRegistryHeader Header;
@@ -5068,27 +5077,163 @@ TArray<FCachePayload> LoadCacheFiles(TConstArrayView<FString> InCacheFilenames)
 			return RegistryReader.IsError() ? FCachePayload() : SerializeCacheLoad(RegistryReader);
 		};
 
-		FCachePayload Payload;
-		if (FPlatformProperties::SupportsMemoryMappedFiles())
+	if (FPlatformProperties::SupportsMemoryMappedFiles())
+	{
+		FCriticalSection ParallelForLock;
+		struct FSharedFileData
 		{
-			FMemoryMappedFile File(*CacheFilename);
-			UE::Tasks::FTask Preload = UE::Tasks::Launch(UE_SOURCE_LOCATION, [&]() { File.Preload(); });
+			void ConditionalConstruct(const TCHAR* FilePath)
+			{
+				if (!File)
+				{
+					File.Emplace(FilePath);
+				}
+			}
+			void ConditionalQueueDestruction(TArray<FSharedFileData*>& Queue)
+			{
+				if (bLoadComplete && bAsyncCacheComplete)
+				{
+					Queue.Add(this);
+				}
+			}
+			TOptional<FMemoryMappedFile> File;
+			bool bLoadComplete = false;
+			bool bAsyncCacheComplete = false;
+		};
+
+		// Allocate multiple ParallelFor body invocations for each cache file.
+		// Each group of ParallelFor body invocations will work on the same cache file together in parallel,
+		// or in series if we run out of threads.
+		enum class ETaskType : uint8
+		{
+			// Load byts from the MemoryMapped file and parse its data into Results
+			DoLoad = 0,
+			// Call Preload on the MemoryMapped file to swap its bytes into memory.
+			// This starts at the same time or later than DoLoad, but it will outpace DoLoad
+			// because DoLoad is doing a lot of cpuwork, and so the AsyncCache thread will 
+			// be the one that incurs IO costs rather than the DoLoad thread.
+			AsyncCache,
+			Count,
+		};
+		constexpr int32 TasksPerResult = (int32)ETaskType::Count;
+		TArray<FSharedFileData> CacheFiles;
+		TArray<FSharedFileData*> DestructQueue;
+		CacheFiles.SetNum(InCacheFilenames.Num());
+		DestructQueue.Reserve(InCacheFilenames.Num());
+		const int32 NumTasks = InCacheFilenames.Num() * TasksPerResult;
+		int32 NextTaskIndex = 0;
+
+		// We want to restrict the number of threads to reduce maximum memory use due to having multiple intermediate cache files
+		// in parallel. But ParallelFor only provides an API for restricting MinBatchSize.
+		const int32 MaxNumThreads = AssetDataGathererConstants::GARGatherCacheParallelism * TasksPerResult;
+		const int32 MinBatchSize = FMath::DivideAndRoundUp(NumTasks, MaxNumThreads);
+
+		ParallelFor(TEXT("AssetDataGatherer::LoadCacheFiles"), NumTasks, MinBatchSize,
+			[&DoLoad, &ParallelForLock, &InCacheFilenames, &Results, &CacheFiles, &DestructQueue, &NextTaskIndex]
+			(int32 /* Unused ParallelFor Index */)
+		{
+			FSharedFileData* FileData = nullptr;
+			const FString* CacheFilename = nullptr;
+			bool bShouldAsyncCache = false;
+
+			int32 ResultIndex = -1;
+			ETaskType TaskType = ETaskType::DoLoad;
+			FSharedFileData* DestructData = nullptr;
+
+			do
+			{
+				// Pick up the next TaskIndex, but before doing that, destruct any previously completed CacheFiles,
+				// to reduce the maximum amount of memory we have allocated at once across the multiple threads.
+				DestructData = nullptr;
+				{
+					FScopeLock ParallelForScopeLock(&ParallelForLock);
+
+					// Peek at what our task would be but don't claim it yet.
+					int32 TaskIndex = NextTaskIndex;
+					ResultIndex = TaskIndex / TasksPerResult;
+					TaskType = (ETaskType)(TaskIndex - ResultIndex * TasksPerResult);
+					check(0 <= (int32)TaskType && TaskType < ETaskType::Count);
+
+					if (!DestructQueue.IsEmpty()
+						// Don't do destruct work on the last AsyncCache thread, save it for after the parallelfor
+						&& (TaskType == ETaskType::DoLoad || ResultIndex + 1 < InCacheFilenames.Num()))
+					{
+						DestructData = DestructQueue.Pop(EAllowShrinking::No);
+						check(DestructData != nullptr);
+					}
+					else
+					{
+						// Claim the task.
+						NextTaskIndex++;
+
+						CacheFilename = &InCacheFilenames[ResultIndex];
+						FileData = &CacheFiles[ResultIndex];
+						FileData->ConditionalConstruct(**CacheFilename);
+						bShouldAsyncCache = TaskType == ETaskType::AsyncCache && !FileData->bLoadComplete;
+					}
+				}
 
-			FMemoryViewReader FileReader(File.View());
-			TOptional<uint64> Version = FileReader.TryLoad<uint64>();
-			if (Version == AssetDataGathererConstants::CurrentVersion)
+				if (DestructData)
+				{
+					DestructData->File.Reset();
+				}
+			} while (DestructData != nullptr);
+			check(FileData != nullptr && CacheFilename != nullptr); // Required to inform static analysis.
+
+			if (TaskType == ETaskType::AsyncCache)
 			{
-				FChecksumViewReader ChecksummingReader(MoveTemp(FileReader), CacheFilename);
-				Payload = DoLoad(ChecksummingReader);
-				UE_CLOG(Payload.bSucceeded, LogAssetRegistry, Display, TEXT("Asset registry cache read as %.1f MiB from %s"), static_cast<float>(FileReader.GetTotalSize())/1024.f/1024.f, *CacheFilename);
-				UE_CLOG(!Payload.bSucceeded, LogAssetRegistry, Warning, TEXT("There was an error loading the asset registry cache using memory mapping"));
+				if (bShouldAsyncCache)
+				{
+					FileData->File->Preload();
+				}
+				FScopeLock ParallelForScopeLock(&ParallelForLock);
+				FileData->bAsyncCacheComplete = true;
+				FileData->ConditionalQueueDestruction(DestructQueue);
+				// Don't do any destruction work on our current loop: we are possibly the last iteration
+				// and should allow the ParallelFor to complete, and destruct in a non-blocking task afterwards.
 			}
+			else // ETaskType::DoLoad
+			{
+				FCachePayload Payload;
+				FMemoryViewReader FileReader(FileData->File->View());
+				TOptional<uint64> Version = FileReader.TryLoad<uint64>();
+				if (Version == AssetDataGathererConstants::CurrentVersion)
+				{
+					FChecksumViewReader ChecksummingReader(MoveTemp(FileReader), *CacheFilename);
+					Payload = DoLoad(ChecksummingReader);
+					UE_CLOG(Payload.bSucceeded, LogAssetRegistry, Display,
+						TEXT("Asset registry cache read as %.1f MiB from %s."),
+						static_cast<float>(FileReader.GetTotalSize()) / 1024.f / 1024.f, **CacheFilename);
+					UE_CLOG(!Payload.bSucceeded, LogAssetRegistry, Warning,
+						TEXT("There was an error loading the asset registry cache using memory mapping from %s."),
+						**CacheFilename);
+				}
+				Results[ResultIndex] = MoveTemp(Payload);
 
-			Preload.Wait();
-			UE::Tasks::Launch(UE_SOURCE_LOCATION, [KillAsync = MoveTemp(File)](){});
+				FScopeLock ParallelForScopeLock(&ParallelForLock);
+				FileData->bLoadComplete = true;
+				FileData->ConditionalQueueDestruction(DestructQueue);
+				// Don't do any destruction work on our current loop: we are possibly the last iteration
+				// and should allow the ParallelFor to complete, and destruct in a non-blocking task afterwards.
+			}
+		}, EParallelForFlags::BackgroundPriority);
+		for (FSharedFileData& FileData : CacheFiles)
+		{
+			check(FileData.bLoadComplete && FileData.bAsyncCacheComplete);
 		}
-		else
+
+		// Ignore the remaining elements of the DestructQueue, and instead pass the entire set of CacheFiles to another
+		// thread to finish the destruction of any that still need it.
+		DestructQueue.Empty();
+		UE::Tasks::Launch(UE_SOURCE_LOCATION, [KillAsync = MoveTemp(CacheFiles)]() {});
+	}
+	else
+	{
+		ParallelFor(InCacheFilenames.Num(), [&DoLoad, &InCacheFilenames, &Results](int32 Index)
 		{
+			const FString& CacheFilename = InCacheFilenames[Index];
+
+			FCachePayload Payload;
 			TUniquePtr<FArchive> FileAr(IFileManager::Get().CreateFileReader(*CacheFilename, FILEREAD_Silent));
 			if (FileAr && !FileAr->IsError() && FileAr->TotalSize() > sizeof(uint64))
 			{
@@ -5098,14 +5243,18 @@ TArray<FCachePayload> LoadCacheFiles(TConstArrayView<FString> InCacheFilenames)
 				{
 					FChecksumArchiveReader ChecksummingReader(*FileAr);
 					Payload = DoLoad(ChecksummingReader);
-					UE_CLOG(Payload.bSucceeded, LogAssetRegistry, Display, TEXT("Asset registry cache read as %.1f MiB from %s"), static_cast<float>(FileAr->TotalSize())/1024.f/1024.f, *CacheFilename);
-					UE_CLOG(!Payload.bSucceeded, LogAssetRegistry, Warning, TEXT("There was an error loading the asset registry cache"));
+					UE_CLOG(Payload.bSucceeded, LogAssetRegistry, Display,
+						TEXT("Asset registry cache read as %.1f MiB from %s"),
+						static_cast<float>(FileAr->TotalSize()) / 1024.f / 1024.f, *CacheFilename);
+					UE_CLOG(!Payload.bSucceeded, LogAssetRegistry, Warning,
+						TEXT("There was an error loading the asset registry cache from %s."),
+						*CacheFilename);
 				}
 			}
-		}
-	
-		Results[Index] = MoveTemp(Payload);
-	}, EParallelForFlags::BackgroundPriority);
+
+			Results[Index] = MoveTemp(Payload);
+		}, EParallelForFlags::BackgroundPriority);
+	}
 
 	return MoveTemp(Results);
 }
