AssetDataMap: Fix compile warning: unreachable code.
#rnx

diff --git a/Engine/Source/Runtime/AssetRegistry/Private/AssetDataMap.cpp b/Engine/Source/Runtime/AssetRegistry/Private/AssetDataMap.cpp
--- a/Engine/Source/Runtime/AssetRegistry/Private/AssetDataMap.cpp
+++ b/Engine/Source/Runtime/AssetRegistry/Private/AssetDataMap.cpp
@@ -458,7 +458,7 @@ TConstArrayView<FAssetDataPtrIndex> FIndirectAssetDataArrays::Iterate(const FAss
 	else if (ArrayPtr->IsAssetDataPtrIndex())
 	{
 		static_assert(FAssetDataOrArrayIndex::AssetDataType == 0,
-			"We rely on the converted value for an FAssetDataOrArrayIndex to FAssetDataPtrIndex being the same bits so we can do a reintpret_cast on the pointer.");
+			"We rely on the converted value for an FAssetDataOrArrayIndex to FAssetDataPtrIndex being the same bits so we can do a reinterpret_cast on the pointer.");
 		const FAssetDataPtrIndex* AssetIndexPtr = reinterpret_cast<const FAssetDataPtrIndex*>(ArrayPtr);
 		return TConstArrayView<FAssetDataPtrIndex>(AssetIndexPtr, 1);
 	}
@@ -735,9 +735,10 @@ bool FAssetPackageNameMap::AssetOrArrayByPackageNameValueMatches(FAssetDataOrArr
 
 FAssetData* FAssetPackageNameMap::AssetOrArrayIndexToFirstAssetDataPtr(FAssetDataOrArrayIndex DataOrArrayIndex) const
 {
-	for (FAssetDataPtrIndex AssetIndex : IndirectArrays.Iterate(&DataOrArrayIndex))
+	TConstArrayView<FAssetDataPtrIndex> AssetIndices = IndirectArrays.Iterate(&DataOrArrayIndex);
+	if (AssetIndices.Num() > 0)
 	{
-		return AssetDataMap[AssetIndex];
+		return AssetDataMap[AssetIndices[0]];
 	}
 	return nullptr;
 }
