Added CVar to disable directory watcher in asset registry.

diff --git a/Engine/Source/Runtime/AssetRegistry/Private/AssetRegistry.cpp b/Engine/Source/Runtime/AssetRegistry/Private/AssetRegistry.cpp
--- a/Engine/Source/Runtime/AssetRegistry/Private/AssetRegistry.cpp
+++ b/Engine/Source/Runtime/AssetRegistry/Private/AssetRegistry.cpp
@@ -188,6 +188,18 @@ namespace UE::AssetRegistry::Impl
 		UE::AssetRegistry::Impl::bDeferReferencerSort,
 		TEXT("If true, the referencer list on dependency nodes will not be sorted until after the initial load is complete"));
 
+	bool bDisableDirectoryWatcher = false;
+	static FAutoConsoleVariableRef CVarAssetRegistryDisableDirectoryWatcher(
+		TEXT("AssetRegistry.DisableDirectoryWatcher"),
+		UE::AssetRegistry::Impl::bDisableDirectoryWatcher,
+		TEXT("If true, do not listen to mounted directories for file changes"));
+
+	bool IsDirectoryWatcherEnabled()
+	{
+		// In-game doesn't listen for directory changes
+		return GIsEditor && !bDisableDirectoryWatcher;
+	}
+
 	/** Name of UObjectRedirector property */
 	const FName DestinationObjectFName(TEXT("DestinationObject"));
 }
@@ -1185,8 +1197,7 @@ void UAssetRegistryImpl::InitializeEvents(UE::AssetRegistry::Impl::FInitializeCo
 	}
 
 #if WITH_EDITOR
-	// In-game doesn't listen for directory changes
-	if (GIsEditor)
+	if (UE::AssetRegistry::Impl::IsDirectoryWatcherEnabled())
 	{
 		FDirectoryWatcherModule& DirectoryWatcherModule = FModuleManager::LoadModuleChecked<FDirectoryWatcherModule>(TEXT("DirectoryWatcher"));
 		IDirectoryWatcher* DirectoryWatcher = DirectoryWatcherModule.Get();
@@ -1785,7 +1796,7 @@ void UAssetRegistryImpl::FinishDestroy()
 		IPluginManager::Get().OnLoadingPhaseComplete().RemoveAll(this);
 
 #if WITH_EDITOR
-		if (GIsEditor)
+		if (UE::AssetRegistry::Impl::IsDirectoryWatcherEnabled())
 		{
 			// If the directory module is still loaded, unregister any delegates
 			if (FModuleManager::Get().IsModuleLoaded("DirectoryWatcher"))
@@ -2088,7 +2099,7 @@ void UAssetRegistryImpl::WaitForCompletion()
 		// If we do need to wait, then tick the DirectoryWatcher so we have the most up to date information.
 		// This is also important because we ignore rescan events from the directory watcher if they are sent
 		// during startup, so if there is a rescan event pending we want to trigger it now and ignore it
-		if (GIsEditor) 	// In-game doesn't listen for directory changes
+		if (UE::AssetRegistry::Impl::IsDirectoryWatcherEnabled())
 		{
 			FDirectoryWatcherModule& DirectoryWatcherModule = FModuleManager::LoadModuleChecked<FDirectoryWatcherModule>(TEXT("DirectoryWatcher"));
 			DirectoryWatcherModule.Get()->Tick(-1.f);
@@ -8136,7 +8147,7 @@ void UAssetRegistryImpl::OnContentPathMounted(const FString& InAssetPath, const
 #if WITH_EDITOR
 	FDirectoryWatcherModule& DirectoryWatcherModule = FModuleManager::LoadModuleChecked<FDirectoryWatcherModule>(TEXT("DirectoryWatcher"));
 	IDirectoryWatcher* DirectoryWatcher = nullptr;
-	if (GIsEditor) 	// In-game doesn't listen for directory changes
+	if (UE::AssetRegistry::Impl::IsDirectoryWatcherEnabled())
 	{
 		LLM_SCOPE(ELLMTag::AssetRegistry);
 		DirectoryWatcher = DirectoryWatcherModule.Get();
@@ -8231,7 +8242,7 @@ void UAssetRegistryImpl::OnContentPathDismounted(const FString& InAssetPath, con
 #if WITH_EDITOR
 	FDirectoryWatcherModule& DirectoryWatcherModule = FModuleManager::LoadModuleChecked<FDirectoryWatcherModule>(TEXT("DirectoryWatcher"));
 	IDirectoryWatcher* DirectoryWatcher = nullptr;
-	if (GIsEditor) 	// In-game doesn't listen for directory changes
+	if (UE::AssetRegistry::Impl::IsDirectoryWatcherEnabled())
 	{
 		DirectoryWatcher = DirectoryWatcherModule.Get();
 	}
