AssetRegistryHelpers: Add ASSETREGISTRY_API to the GetBlueprintAssets function since it is non-trivial and is useful to call from c++ functions outside of the module.
#rnx
#12727

diff --git a/Engine/Source/Runtime/AssetRegistry/Public/AssetRegistry/AssetRegistryHelpers.h b/Engine/Source/Runtime/AssetRegistry/Public/AssetRegistry/AssetRegistryHelpers.h
--- a/Engine/Source/Runtime/AssetRegistry/Public/AssetRegistry/AssetRegistryHelpers.h
+++ b/Engine/Source/Runtime/AssetRegistry/Public/AssetRegistry/AssetRegistryHelpers.h
@@ -107,7 +107,7 @@ public:
 
 	/** Gets asset data for all blueprint assets that match the filter. ClassPaths in the filter specify the blueprint's parent class. */
 	UFUNCTION(BlueprintPure, Category = "Asset Registry", meta=(ScriptMethod))
-	static void GetBlueprintAssets(const FARFilter& InFilter, TArray<FAssetData>& OutAssetData);
+	static ASSETREGISTRY_API void GetBlueprintAssets(const FARFilter& InFilter, TArray<FAssetData>& OutAssetData);
 
 	/**
 	 * Returns the first native class of the asset type that can be found.  Normally this is just the FAssetData::GetClass(),
