Add deprecation warning log in CookedPackageNamesWithoutAssetDataGathered.

#rnx

diff --git a/Engine/Source/Runtime/AssetRegistry/Private/AssetRegistry.cpp b/Engine/Source/Runtime/AssetRegistry/Private/AssetRegistry.cpp
--- a/Engine/Source/Runtime/AssetRegistry/Private/AssetRegistry.cpp
+++ b/Engine/Source/Runtime/AssetRegistry/Private/AssetRegistry.cpp
@@ -7056,6 +7056,10 @@ void FAssetRegistryImpl::CookedPackageNamesWithoutAssetDataGathered(Impl::FEvent
 			// If this data is cooked and it we couldn't find any asset in its export table then try to load the entire package 
 			// Loading the entire package will make all of its assets searchable through the in-memory scanning performed by GetAssets
 			EventContext.RequiredLoads.Add(CookedPackageNamesWithoutAssetDataResults.PopFrontValue());
+
+			// This function has a bug in multithreaded environment (ie UE-209843). But the feature of loading cooked packages seems to be never used
+			// so let's try to deprecate it. Also a new way to load cooked package is on the way.
+			UE_LOG(LogAssetRegistry, Warning, TEXT("CookedPackageNamesWithoutAssetDataGathered : Deprecated in 5.6 due to poor performance, will be removed in a future version, contact Epic if you need this functionality."));
 		}
 		// Avoid marking the scan complete before we have loaded all the relevant assets. By interrupting here
 		// we intend to ensure that the event context is processed, triggering a LoadPackage, and then a ProcessLoadedAssetsToUpdateCache,
