Adding a new asset registry callback that will be invoked after the initial asset scan (like the current OnFilesLoaded) and then again after any subsequent gathers. Improved support for async gather functionality.

diff --git a/Engine/Config/BaseEngine.ini b/Engine/Config/BaseEngine.ini
--- a/Engine/Config/BaseEngine.ini
+++ b/Engine/Config/BaseEngine.ini
@@ -2900,6 +2900,9 @@ bFilterSearchableNames=true
 ; the same way the editor does to have the same AssetRegistry available that the cooked game has. This adds cputime and
 ; memory cost to the editor game startup, so some projects may want to turn it off for their game when running editor -game.
 ; Set EditorGameScansAR=false to turn it off.
+; Set EditorGameScansAR=true to turn it on (default). Will use async scanning for -game and, when perf threads are enabled, -server
+; Set EditorGameScansAR=sync to enable scans but run them synchronously
+; Set EditorGameScansAR=async (same as 'true')
 EditorGameScansAR=true
 
 ; Performance settings for the AssetRegistryScan
diff --git a/Engine/Source/Runtime/AssetRegistry/Private/AssetRegistry.cpp b/Engine/Source/Runtime/AssetRegistry/Private/AssetRegistry.cpp
--- a/Engine/Source/Runtime/AssetRegistry/Private/AssetRegistry.cpp
+++ b/Engine/Source/Runtime/AssetRegistry/Private/AssetRegistry.cpp
@@ -207,6 +207,43 @@ namespace UE::AssetRegistry::Impl
 namespace UE::AssetRegistry
 {
 
+/**
+ * How the editor, when run with -game or (under more limited circumstances) -server, performs its asset registry scan
+ */
+enum class EEditorGameScanMode
+{
+	Sync,
+	Async,
+	None
+};
+
+EEditorGameScanMode GetEditorGameScanMode()
+{
+	FString EditorGameScansAR = TEXT("Sync");
+	EEditorGameScanMode Result = EEditorGameScanMode::Sync;
+
+	GConfig->GetString(TEXT("AssetRegistry"), TEXT("EditorGameScansAR"), EditorGameScansAR, GEngineIni);
+	if (EditorGameScansAR.Equals(TEXT("Async"), ESearchCase::IgnoreCase))
+	{
+		Result = EEditorGameScanMode::Async;
+	}
+	else if (EditorGameScansAR.Equals(TEXT("False"), ESearchCase::IgnoreCase))
+	{
+		Result = EEditorGameScanMode::None;
+	}
+	else if (EditorGameScansAR.Equals(TEXT("True"), ESearchCase::IgnoreCase))
+	{
+		Result = EEditorGameScanMode::Async;
+	}
+	else
+	{
+		ensureMsgf(EditorGameScansAR.Equals(TEXT("Sync"), ESearchCase::IgnoreCase), 
+			TEXT("Valid values for EditorGameScansAR are: true|false|sync|async. Received %s"), *EditorGameScansAR);
+	}
+
+	return Result;
+}
+
 /** 
  * Tracks high-priority waiters on a lock, in an thread- and transactionally-safe manner.
  */
@@ -1902,7 +1939,14 @@ bool FAssetRegistryImpl::TryConstructGathererIfNeeded()
 		EngineIni->GetArray(TEXT("AssetRegistry"), TEXT("BlacklistContentSubPathScanFilters"), ContentSubPathsDenyList);
 	}
 
-	bool bAsyncGatherEnabled = !IsRunningGame() && !IsRunningDedicatedServer();
+	// If EditorGameScanMode is set to allow async then we will use async for any editor build (editor, game, or server)
+	// Otherwise we will use async for commandlets and editor proper
+	bool bAsyncGatherEnabled = (!IsRunningGame() && !IsRunningDedicatedServer())
+#if WITH_EDITOR
+		|| (GetEditorGameScanMode() == EEditorGameScanMode::Async)
+#endif
+		;
+
 	GlobalGatherer = MakeUnique<FAssetDataGatherer>(PathsDenyList, ContentSubPathsDenyList, bAsyncGatherEnabled, *this);
 	UpdateMaxSecondsPerFrame();
 
@@ -5374,7 +5418,8 @@ void FAssetRegistryImpl::OnInitialSearchCompleted(Impl::FEventContext& EventCont
 	GlobalGatherer->OnInitialSearchCompleted();
 
 	EventContext.bFileLoadedEventBroadcast = true;
-	
+	EventContext.bKnownGathersCompleteEventBroadcast = true;
+
 	bInitialSearchCompleted.store(true, std::memory_order_relaxed);
 	UpdateMaxSecondsPerFrame();
 }
@@ -5394,6 +5439,8 @@ void FAssetRegistryImpl::OnAdditionalMountSearchCompleted(Impl::FEventContext& E
 
 	GlobalGatherer->OnAdditionalMountSearchCompleted();
 
+	EventContext.bKnownGathersCompleteEventBroadcast = true;
+
 	bAdditionalMountSearchInProgress.store(false, std::memory_order_relaxed);
 }
 
@@ -9202,6 +9249,7 @@ void FEventContext::Clear()
 {
 	bScanStartedEventBroadcast = false;
 	bFileLoadedEventBroadcast = false;
+	bKnownGathersCompleteEventBroadcast = false;
 	bHasSentFileLoadedEventBroadcast = false;
 	ProgressUpdateData.Reset();
 	PathEvents.Empty();
@@ -9214,6 +9262,7 @@ bool FEventContext::IsEmpty() const
 {
 	return !bScanStartedEventBroadcast &&
 		!bFileLoadedEventBroadcast &&
+		!bKnownGathersCompleteEventBroadcast &&
 		!ProgressUpdateData.IsSet() &&
 		PathEvents.Num() == 0 &&
 		AssetEvents.Num() == 0 &&
@@ -9231,6 +9280,8 @@ void FEventContext::Append(FEventContext&& Other)
 	Other.bScanStartedEventBroadcast = false;
 	bFileLoadedEventBroadcast |= Other.bFileLoadedEventBroadcast;
 	Other.bFileLoadedEventBroadcast = false;
+	bKnownGathersCompleteEventBroadcast |= Other.bKnownGathersCompleteEventBroadcast;
+	Other.bKnownGathersCompleteEventBroadcast = false;
 	if (Other.ProgressUpdateData.IsSet())
 	{
 		ProgressUpdateData = MoveTemp(Other.ProgressUpdateData);
@@ -9445,8 +9496,11 @@ void UAssetRegistryImpl::Broadcast(UE::AssetRegistry::Impl::FEventContext& Event
 	// their FileLoadedEvent handler. The AssetEvents precede the FileLoadedEvent in the broadcast that is sent from
 	// TickGatherer, so it is correct to make them precede it in the order in which we broadcast the events.
 
-	if (EventContext.bFileLoadedEventBroadcast)
+	if (EventContext.bFileLoadedEventBroadcast || EventContext.bKnownGathersCompleteEventBroadcast)
 	{
+		bool bLocalFileLoadedEventBroadcast = EventContext.bFileLoadedEventBroadcast;
+		bool bLocalKnownGathersCompleteEventBroadcast = EventContext.bKnownGathersCompleteEventBroadcast;
+
 		if (!bAllowFileLoadedEvent)
 		{
 			// Do not send the file loaded event yet; pass the flag on instead
@@ -9456,8 +9510,9 @@ void UAssetRegistryImpl::Broadcast(UE::AssetRegistry::Impl::FEventContext& Event
 			check(&EventContext != &DeferredEvents);
 			DeferredEvents.Append(MoveTemp(EventContext));
 			EventContext.Clear();
-			check(!EventContext.bFileLoadedEventBroadcast); // was cleared by Append and by Clear
-			check(DeferredEvents.bFileLoadedEventBroadcast); // was set by Append
+			check(!EventContext.bFileLoadedEventBroadcast && !EventContext.bKnownGathersCompleteEventBroadcast); // was cleared by Append and by Clear
+			check(!bLocalFileLoadedEventBroadcast || DeferredEvents.bFileLoadedEventBroadcast); // was set by Append
+			check(!bLocalKnownGathersCompleteEventBroadcast || DeferredEvents.bKnownGathersCompleteEventBroadcast); // was set by Append
 			return;
 		}
 
@@ -9470,19 +9525,33 @@ void UAssetRegistryImpl::Broadcast(UE::AssetRegistry::Impl::FEventContext& Event
 		}
 		if (!CopiedDeferredEvents.IsEmpty())
 		{
-			// Recursively send all of the deferred events, except for the fileloaded event (which should not exist
-			// on DeferredEvents at this point, but it's not a problem for us if it does).
+			// Recursively send all of the deferred events, except for the completion events
+			// (FileLoaded, KnownGathersComplete). The completion events should not exist on DeferredEvents at this
+			// point, but it's not a problem if they do; merge them into our Local variables.
+			bLocalFileLoadedEventBroadcast |= CopiedDeferredEvents.bFileLoadedEventBroadcast;
 			CopiedDeferredEvents.bFileLoadedEventBroadcast = false;
+
+			bLocalKnownGathersCompleteEventBroadcast |= CopiedDeferredEvents.bKnownGathersCompleteEventBroadcast;
+			CopiedDeferredEvents.bKnownGathersCompleteEventBroadcast = false;
+
 			Broadcast(CopiedDeferredEvents, false /* bAllowFileLoadedEvent */);
 		}
-		// Now it is safe to broadcast the FileLoadedEvent. If other deferred events come in on another thread after we
-		// copied from DeferredEvents, that is okay; the contract for FileLoadedEvent is that it can not be sent before
-		// sending the events that were fired before it was fired, and new deferred events were fired after it.
-
-		FileLoadedEvent.Broadcast();
-		ScanEndedEvent.Broadcast();
-		EventContext.bFileLoadedEventBroadcast = false;
-		EventContext.bHasSentFileLoadedEventBroadcast = true;
+		// Now it is safe to broadcast the completion events. If other deferred events come in on another thread
+		// after we copied from DeferredEvents, that is okay; the contract for completion events is that they are
+		// guaranteed to be sent after any non-completion events sent before completion was triggered,
+		// but they can be before or after non-completion events that occurred after completion was triggered.
+		if (bLocalFileLoadedEventBroadcast)
+		{
+			FileLoadedEvent.Broadcast();
+			ScanEndedEvent.Broadcast();
+			EventContext.bHasSentFileLoadedEventBroadcast = true;
+			EventContext.bFileLoadedEventBroadcast = false;
+		}
+		if (bLocalKnownGathersCompleteEventBroadcast)
+		{
+			KnownGathersCompleteEvent.Broadcast();
+			EventContext.bKnownGathersCompleteEventBroadcast = false;
+		}
 	}
 }
 
@@ -9602,6 +9671,11 @@ UAssetRegistryImpl::FScanEndedEvent& UAssetRegistryImpl::OnScanEnded()
 	return ScanEndedEvent;
 }
 
+UAssetRegistryImpl::FKnownGathersCompleteEvent& UAssetRegistryImpl::OnKnownGathersComplete()
+{
+	return KnownGathersCompleteEvent;
+}
+
 namespace UE::AssetRegistry
 {
 const FAssetData* GetMostImportantAsset(TConstArrayView<const FAssetData*> PackageAssetDatas, EGetMostImportantAssetFlags InFlags)
@@ -9775,9 +9849,7 @@ bool ShouldSearchAllAssetsAtStart()
 #if WITH_EDITOR
 	else
 	{
-		bool bEditorGameScansAR = true;
-		GConfig->GetBool(TEXT("AssetRegistry"), TEXT("EditorGameScansAR"), bEditorGameScansAR, GEngineIni);
-		bSearchAllAssetsAtStart = bEditorGameScansAR;
+		bSearchAllAssetsAtStart = GetEditorGameScanMode() != EEditorGameScanMode::None;
 	}
 #endif
 #if WITH_EDITOR || !UE_BUILD_SHIPPING
diff --git a/Engine/Source/Runtime/AssetRegistry/Private/AssetRegistry.h b/Engine/Source/Runtime/AssetRegistry/Private/AssetRegistry.h
--- a/Engine/Source/Runtime/AssetRegistry/Private/AssetRegistry.h
+++ b/Engine/Source/Runtime/AssetRegistry/Private/AssetRegistry.h
@@ -147,6 +147,9 @@ public:
 	DECLARE_DERIVED_EVENT(UAssetRegistryImpl, IAssetRegistry::FScanEndedEvent, FScanEndedEvent);
 	virtual FScanEndedEvent& OnScanEnded() override;
 
+	DECLARE_DERIVED_EVENT(UAssetRegistryImpl, IAssetRegistry::FKnownGathersCompleteEvent, FKnownGathersCompleteEvent);
+	virtual FKnownGathersCompleteEvent& OnKnownGathersComplete() override;
+
 	virtual void AssetCreated(UObject* NewAsset) override;
 	virtual void AssetDeleted(UObject* DeletedAsset) override;
 	virtual void AssetRenamed(const UObject* RenamedAsset, const FString& OldObjectPath) override;
@@ -416,6 +419,9 @@ private:
 	/** The delegate to execute scanning has ended */
 	FScanEndedEvent ScanEndedEvent;
 
+	/** The delegate to execute when each batch of searching/gathering is complete */
+	FKnownGathersCompleteEvent KnownGathersCompleteEvent;
+
 	/** Storage for events that will be broadcast later on the game thread. Only safe
 	 *  to access under the DeferredEventsCriticalSection. 
 	 */
diff --git a/Engine/Source/Runtime/AssetRegistry/Private/AssetRegistryImpl.h b/Engine/Source/Runtime/AssetRegistry/Private/AssetRegistryImpl.h
--- a/Engine/Source/Runtime/AssetRegistry/Private/AssetRegistryImpl.h
+++ b/Engine/Source/Runtime/AssetRegistry/Private/AssetRegistryImpl.h
@@ -739,6 +739,7 @@ struct FEventContext
 	bool bFileLoadedEventBroadcast = false;
 	bool bScanStartedEventBroadcast = false;
 	bool bHasSentFileLoadedEventBroadcast = false;
+	bool bKnownGathersCompleteEventBroadcast = false;
 
 	/** Remove all stored events */
 	void Clear();
diff --git a/Engine/Source/Runtime/AssetRegistry/Public/AssetRegistry/IAssetRegistry.h b/Engine/Source/Runtime/AssetRegistry/Public/AssetRegistry/IAssetRegistry.h
--- a/Engine/Source/Runtime/AssetRegistry/Public/AssetRegistry/IAssetRegistry.h
+++ b/Engine/Source/Runtime/AssetRegistry/Public/AssetRegistry/IAssetRegistry.h
@@ -828,6 +828,12 @@ public:
 	DECLARE_TS_MULTICAST_DELEGATE(FScanEndedEvent);
 	virtual FScanEndedEvent& OnScanEnded() = 0;
 
+	/** Event for when the asset registry is done with all currently known gathers. This will fire after both the
+	 *  initial asset scan and again after each subsequent batch of additional mount searches
+	 */
+	DECLARE_TS_MULTICAST_DELEGATE(FKnownGathersCompleteEvent);
+	virtual FKnownGathersCompleteEvent& OnKnownGathersComplete() = 0;
+
 	/** Payload data for a file progress update */
 	struct FFileLoadProgressUpdateData
 	{
