[AR] Extra logging for AssetRegistry size analysis

#rnx

diff --git a/Engine/Source/Runtime/AssetRegistry/Private/AssetRegistry.cpp b/Engine/Source/Runtime/AssetRegistry/Private/AssetRegistry.cpp
--- a/Engine/Source/Runtime/AssetRegistry/Private/AssetRegistry.cpp
+++ b/Engine/Source/Runtime/AssetRegistry/Private/AssetRegistry.cpp
@@ -5882,6 +5882,14 @@ void FAssetRegistryImpl::GetAllocatedSize(bool bLogDetailed, SIZE_T& StateSize,
 {
 	StateSize = State.GetAllocatedSize(bLogDetailed);
 
+	if (bLogDetailed)
+	{
+		UE_LOG(LogAssetRegistry, Log, TEXT("AssetRegistry CachedEmptyPackages Size: %" SIZE_T_FMT "k"), CachedEmptyPackages.GetAllocatedSize() / 1024);
+		UE_LOG(LogAssetRegistry, Log, TEXT("AssetRegistry CachedBPInheritanceMap Size: %" SIZE_T_FMT "k"), CachedBPInheritanceMap.GetAllocatedSize() / 1024);
+		UE_LOG(LogAssetRegistry, Log, TEXT("AssetRegistry ClassGeneratorNames Size: %" SIZE_T_FMT "k"), ClassGeneratorNames.GetAllocatedSize() / 1024);
+		UE_LOG(LogAssetRegistry, Log, TEXT("AssetRegistry BackgroundResults Size: %" SIZE_T_FMT "k"), BackgroundResults.GetAllocatedSize() / 1024);
+		UE_LOG(LogAssetRegistry, Log, TEXT("AssetRegistry CachedPathTree Size: %" SIZE_T_FMT "k"), CachedPathTree.GetAllocatedSize() / 1024);
+	}
 	StaticSize = CachedEmptyPackages.GetAllocatedSize() + CachedBPInheritanceMap.GetAllocatedSize() + ClassGeneratorNames.GetAllocatedSize();
 	SearchSize = BackgroundResults.GetAllocatedSize() + CachedPathTree.GetAllocatedSize();
 
@@ -5896,6 +5904,11 @@ void FAssetRegistryImpl::GetAllocatedSize(bool bLogDetailed, SIZE_T& StateSize,
 	{
 		SearchSize += sizeof(*GlobalGatherer);
 		SearchSize += GlobalGatherer->GetAllocatedSize();
+		
+		if (bLogDetailed)
+		{
+			UE_LOG(LogAssetRegistry, Log, TEXT("AssetRegistry GlobalGatherer Size: %" SIZE_T_FMT "k"), GlobalGatherer->GetAllocatedSize() / 1024);
+		}
 	}
 
 	StaticSize += SerializationOptions.CookFilterlistTagsByClass.GetAllocatedSize();
diff --git a/Engine/Source/Runtime/AssetRegistry/Private/AssetRegistryState.cpp b/Engine/Source/Runtime/AssetRegistry/Private/AssetRegistryState.cpp
--- a/Engine/Source/Runtime/AssetRegistry/Private/AssetRegistryState.cpp
+++ b/Engine/Source/Runtime/AssetRegistry/Private/AssetRegistryState.cpp
@@ -2206,6 +2206,7 @@ SIZE_T FAssetRegistryState::GetAllocatedSize(bool bLogDetailed) const
 		UE_LOG(LogAssetRegistry, Log, TEXT("AssetData Count: %d"), CachedAssets.Num());
 		UE_LOG(LogAssetRegistry, Log, TEXT("AssetData Static Size: %" SIZE_T_FMT "k"), AssetDataSize / 1024);
 		UE_LOG(LogAssetRegistry, Log, TEXT("Loose Tags: %" SIZE_T_FMT "k"), TagMemoryUsage.GetLooseSize() / 1024);
+		TagMemoryUsage.ReportFixedStoreBreakdown();
 		UE_LOG(LogAssetRegistry, Log, TEXT("Fixed Tags: %" SIZE_T_FMT "k"), TagMemoryUsage.GetFixedSize() / 1024);
 		UE_LOG(LogAssetRegistry, Log, TEXT("TArray<FAssetData*>: %" SIZE_T_FMT "k"), MapArrayMemory / 1024);
 		UE_LOG(LogAssetRegistry, Log, TEXT("AssetBundle Count: %u"), NumAssetBundles);
diff --git a/Engine/Source/Runtime/CoreUObject/Private/AssetRegistry/AssetDataTagMap.cpp b/Engine/Source/Runtime/CoreUObject/Private/AssetRegistry/AssetDataTagMap.cpp
--- a/Engine/Source/Runtime/CoreUObject/Private/AssetRegistry/AssetDataTagMap.cpp
+++ b/Engine/Source/Runtime/CoreUObject/Private/AssetRegistry/AssetDataTagMap.cpp
@@ -1698,6 +1698,17 @@ SIZE_T FAssetDataTagMapSharedView::FMemoryCounter::GetFixedSize() const
 	return Out;
 }
 
+void FAssetDataTagMapSharedView::FMemoryCounter::ReportFixedStoreBreakdown() const
+{
+	for (uint32 StoreIndex : FixedStoreIndices)
+	{
+		uint32 Idx = 0;
+		VisitViews(FixedTagPrivate::GStores[StoreIndex], [StoreIndex, &Idx] (auto& View) {
+			UE_LOG(LogAssetDataTags, Log, TEXT("Store[%d]:View[%d]: Num=%d Size=%" SIZE_T_FMT "k"), StoreIndex, Idx++, View.Num(), View.Num() * View.GetTypeSize() / 1024);
+		});
+	}
+}
+
 ////////////////////////////////////////////////////////////////////////////
 
 #if WITH_DEV_AUTOMATION_TESTS 
diff --git a/Engine/Source/Runtime/CoreUObject/Public/AssetRegistry/AssetDataTagMap.h b/Engine/Source/Runtime/CoreUObject/Public/AssetRegistry/AssetDataTagMap.h
--- a/Engine/Source/Runtime/CoreUObject/Public/AssetRegistry/AssetDataTagMap.h
+++ b/Engine/Source/Runtime/CoreUObject/Public/AssetRegistry/AssetDataTagMap.h
@@ -478,7 +478,8 @@ public:
 	public:
 		void Include(const FAssetDataTagMapSharedView& Tags);
 		SIZE_T GetLooseSize() const { return LooseBytes; }
-		SIZE_T GetFixedSize() const;		
+		SIZE_T GetFixedSize() const;
+		void ReportFixedStoreBreakdown() const;
 	};
 
 	friend inline bool operator==(const FAssetDataTagMap& A, const FAssetDataTagMapSharedView& B)			{ return B == A; }
