UAssetRegistryHelpers::GetAssetsWithOuterForPaths now supports returning assets outered to any part of the outer chain.

diff --git a/Engine/Source/Runtime/AssetRegistry/Private/AssetRegistryHelpers.cpp b/Engine/Source/Runtime/AssetRegistry/Private/AssetRegistryHelpers.cpp
--- a/Engine/Source/Runtime/AssetRegistry/Private/AssetRegistryHelpers.cpp
+++ b/Engine/Source/Runtime/AssetRegistry/Private/AssetRegistryHelpers.cpp
@@ -335,7 +335,7 @@ void UAssetRegistryHelpers::FixupRedirectedAssetPath(FName& InOutAssetPath)
 }
 
 #if WITH_EDITOR
-TArray<FAssetData> UAssetRegistryHelpers::GetAssetsWithOuterForPaths(const TArray<FName>& InPackagePaths, FName InOuterPath, bool bInRecursivePaths, bool bInIncludeOnlyOnDiskAssets)
+TArray<FAssetData> UAssetRegistryHelpers::GetAssetsWithOuterForPaths(const TArray<FName>& InPackagePaths, FName InOuterPath, bool bInRecursivePaths, bool bInIncludeOnlyOnDiskAssets, bool bInExactOuter)
 {
 	FARFilter Filter;
 	Filter.PackagePaths = InPackagePaths;
@@ -351,7 +351,21 @@ TArray<FAssetData> UAssetRegistryHelpers::GetAssetsWithOuterForPaths(const TArra
 	AssetRegistry.ScanSynchronous(PathsToScan, TArray<FString>());
 	AssetRegistry.GetAssets(Filter, Assets);
 
-	Assets.SetNum(Algo::RemoveIf(Assets, [InOuterPath](const FAssetData& InAssetData) { return InAssetData.GetOptionalOuterPathName() != InOuterPath; }));
+	Assets.SetNum(
+		Algo::RemoveIf(Assets, [InOuterPath, bInExactOuter](const FAssetData& InAssetData)
+		{
+			if (InAssetData.GetOptionalOuterPathName() == InOuterPath)
+			{
+				return false;
+			}
+
+			if (!bInExactOuter && InAssetData.GetOptionalOuterPathName().ToString().StartsWith(InOuterPath.ToString()))
+			{
+				return false;
+			}
+
+			return true;
+		}));
 	Assets.Sort();
 
 	return Assets;
diff --git a/Engine/Source/Runtime/AssetRegistry/Public/AssetRegistry/AssetRegistryHelpers.h b/Engine/Source/Runtime/AssetRegistry/Public/AssetRegistry/AssetRegistryHelpers.h
--- a/Engine/Source/Runtime/AssetRegistry/Public/AssetRegistry/AssetRegistryHelpers.h
+++ b/Engine/Source/Runtime/AssetRegistry/Public/AssetRegistry/AssetRegistryHelpers.h
@@ -179,7 +179,8 @@ public:
 	 * @param	InOuterPath The outer path to filter asset with.
 	 * @param	bInRecursivePaths If we should scan recursively or not.
 	 * @param	bInIncludeOnlyOnDiskAsset If we should only include on disk assets.
+	 * @param	bInExactOuter If we should only return assets directly outered to the provided outer.
 	 */
-	 static ASSETREGISTRY_API TArray<FAssetData> GetAssetsWithOuterForPaths(const TArray<FName>& InPackagePaths, FName InOuterPath, bool bInRecursivePaths = false, bool bInIncludeOnlyOnDiskAsset = false);
+	 static ASSETREGISTRY_API TArray<FAssetData> GetAssetsWithOuterForPaths(const TArray<FName>& InPackagePaths, FName InOuterPath, bool bInRecursivePaths = false, bool bInIncludeOnlyOnDiskAsset = false, bool bInExactOuter = true);
 #endif
 };
