AssetRegistry: When the gatherer thread is waiting on the discovery thread and we sleep rather than busy wait, sleep outside of the critical section, to avoid blocking the gamethread if a query for a package comes in from the game thread while we are sleeping.
#rnx

diff --git a/Engine/Source/Runtime/AssetRegistry/Private/AssetDataGatherer.cpp b/Engine/Source/Runtime/AssetRegistry/Private/AssetDataGatherer.cpp
--- a/Engine/Source/Runtime/AssetRegistry/Private/AssetDataGatherer.cpp
+++ b/Engine/Source/Runtime/AssetRegistry/Private/AssetDataGatherer.cpp
@@ -75,6 +75,9 @@ namespace AssetDataGathererConstants
 	 * Other than the memory cost, it should be set as high as possible to load the cache faster.
 	 */
 	static int32 GARGatherCacheParallelism = 16;
+
+	/** How long we wait on the gatherer thread between asking the discovery thread for no results. */
+	constexpr float PollDiscoveryPeriodSeconds = .005f;
 }
 
 namespace UE::AssetDataGather::Private
@@ -3723,7 +3726,8 @@ uint32 FAssetDataGatherer::Run()
 
 	while (!IsStopped)
 	{
-		InnerTickLoop(false /* bInSynchronousTick */, true /* bContributeToCacheSave */, -1. /* EndTimeSeconds */);
+		ETickResult TickResult = InnerTickLoop(false /* bInSynchronousTick */, true /* bContributeToCacheSave */,
+			-1. /* EndTimeSeconds */);
 
 		for (;;)
 		{
@@ -3756,14 +3760,24 @@ uint32 FAssetDataGatherer::Run()
 				TRACE_CPUPROFILER_EVENT_SCOPE_STR_CONDITIONAL("FAssetDataGatherer Sleep", bShouldLog);
 				FPlatformProcess::Sleep(bLocalIdle ? IdleSleepTime : PausedSleepTime);
 			}
+
+			if (TickResult == ETickResult::PollDiscovery)
+			{
+				// The gatherer thread is waiting on results from the discovery thread, and we should sleep rather
+				// than busy wait on it, to reduce contention. TODO: Change this to an event triggered by discovery
+				// results, rather than sleeping for a fixed time interval.
+				FPlatformProcess::Sleep(AssetDataGathererConstants::PollDiscoveryPeriodSeconds);
+			}
 		}
 	}
 	return 0;
 }
 
-void FAssetDataGatherer::InnerTickLoop(bool bInSynchronousTick, bool bContributeToCacheSave, double EndTimeSeconds)
+FAssetDataGatherer::ETickResult FAssetDataGatherer::InnerTickLoop(bool bInSynchronousTick,
+	bool bContributeToCacheSave, double EndTimeSeconds)
 {
 	using namespace UE::AssetDataGather::Private;
+	ETickResult Result = ETickResult::KeepTicking;
 
 	// Synchronous ticks during Wait contribute to saving of the async cache only if there is no dedicated async thread to do it (Thread == nullptr)
 	// The dedicated async thread always contributes
@@ -3778,7 +3792,6 @@ void FAssetDataGatherer::InnerTickLoop(bool bInSynchronousTick, bool bContribute
 		TRACE_CPUPROFILER_EVENT_SCOPE(FAssetDataGatherer::Tick);
 		double TickStartTime = FPlatformTime::Seconds();
 		bool bPollDiscovery = true;
-		constexpr float PollDiscoveryPeriodSeconds = .005f;
 		double LastPollTimeSeconds = 0;
 
 		for (;;)
@@ -3786,14 +3799,17 @@ void FAssetDataGatherer::InnerTickLoop(bool bInSynchronousTick, bool bContribute
 			ETickResult TickResult = TickInternal(TickStartTime, bPollDiscovery);
 			if (EndTimeSeconds > 0. && FPlatformTime::Seconds() > EndTimeSeconds)
 			{
+				Result = ETickResult::Interrupt;
 				break;
 			}
 			if (IsStopped || (!bInSynchronousTick && IsGatheringPaused))
 			{
+				Result = ETickResult::Idle;
 				break;
 			}
 			if (TickResult != ETickResult::KeepTicking && TickResult != ETickResult::PollDiscovery)
 			{
+				Result = TickResult; // Interrupt or Idle
 				break;
 			}
 			double CurrentTimeSeconds = FPlatformTime::Seconds();
@@ -3804,35 +3820,27 @@ void FAssetDataGatherer::InnerTickLoop(bool bInSynchronousTick, bool bContribute
 			if (TickResult == ETickResult::KeepTicking)
 			{
 				// Poll discovery every so often to reduce super-linear costs
-				bPollDiscovery =
-					static_cast<float>(CurrentTimeSeconds - LastPollTimeSeconds) > PollDiscoveryPeriodSeconds;
-			}
-			else if (!bPollDiscovery)
-			{
-				bPollDiscovery = true;
+				float TimeSinceLastPoll = static_cast<float>(CurrentTimeSeconds - LastPollTimeSeconds);
+				bPollDiscovery = TimeSinceLastPoll > AssetDataGathererConstants::PollDiscoveryPeriodSeconds;
 			}
 			else
 			{
-				// We just polled discovery; go to sleep to wait for the discovery thread rather than busy-spinning and
-				// causing contention on its critical section
-				float SleepTimeSeconds = PollDiscoveryPeriodSeconds;
-				if (EndTimeSeconds > 0.)
+				// ETickResult::PollDiscovery
+				if (!bPollDiscovery)
 				{
-					SleepTimeSeconds = FMath::Min(SleepTimeSeconds,
-						static_cast<float>(EndTimeSeconds - CurrentTimeSeconds));
+					bPollDiscovery = true;
 				}
-				if (SleepTimeSeconds > 0.f)
+				else
 				{
-					if (TickStartTime >= 0.)
-					{
-						CurrentSearchTime += CurrentTimeSeconds - TickStartTime;
-					}
-					FPlatformProcess::Sleep(SleepTimeSeconds);
-					TickStartTime = FPlatformTime::Seconds();
+					// We just polled discovery, don't poll it again because we want to avoid busy-spinning and
+					// causing contention on the discovery critical section.
+					// Report back to the caller that they should wait for the discovery thread before they tick again.
+					Result = ETickResult::PollDiscovery;
+					break;
 				}
 			}
 		}
-		if (TickStartTime >= 0.)
+		if (TickStartTime >= 0.) // TickInternal might have updated CurrentSearchTime and cleared TickStartTime
 		{
 			CurrentSearchTime += FPlatformTime::Seconds() - TickStartTime;
 			TickStartTime = 0.;
@@ -3851,6 +3859,7 @@ void FAssetDataGatherer::InnerTickLoop(bool bInSynchronousTick, bool bContribute
 	{
 		SaveCacheFile(AssetsToSave);
 	}
+	return Result;
 }
 
 void FAssetDataGatherer::SaveCacheFile(const TArray<TPair<FName,FDiskCachedAssetData*>>& AssetsToSave)
diff --git a/Engine/Source/Runtime/AssetRegistry/Private/AssetDataGatherer.h b/Engine/Source/Runtime/AssetRegistry/Private/AssetDataGatherer.h
--- a/Engine/Source/Runtime/AssetRegistry/Private/AssetDataGatherer.h
+++ b/Engine/Source/Runtime/AssetRegistry/Private/AssetDataGatherer.h
@@ -268,7 +268,7 @@ private:
 	 * Helper function to run the tick in a loop-within-a-loop to minimize critical section entry, and to move expensive
 	 * operations out of the critical section
 	 */
-	void InnerTickLoop(bool bInSynchronousTick, bool bContributeToCacheSave, double EndTimeSeconds);
+	ETickResult InnerTickLoop(bool bInSynchronousTick, bool bContributeToCacheSave, double EndTimeSeconds);
 	/**
 	 * Tick function to pump scanning and push results into the search results structure. May be called from devoted
 	 * thread or inline from synchronous functions on other threads.
