AssetRegistry: Allow packages with invalid data in the AssetRegistryDependencyData to be loaded in the AssetRegistry (and therefore the editor). There is an existing bug in the FAssetHeaderPatcher which can cause the invalid data, and we need to be able to load packages which were saved with that bug present. Set the data to something valid and conservative (at the cost of inefficient cooking) and issue a warning to resave the package to fix it.
[RN] Bugfix, Core

diff --git a/Engine/Source/Runtime/AssetRegistry/Private/PackageReader.cpp b/Engine/Source/Runtime/AssetRegistry/Private/PackageReader.cpp
--- a/Engine/Source/Runtime/AssetRegistry/Private/PackageReader.cpp
+++ b/Engine/Source/Runtime/AssetRegistry/Private/PackageReader.cpp
@@ -572,6 +572,7 @@ bool FPackageReader::SerializeAssetRegistryDependencyData(TBitArray<>& OutImport
 	Args.NumImports = ImportMap.Num();
 	Args.NumSoftPackageReferences = SoftPackageReferenceList.Num();
 	Args.PackageVersion = PackageFileSummary.GetFileVersionUE();
+	Args.PackageName = FName(*GetLongPackageName());
 
 	ClearError();
 	Loader->ClearError();
@@ -1336,6 +1337,7 @@ bool FPackageReader::SerializeEditorOnlyFlags(TBitArray<>& OutImportUsedInGame,
 	Args.NumImports = ImportMap.Num();
 	Args.NumSoftPackageReferences = SoftPackageReferenceList.Num();
 	Args.PackageVersion = PackageFileSummary.GetFileVersionUE();
+	Args.PackageName = FName(*GetLongPackageName());
 
 	if (!UE::AssetRegistry::ReadPackageDataDependencies(Args))
 	{
@@ -1888,10 +1890,19 @@ namespace UE::AssetRegistry
 		{
 			Args.ExtraPackageDependencies.Empty();
 		}
-		if (Args.ImportUsedInGame.Num() != Args.NumImports ||
-			Args.SoftPackageUsedInGame.Num() != Args.NumSoftPackageReferences)
+		if (Args.ImportUsedInGame.Num() != Args.NumImports)
 		{
-			return false;
+			UE_LOG(LogAssetRegistry, Warning,
+				TEXT("Package is loadable but has invalid data; resave the package! Package: %s. Reason: AssetRegistryDependency data has ImportUsedInGame.Num() == %d, which does not match NumImports == %d."),
+				*Args.PackageName.ToString(), Args.ImportUsedInGame.Num(), Args.NumImports);
+			Args.ImportUsedInGame.Init(true /* bValue */, Args.NumImports);
+		}
+		if (Args.SoftPackageUsedInGame.Num() != Args.NumSoftPackageReferences)
+		{
+			UE_LOG(LogAssetRegistry, Warning,
+				TEXT("Package is loadable but has invalid data; resave the package! Package: %s. Reason: AssetRegistryDependency data has SoftPackageUsedInGame.Num() == %d, which does not match NumSoftPackageReferences == %d."),
+				*Args.PackageName.ToString(), Args.SoftPackageUsedInGame.Num(), Args.NumSoftPackageReferences);
+			Args.SoftPackageUsedInGame.Init(true /* bValue */, Args.NumSoftPackageReferences);
 		}
 		return !Ar.IsError();
 	}
diff --git a/Engine/Source/Runtime/AssetRegistry/Public/AssetRegistry/IAssetRegistry.h b/Engine/Source/Runtime/AssetRegistry/Public/AssetRegistry/IAssetRegistry.h
--- a/Engine/Source/Runtime/AssetRegistry/Public/AssetRegistry/IAssetRegistry.h
+++ b/Engine/Source/Runtime/AssetRegistry/Public/AssetRegistry/IAssetRegistry.h
@@ -1064,6 +1064,7 @@ namespace UE::AssetRegistry
 		int32 NumImports = -1;
 		int32 NumSoftPackageReferences = -1;
 		FPackageFileVersion PackageVersion;
+		FName PackageName;
 
 		// Outputs
 		TBitArray<> ImportUsedInGame;
