AssetRegistry: Cache HasMarkOfTheWeb in the gatherer cache to reduce time spent in discovery.
#rnx

diff --git a/Engine/Source/Runtime/AssetRegistry/Private/AssetDataGatherer.cpp b/Engine/Source/Runtime/AssetRegistry/Private/AssetDataGatherer.cpp
--- a/Engine/Source/Runtime/AssetRegistry/Private/AssetDataGatherer.cpp
+++ b/Engine/Source/Runtime/AssetRegistry/Private/AssetDataGatherer.cpp
@@ -48,7 +48,7 @@ namespace AssetDataGathererConstants
 	constexpr int32 MinSecondsToElapseBeforeCacheWrite = 30;
 	constexpr int32 MinAssetReadsBeforeCacheWrite = 100000;
 	constexpr int32 CacheShardAssetCount = 250000;
-	static constexpr uint32 CacheSerializationMagic = 0x2779E7E0; // Versioning and integrity checking
+	static constexpr uint32 CacheSerializationMagic = 0xA46B0FDA; // Versioning and integrity checking
 	static constexpr uint64 CurrentVersion = FAssetRegistryVersion::LatestVersion | (uint64(CacheSerializationMagic) << 32);
 	static const FTimespan CachePruneAge = FTimespan(7, 0, 0, 0);
 	/**
@@ -358,11 +358,6 @@ TArray<FString> FPreloadSettings::FindShardedCacheFiles() const
 
 FPreloadSettings GPreloadSettings;
 
-bool IsPackageBlocked(const FStringView& FilePath)
-{
-	return bBlockPackagesWithMarkOfTheWeb && IPlatformFile::GetPlatformPhysical().HasMarkOfTheWeb(*FString(FilePath));
-}
-
 /** A structure to hold serialized cache data from async loads before adding it to the Gatherer's main cache. */
 struct FCachePayload
 {
@@ -400,37 +395,38 @@ void AssignStringWithoutShrinking(FString& InOutResult, FStringView Value)
 	}
 }
 
-FDiscoveredPathData::FDiscoveredPathData(FStringView InLocalAbsPath, FStringView InLongPackageName, FStringView InRelPath, const FDateTime& InPackageTimestamp, EGatherableFileType InType, bool bInBlocked)
+FDiscoveredPathData::FDiscoveredPathData(FStringView InLocalAbsPath, FStringView InLongPackageName,
+	FStringView InRelPath, const FDateTime& InPackageTimestamp, EGatherableFileType InType)
 	: LocalAbsPath(InLocalAbsPath)
 	, LongPackageName(InLongPackageName)
 	, RelPath(InRelPath)
 	, PackageTimestamp(InPackageTimestamp)
 	, Type(InType)
-	, bBlocked(bInBlocked)
 {
 }
 
-FDiscoveredPathData::FDiscoveredPathData(FStringView InLocalAbsPath, FStringView InLongPackageName, FStringView InRelPath, EGatherableFileType InType, bool bInBlocked)
+FDiscoveredPathData::FDiscoveredPathData(FStringView InLocalAbsPath, FStringView InLongPackageName, 
+	FStringView InRelPath, EGatherableFileType InType)
 	: LocalAbsPath(InLocalAbsPath)
 	, LongPackageName(InLongPackageName)
 	, RelPath(InRelPath)
 	, Type(InType)
-	, bBlocked(bInBlocked)
 {
 }
 
-void FDiscoveredPathData::Assign(FStringView InLocalAbsPath, FStringView InLongPackageName, FStringView InRelPath, EGatherableFileType InType, bool bInBlocked)
+void FDiscoveredPathData::Assign(FStringView InLocalAbsPath, FStringView InLongPackageName, FStringView InRelPath,
+	EGatherableFileType InType)
 {
 	AssignStringWithoutShrinking(LocalAbsPath, InLocalAbsPath);
 	AssignStringWithoutShrinking(LongPackageName, InLongPackageName);
 	AssignStringWithoutShrinking(RelPath, InRelPath);
 	Type = InType;
-	bBlocked = bInBlocked;
 }
 
-void FDiscoveredPathData::Assign(FStringView InLocalAbsPath, FStringView InLongPackageName, FStringView InRelPath, const FDateTime& InPackageTimestamp, EGatherableFileType InType, bool bInBlocked)
+void FDiscoveredPathData::Assign(FStringView InLocalAbsPath, FStringView InLongPackageName, FStringView InRelPath,
+	const FDateTime& InPackageTimestamp, EGatherableFileType InType)
 {
-	Assign(InLocalAbsPath, InLongPackageName, InRelPath, InType, bInBlocked);
+	Assign(InLocalAbsPath, InLongPackageName, InRelPath, InType);
 	PackageTimestamp = InPackageTimestamp;
 }
 
@@ -439,17 +435,18 @@ SIZE_T FDiscoveredPathData::GetAllocatedSize() const
 	return LocalAbsPath.GetAllocatedSize() + LongPackageName.GetAllocatedSize() + RelPath.GetAllocatedSize();
 }
 
-FGatheredPathData::FGatheredPathData(FStringView InLocalAbsPath, FStringView InLongPackageName, const FDateTime& InPackageTimestamp, EGatherableFileType InType, bool bInBlocked)
+FGatheredPathData::FGatheredPathData(FStringView InLocalAbsPath, FStringView InLongPackageName,
+	const FDateTime& InPackageTimestamp, EGatherableFileType InType)
 	: LocalAbsPath(InLocalAbsPath)
 	, LongPackageName(InLongPackageName)
 	, PackageTimestamp(InPackageTimestamp)
 	, Type(InType)
-	, bBlocked(bInBlocked)
 {
 }
 
 FGatheredPathData::FGatheredPathData(const FDiscoveredPathData& DiscoveredData)
-	: FGatheredPathData(DiscoveredData.LocalAbsPath, DiscoveredData.LongPackageName, DiscoveredData.PackageTimestamp, DiscoveredData.Type, DiscoveredData.bBlocked)
+	: FGatheredPathData(DiscoveredData.LocalAbsPath, DiscoveredData.LongPackageName,
+		DiscoveredData.PackageTimestamp, DiscoveredData.Type)
 {
 }
 
@@ -458,22 +455,22 @@ FGatheredPathData::FGatheredPathData(FDiscoveredPathData&& DiscoveredData)
 	, LongPackageName(MoveTemp(DiscoveredData.LongPackageName))
 	, PackageTimestamp(MoveTemp(DiscoveredData.PackageTimestamp))
 	, Type(DiscoveredData.Type)
-	, bBlocked(DiscoveredData.bBlocked)
 {
 }
 
-void FGatheredPathData::Assign(FStringView InLocalAbsPath, FStringView InLongPackageName, const FDateTime& InPackageTimestamp, EGatherableFileType InType, bool bInBlocked)
+void FGatheredPathData::Assign(FStringView InLocalAbsPath, FStringView InLongPackageName,
+	const FDateTime& InPackageTimestamp, EGatherableFileType InType)
 {
 	AssignStringWithoutShrinking(LocalAbsPath, InLocalAbsPath);
 	AssignStringWithoutShrinking(LongPackageName, InLongPackageName);
 	PackageTimestamp = InPackageTimestamp;
 	Type = InType;
-	bBlocked = bInBlocked;
 }
 
 void FGatheredPathData::Assign(const FDiscoveredPathData& DiscoveredData)
 {
-	Assign(DiscoveredData.LocalAbsPath, DiscoveredData.LongPackageName, DiscoveredData.PackageTimestamp, DiscoveredData.Type, DiscoveredData.bBlocked);
+	Assign(DiscoveredData.LocalAbsPath, DiscoveredData.LongPackageName, DiscoveredData.PackageTimestamp,
+		DiscoveredData.Type);
 }
 
 SIZE_T FGatheredPathData::GetAllocatedSize() const
@@ -1972,7 +1969,7 @@ void FAssetDataDiscovery::TickInternal(bool bTickAll)
 						Data.IteratedSubDirs.Emplace();
 					}
 					Data.IteratedSubDirs[Data.NumIteratedDirs++].Assign(Data.DirLocalAbsPath, Data.DirLongPackageName, RelPath,
-						EGatherableFileType::Directory, false /* bBlocked */);
+						EGatherableFileType::Directory);
 				}
 			}
 			for (const FCachedDirScanFile& FileData : PathData->Files)
@@ -2002,9 +1999,8 @@ void FAssetDataDiscovery::TickInternal(bool bTickAll)
 						}
 						FPathViews::AppendPath(Data.DirLongPackageName, BaseName);
 						FPathViews::AppendPath(Data.DirLocalAbsPath, RelPath);
-						const bool bBlocked = FileType == EGatherableFileType::PackageFile && IsPackageBlocked(Data.DirLocalAbsPath);
 						Data.IteratedFiles[Data.NumIteratedFiles++].Assign(Data.DirLocalAbsPath, Data.DirLongPackageName, RelPath,
-							FileData.ModificationTime, FileType, bBlocked);
+							FileData.ModificationTime, FileType);
 					}
 				}
 			}
@@ -2105,7 +2101,7 @@ void FAssetDataDiscovery::TickInternal(bool bTickAll)
 							Data.IteratedSubDirs.Emplace();
 						}
 						Data.IteratedSubDirs[Data.NumIteratedDirs++].Assign(LocalAbsPath, Data.DirLongPackageName, RelPath,
-							EGatherableFileType::Directory, false /* bBlocked */);
+							EGatherableFileType::Directory);
 					}
 				}
 				else
@@ -2130,9 +2126,8 @@ void FAssetDataDiscovery::TickInternal(bool bTickAll)
 								Data.IteratedFiles.Emplace();
 							}
 							FPathViews::AppendPath(Data.DirLongPackageName, BaseName);
-							const bool bBlocked = FileType == EGatherableFileType::PackageFile && IsPackageBlocked(LocalAbsPath);
 							Data.IteratedFiles[Data.NumIteratedFiles++].Assign(LocalAbsPath, Data.DirLongPackageName, RelPath,
-								ModificationTime, FileType, bBlocked);
+								ModificationTime, FileType);
 						}
 					}
 				}
@@ -2528,8 +2523,7 @@ void FAssetDataDiscovery::SetPropertiesAndWait(TArrayView<FPathExistence> QueryP
 							LongPackageName << MountDir->GetLongPackageName();
 							FPathViews::AppendPath(LongPackageName, ScanDir->GetMountRelPath());
 							FPathViews::AppendPath(LongPackageName, FileRelPathNoExt);
-							const bool bBlocked = FileType == EGatherableFileType::PackageFile && IsPackageBlocked(SearchPath);
-							AddDiscoveredFile(FDiscoveredPathData(SearchPath, LongPackageName, RelPathFromParentDir, QueryPath.GetModificationTime(), FileType, bBlocked));
+							AddDiscoveredFile(FDiscoveredPathData(SearchPath, LongPackageName, RelPathFromParentDir, QueryPath.GetModificationTime(), FileType));
 							if (FPathViews::IsPathLeaf(RelPath) && !ScanDir->HasScanned())
 							{
 								SetIsIdle(false);
@@ -2994,8 +2988,7 @@ void FAssetDataDiscovery::OnFileCreated(const FString& LocalAbsPath)
 			LongPackageName << MountDir->GetLongPackageName();
 			FPathViews::AppendPath(LongPackageName, ScanDir->GetMountRelPath());
 			FPathViews::AppendPath(LongPackageName, FileRelPathNoExt);
-			const bool bBlocked = FileType == EGatherableFileType::PackageFile && IsPackageBlocked(LocalAbsPath);
-			AddDiscoveredFile(FDiscoveredPathData(LocalAbsPath, LongPackageName, RelPathFromParentDir, StatData.ModificationTime, FileType, bBlocked));
+			AddDiscoveredFile(FDiscoveredPathData(LocalAbsPath, LongPackageName, RelPathFromParentDir, StatData.ModificationTime, FileType));
 			if (FPathViews::IsPathLeaf(FileRelPath))
 			{
 				ScanDir->MarkFileAlreadyScanned(FileRelPath);
@@ -4152,6 +4145,7 @@ FAssetDataGatherer::ETickResult FAssetDataGatherer::TickInternal(double& TickSta
 		TArray<FAssetData*> AssetDataFromFile;
 		FPackageDependencyData DependencyData;
 		TArray<FString> CookedPackageNamesWithoutAssetData;
+		UE::AssetRegistry::Private::EOptionalBool HasMarkOfTheWeb = UE::AssetRegistry::Private::EOptionalBool::Unset;
 		bool bCanAttemptAssetRetry = false;
 		bool bResult = false;
 		bool bCanceled = false;
@@ -4184,13 +4178,6 @@ FAssetDataGatherer::ETickResult FAssetDataGatherer::TickInternal(double& TickSta
 			continue;
 		}
 
-		// If this is a blocked package, just directly add its file path to the Blocked results
-		if (AssetFileData.bBlocked)
-		{
-			LocalBlockedResults.Add(AssetFileData.LocalAbsPath);
-			continue;
-		}
-
 		const FName PackageName = *AssetFileData.LongPackageName;
 
 		FDiskCachedAssetData** DiskCachedAssetDataPtr = DiskCachedAssetDataMap.Find(PackageName);
@@ -4211,8 +4198,30 @@ FAssetDataGatherer::ETickResult FAssetDataGatherer::TickInternal(double& TickSta
 			}
 		}
 
+		// Check for whether the cache data has the information we need to know to report whether the package is blocked
+		if (bBlockPackagesWithMarkOfTheWeb && DiskCachedAssetData
+			&& DiskCachedAssetData->HasMarkOfTheWeb == UE::AssetRegistry::Private::EOptionalBool::Unset)
+		{
+			DiskCachedAssetData = nullptr;
+		}
+
 		if (DiskCachedAssetData)
 		{
+			// If this is a blocked package, just directly add its file path to the Blocked results
+			if (bBlockPackagesWithMarkOfTheWeb
+				&& DiskCachedAssetData->HasMarkOfTheWeb == UE::AssetRegistry::Private::EOptionalBool::True)
+			{
+				// To avoid falsely blocked packages that have had the mark removed, always recalculate HasMarkOfTheWeb
+				// before allowing it to block the package
+				DiskCachedAssetData->HasMarkOfTheWeb = BoolToOptionalBool(
+					IPlatformFile::GetPlatformPhysical().HasMarkOfTheWeb(AssetFileData.LocalAbsPath));
+				if (DiskCachedAssetData->HasMarkOfTheWeb == UE::AssetRegistry::Private::EOptionalBool::True)
+				{
+					LocalBlockedResults.Add(AssetFileData.LocalAbsPath);
+				}
+				continue;
+			}
+
 			// Add the valid cached data to our results, and to the map of data we keep to write out the new version of the cache file
 			++NewCachedAssetFiles;
 
@@ -4286,7 +4295,14 @@ FAssetDataGatherer::ETickResult FAssetDataGatherer::TickInternal(double& TickSta
 				return;
 			}
 			UE_SCOPED_IO_ACTIVITY(*WriteToString<512>(TEXT("Loading Asset"), ReadContext.PackageName.ToString()));
-			ReadContext.bResult = ReadAssetFile(ReadContext.AssetFileData.LongPackageName, ReadContext.AssetFileData.LocalAbsPath, ReadContext.AssetDataFromFile, ReadContext.DependencyData, ReadContext.CookedPackageNamesWithoutAssetData, ReadContext.bCanAttemptAssetRetry);
+			ReadContext.bResult = ReadAssetFile(ReadContext.AssetFileData.LongPackageName,
+				ReadContext.AssetFileData.LocalAbsPath, ReadContext.AssetDataFromFile, ReadContext.DependencyData,
+				ReadContext.CookedPackageNamesWithoutAssetData, ReadContext.bCanAttemptAssetRetry);
+			if (ReadContext.bResult && bBlockPackagesWithMarkOfTheWeb)
+			{
+				ReadContext.HasMarkOfTheWeb = BoolToOptionalBool(
+					IPlatformFile::GetPlatformPhysical().HasMarkOfTheWeb(ReadContext.AssetFileData.LocalAbsPath));
+			}
 		},
 		ParallelFlags
 	);
@@ -4301,13 +4317,9 @@ FAssetDataGatherer::ETickResult FAssetDataGatherer::TickInternal(double& TickSta
 		}
 		else if (ReadContext.bResult)
 		{
-			++NewUncachedAssetFiles;
-
-			// Add the results from a cooked package into our results on cooked package
-			LocalCookedPackageNamesWithoutAssetDataResults.Append(MoveTemp(ReadContext.CookedPackageNamesWithoutAssetData));
 			// Do not add the results from a cooked package into the map of data we keep to write out the new version of the cache file
 			bool bCachePackage = bLocalIsCacheWriteEnabled
-				&& LocalCookedPackageNamesWithoutAssetDataResults.Num() == 0
+				&& ReadContext.CookedPackageNamesWithoutAssetData.Num() == 0
 				&& ensure(ReadContext.AssetFileData.Type == EGatherableFileType::PackageFile);
 			if (bCachePackage)
 			{
@@ -4332,37 +4344,54 @@ FAssetDataGatherer::ETickResult FAssetDataGatherer::TickInternal(double& TickSta
 					NewData->AssetDataList.Add(*BackgroundAssetData);
 				}
 
-				// MoveTemp only used if we don't need DependencyData anymore
 				NewData->DependencyData = ReadContext.DependencyData;
+				NewData->HasMarkOfTheWeb = ReadContext.HasMarkOfTheWeb;
 
 				NewCachedAssetData.Add(NewData);
 				AddToCache(ReadContext.PackageName, NewData);
 			}
 
-			bool MustBeHandledByGameThread = false;
-			// 
-			// In the future, we may need to process certain assets on the game thread because, e.g.,
-			// we may need to support PostLoadAssetRegistryTags running on the game thread. 
-			// The infrastructure is provided here to handle that case. In order to do so,
-			// implement ClassRequiresGameThreadProcessing in UAssetRegistryImpl and call it as shown below
-			// The rest of the functions in the asset registry respect the separation of data into general and 
-			// ForGameThread containers. In particular, these are consumed in FAssetRegistryImpl::TickGatherer.
-			// 
-			//for (const FAssetData* BackgroundAssetData : ReadContext.AssetDataFromFile)
-			//{
-			//	if (AssetRegistry.ClassRequiresGameThreadProcessing(BackgroundAssetData->GetClass()))
-			//	{
-			//		MustBeHandledByGameThread = true;
-			//		break;
-			//	}
-			//}
-
-			TArray<TUniquePtr<FAssetData>, FBatchInlineAllocator>& TargetAssetResults = MustBeHandledByGameThread ? LocalAssetResultsForGameThread : LocalAssetResults;
-			TArray<FPackageDependencyData, FBatchInlineAllocator>& TargetDependencyResults = MustBeHandledByGameThread ? LocalDependencyResultsForGameThread : LocalDependencyResults;
-
-			// Add the results from the package into our output results
-			TargetAssetResults.Append(MoveTemp(ReadContext.AssetDataFromFile));
-			TargetDependencyResults.Add(MoveTemp(ReadContext.DependencyData));
+			bool bPackageBlocked = bBlockPackagesWithMarkOfTheWeb
+				&& ReadContext.HasMarkOfTheWeb == UE::AssetRegistry::Private::EOptionalBool::True;
+			if (bPackageBlocked)
+			{
+				// If this is a blocked package, add its file path to the Blocked results, and not the published result.
+				// It still will have been cached above, so we don't have to read its data next time.
+				LocalBlockedResults.Add(ReadContext.AssetFileData.LocalAbsPath);
+			}
+			else
+			{
+				++NewUncachedAssetFiles;
+				// Add the results from a cooked package into our results on cooked package
+				LocalCookedPackageNamesWithoutAssetDataResults.Append(MoveTemp(ReadContext.CookedPackageNamesWithoutAssetData));
+
+				bool MustBeHandledByGameThread = false;
+				// 
+				// In the future, we may need to process certain assets on the game thread because, e.g.,
+				// we may need to support PostLoadAssetRegistryTags running on the game thread. 
+				// The infrastructure is provided here to handle that case. In order to do so,
+				// implement ClassRequiresGameThreadProcessing in UAssetRegistryImpl and call it as shown below
+				// The rest of the functions in the asset registry respect the separation of data into general and 
+				// ForGameThread containers. In particular, these are consumed in FAssetRegistryImpl::TickGatherer.
+				// 
+				//for (const FAssetData* BackgroundAssetData : ReadContext.AssetDataFromFile)
+				//{
+				//	if (AssetRegistry.ClassRequiresGameThreadProcessing(BackgroundAssetData->GetClass()))
+				//	{
+				//		MustBeHandledByGameThread = true;
+				//		break;
+				//	}
+				//}
+
+				TArray<TUniquePtr<FAssetData>, FBatchInlineAllocator>& TargetAssetResults =
+					MustBeHandledByGameThread ? LocalAssetResultsForGameThread : LocalAssetResults;
+				TArray<FPackageDependencyData, FBatchInlineAllocator>& TargetDependencyResults =
+					MustBeHandledByGameThread ? LocalDependencyResultsForGameThread : LocalDependencyResults;
+
+				// Add the results from the package into our output results
+				TargetAssetResults.Append(MoveTemp(ReadContext.AssetDataFromFile));
+				TargetDependencyResults.Add(MoveTemp(ReadContext.DependencyData));
+			}
 		}
 		else if (ReadContext.bCanAttemptAssetRetry)
 		{
diff --git a/Engine/Source/Runtime/AssetRegistry/Private/AssetDataGathererPrivate.h b/Engine/Source/Runtime/AssetRegistry/Private/AssetDataGathererPrivate.h
--- a/Engine/Source/Runtime/AssetRegistry/Private/AssetDataGathererPrivate.h
+++ b/Engine/Source/Runtime/AssetRegistry/Private/AssetDataGathererPrivate.h
@@ -85,18 +85,16 @@ struct FDiscoveredPathData
 	FDateTime PackageTimestamp;
 	/** The type of file that was found */
 	EGatherableFileType Type = EGatherableFileType::Invalid;
-	/** Whether the file should just be reported as blocked and not further data gathered */
-	bool bBlocked = false;
 
 	FDiscoveredPathData() = default;
 	FDiscoveredPathData(FStringView InLocalAbsPath, FStringView InLongPackageName, FStringView InRelPath,
-		const FDateTime& InPackageTimestamp, EGatherableFileType InType, bool bInBlocked);
+		const FDateTime& InPackageTimestamp, EGatherableFileType InType);
 	FDiscoveredPathData(FStringView InLocalAbsPath, FStringView InLongPackageName, FStringView InRelPath,
-		EGatherableFileType InType, bool bInBlocked);
+		EGatherableFileType InType);
 	void Assign(FStringView InLocalAbsPath, FStringView InLongPackageName, FStringView InRelPath,
-		const FDateTime& InPackageTimestamp, EGatherableFileType InType, bool bInBlocked);
+		const FDateTime& InPackageTimestamp, EGatherableFileType InType);
 	void Assign(FStringView InLocalAbsPath, FStringView InLongPackageName, FStringView InRelPath,
-		EGatherableFileType InType, bool bInBlocked);
+		EGatherableFileType InType);
 
 	/** Return the total amount of heap memory used by the gatherer (including not-yet-claimed search results). */
 	SIZE_T GetAllocatedSize() const;
@@ -113,16 +111,14 @@ struct FGatheredPathData
 	FDateTime PackageTimestamp;
 	/** The type of file that was found */
 	EGatherableFileType Type = EGatherableFileType::Invalid;
-	/** Whether the file should just be reported as blocked and not further data gathered */
-	bool bBlocked = false;
 
 	FGatheredPathData() = default;
 	FGatheredPathData(FStringView InLocalAbsPath, FStringView InLongPackageName, const FDateTime& InPackageTimestamp,
-		EGatherableFileType InType, bool bInBlocked);
+		EGatherableFileType InType);
 	explicit FGatheredPathData(const FDiscoveredPathData& DiscoveredData);
 	explicit FGatheredPathData(FDiscoveredPathData&& DiscoveredData);
 	void Assign(FStringView InLocalAbsPath, FStringView InLongPackageName, const FDateTime& InPackageTimestamp,
-		EGatherableFileType InType, bool bInBlocked);
+		EGatherableFileType InType);
 	void Assign(const FDiscoveredPathData& DiscoveredData);
 
 	/**
diff --git a/Engine/Source/Runtime/AssetRegistry/Private/DiskCachedAssetData.h b/Engine/Source/Runtime/AssetRegistry/Private/DiskCachedAssetData.h
--- a/Engine/Source/Runtime/AssetRegistry/Private/DiskCachedAssetData.h
+++ b/Engine/Source/Runtime/AssetRegistry/Private/DiskCachedAssetData.h
@@ -5,12 +5,28 @@
 #include "AssetRegistry/AssetData.h"
 #include "PackageDependencyData.h"
 
+namespace UE::AssetRegistry::Private
+{
+
+enum class EOptionalBool : uint8
+{
+	Unset,
+	False,
+	True,
+	Max = EOptionalBool::True,
+};
+
+}
+FArchive& operator<<(FArchive& Ar, UE::AssetRegistry::Private::EOptionalBool& Value);
+UE::AssetRegistry::Private::EOptionalBool BoolToOptionalBool(bool bValue);
+
 class FDiskCachedAssetData
 {
 public:
 	FDateTime ModificationTime; // Modification time of the asset last time it was discovered
 	FDateTime LastGatheredTime;  // The last time this asset was discovered during a scan 
 	FName Extension;
+	UE::AssetRegistry::Private::EOptionalBool HasMarkOfTheWeb = UE::AssetRegistry::Private::EOptionalBool::Unset;
 	TArray<FAssetData> AssetDataList;
 	FPackageDependencyData DependencyData;
 
@@ -33,6 +49,7 @@ public:
 		Ar << ModificationTime;
 		Ar << LastGatheredTime;
 		Ar << Extension;
+		Ar << HasMarkOfTheWeb;
 	
 		int32 AssetDataCount = AssetDataList.Num();
 		Ar << AssetDataCount;
@@ -59,3 +76,21 @@ public:
 		return Result;
 	}
 };
+
+inline FArchive& operator<<(FArchive& Ar, UE::AssetRegistry::Private::EOptionalBool& Value)
+{
+	using namespace UE::AssetRegistry::Private;
+	uint8 ValueInt = static_cast<uint8>(Value);
+	Ar << ValueInt;
+	if (Ar.IsLoading())
+	{
+		Value = ValueInt <= (uint8)EOptionalBool::Max ? (EOptionalBool)ValueInt : EOptionalBool::Unset;
+	}
+	return Ar;
+}
+
+inline UE::AssetRegistry::Private::EOptionalBool BoolToOptionalBool(bool bValue)
+{
+	using namespace UE::AssetRegistry::Private;
+	return bValue ? EOptionalBool::True : EOptionalBool::False;
+}
