AssetRegistry: Suppress spurious warning "ScanPathsSynchronous: %s is not in a mounted path, will not scan" when calling GetRedirectedObjectPath.
#rnx

diff --git a/Engine/Source/Runtime/AssetRegistry/Private/AssetRegistry.cpp b/Engine/Source/Runtime/AssetRegistry/Private/AssetRegistry.cpp
--- a/Engine/Source/Runtime/AssetRegistry/Private/AssetRegistry.cpp
+++ b/Engine/Source/Runtime/AssetRegistry/Private/AssetRegistry.cpp
@@ -3492,7 +3492,8 @@ FSoftObjectPath FAssetRegistryImpl::GetRedirectedObjectPath(const FSoftObjectPat
 
 	if (!AssetData && bNeedsScanning)
 	{
-		UE::AssetRegistry::Impl::FScanPathContext Context(*EventContext, *InheritanceContext, {}, { RedirectedPath.ToString() });
+		UE::AssetRegistry::Impl::FScanPathContext Context(*EventContext, *InheritanceContext, {}, { RedirectedPath.ToString() },
+			UE::AssetRegistry::EScanFlags::IgnoreInvalidPathWarning);
 		ScanPathsSynchronous(Context);
 
 		AssetData = RetrieveAssetData();
