AssetRegistry: Add output variable to ReadPackageDataDependencies for the size of the dependency section, for the benefit of commandlets that are patching the packagefile and need to know how big each header section is.
#rnx

diff --git a/Engine/Source/Runtime/AssetRegistry/Private/PackageReader.cpp b/Engine/Source/Runtime/AssetRegistry/Private/PackageReader.cpp
--- a/Engine/Source/Runtime/AssetRegistry/Private/PackageReader.cpp
+++ b/Engine/Source/Runtime/AssetRegistry/Private/PackageReader.cpp
@@ -1866,6 +1866,7 @@ namespace UE::AssetRegistry
 			Args.ImportUsedInGame.Init(true, Args.NumImports);
 			Args.SoftPackageUsedInGame.Init(true, Args.NumSoftPackageReferences);
 			Args.ExtraPackageDependencies.Empty();
+			Args.AssetRegistryDependencyDataSize = 0;
 			return true;
 		}
 
@@ -1895,6 +1896,8 @@ namespace UE::AssetRegistry
 		{
 			return false;
 		}
+
+		Args.AssetRegistryDependencyDataSize = Ar.Tell() - Args.AssetRegistryDependencyDataOffset;
 		return !Ar.IsError();
 	}
 }
diff --git a/Engine/Source/Runtime/AssetRegistry/Public/AssetRegistry/IAssetRegistry.h b/Engine/Source/Runtime/AssetRegistry/Public/AssetRegistry/IAssetRegistry.h
--- a/Engine/Source/Runtime/AssetRegistry/Public/AssetRegistry/IAssetRegistry.h
+++ b/Engine/Source/Runtime/AssetRegistry/Public/AssetRegistry/IAssetRegistry.h
@@ -1069,6 +1069,7 @@ namespace UE::AssetRegistry
 		TBitArray<> ImportUsedInGame;
 		TBitArray<> SoftPackageUsedInGame;
 		TArray<TPair<FName, EExtraDependencyFlags>> ExtraPackageDependencies;
+		int64 AssetRegistryDependencyDataSize = 0;
 	};
 	ASSETREGISTRY_API bool ReadPackageDataDependencies(FReadPackageDataDependenciesArgs& Args);
 	UE_DEPRECATED(5.5, "Use version that takes FReadPackageDataArgs")
