Avoid adding reparse points to the the Asset directory cache
We are doing this by keeping track when we iterate over directories if they are sym/junctions
If so avoid adding to the cache

This does not add support for it but allows it to work in a symlink setup

diff --git a/Engine/Source/Runtime/AssetRegistry/Private/AssetDataGatherer.cpp b/Engine/Source/Runtime/AssetRegistry/Private/AssetDataGatherer.cpp
--- a/Engine/Source/Runtime/AssetRegistry/Private/AssetDataGatherer.cpp
+++ b/Engine/Source/Runtime/AssetRegistry/Private/AssetDataGatherer.cpp
@@ -2050,7 +2050,7 @@ void FAssetDataDiscovery::TickInternal(bool bTickAll)
 
 			FCachedDirScanDir CacheDataToAdd;
 			auto ProcessIterData = [this, &Data, &CacheDataToAdd]
-			(const TCHAR* IterFilename, bool bIsDirectory, const FDateTime& ModificationTime, FFileJournalFileHandle JournalHandle)
+			(const TCHAR* IterFilename, bool bIsDirectory, const FDateTime& ModificationTime, FFileJournalFileHandle JournalHandle, bool bIsReparsePoint)
 			{
 				FStringView LocalAbsPath(IterFilename);
 				FStringView RelPath;
@@ -2087,7 +2087,7 @@ void FAssetDataDiscovery::TickInternal(bool bTickAll)
 					if (Cache.IsWriteEnabled() != EFeatureEnabled::Never)
 					{
 						CacheDataToAdd.SubDirRelPaths.Add(FString(RelPath));
-						Cache.QueueAdd(FString(LocalAbsPath), JournalHandle);
+						Cache.QueueAdd(FString(LocalAbsPath), JournalHandle, bIsReparsePoint);
 					}
 
 					FPathViews::AppendPath(Data.DirLongPackageName, RelPath);
@@ -2140,7 +2140,7 @@ void FAssetDataDiscovery::TickInternal(bool bTickAll)
 					[this, &Data, &CacheDataToAdd, &ProcessIterData]
 					(const TCHAR* IterFilename, const FFileJournalData& IterData)
 					{
-						return ProcessIterData(IterFilename, IterData.bIsDirectory, IterData.ModificationTime, IterData.JournalHandle);
+						return ProcessIterData(IterFilename, IterData.bIsDirectory, IterData.ModificationTime, IterData.JournalHandle, IterData.bIsReparsePoint);
 					});
 
 				Cache.QueueAdd(FString(Data.DirLocalAbsPath), MoveTemp(CacheDataToAdd));
@@ -2151,7 +2151,7 @@ void FAssetDataDiscovery::TickInternal(bool bTickAll)
 					[this, &Data, &ProcessIterData]
 					(const TCHAR* IterFilename, const FFileStatData& IterData)
 					{
-						return ProcessIterData(IterFilename, IterData.bIsDirectory, IterData.ModificationTime, FFileJournalFileHandle());
+						return ProcessIterData(IterFilename, IterData.bIsDirectory, IterData.ModificationTime, FFileJournalFileHandle(), false /* bIsReparsePoint */);
 					});
 			}
 
diff --git a/Engine/Source/Runtime/AssetRegistry/Private/AssetDataGathererDiscoveryCache.cpp b/Engine/Source/Runtime/AssetRegistry/Private/AssetDataGathererDiscoveryCache.cpp
--- a/Engine/Source/Runtime/AssetRegistry/Private/AssetDataGathererDiscoveryCache.cpp
+++ b/Engine/Source/Runtime/AssetRegistry/Private/AssetDataGathererDiscoveryCache.cpp
@@ -299,8 +299,8 @@ void FAssetDataDiscoveryCache::Shutdown()
 	}
 	while (!ScanQueueDirHandles.IsEmpty())
 	{
-		TPair<FString, FFileJournalFileHandle> Pair;
-		ScanQueueDirHandles.Dequeue(Pair);
+		FDiscoveredSubDirData SubDirData;
+		ScanQueueDirHandles.Dequeue(SubDirData);
 	}
 }
 
@@ -367,23 +367,57 @@ FCachedDirScanDir* FAssetDataDiscoveryCache::FindDir(FStringView Path, FCachedVo
 	return Volume->FindDir(Path);
 }
 
+bool FAssetDataDiscoveryCache::FindOrAddIsReparsePoint(FStringView DirName, FCachedVolumeInfo& Volume)
+{
+	if (DirName.IsEmpty())
+	{
+		return false;
+	}
+
+	FCachedDirScanDir& Existing = Volume.FindOrAddDir(DirName);
+
+	if (!Existing.bIsInsideReparsePoint.IsSet())
+	{
+		// Usually bIsReparsePoint is already set because we calculated it when we encountered this directory in the directory scan of its parent directory.
+		// But for top-level directories or directories that were found to be modified by the journal, we did not already encounter it.
+		// Calculate it now manually if it is not already set.
+		if (!Existing.bIsReparsePoint.IsSet())
+		{
+			IPlatformFile& PlatformFile = FPlatformFileManager::Get().GetPlatformFile();
+			Existing.bIsReparsePoint = PlatformFile.IsSymlink(DirName.GetData()) == ESymlinkResult::Symlink;
+		}
+
+		// Base case for when we have found a reparse point, end the recursion and set true for all our previous calls
+		if (Existing.bIsReparsePoint.GetValue())
+		{
+			Existing.bIsInsideReparsePoint = true;
+			return true;
+		}
+
+		Existing.bIsInsideReparsePoint = FindOrAddIsReparsePoint(FPaths::GetPath(DirName.GetData()), Volume);
+	}
+
+	return Existing.bIsInsideReparsePoint.GetValue();
+}
+
 void FAssetDataDiscoveryCache::QueueConsume()
 {
 	if (WriteEnabled == EFeatureEnabled::Never)
 	{
 		return;
 	}
-	TPair<FString, FFileJournalFileHandle> HandlePair;
-	while (ScanQueueDirHandles.Dequeue(HandlePair))
+	FDiscoveredSubDirData SubDirData;
+	while (ScanQueueDirHandles.Dequeue(SubDirData))
 	{
-		FString& DirName = HandlePair.Key;
+		FString& DirName = SubDirData.DirNameAbsPath;
 		FCachedVolumeInfo& Volume = FindOrAddVolume(DirName);
 		if (!Volume.bJournalAvailable && WriteEnabled != EFeatureEnabled::Always)
 		{
 			continue;
 		}
 		FCachedDirScanDir& Existing = Volume.FindOrAddDir(DirName);
-		Existing.JournalHandle = MoveTemp(HandlePair.Value);
+		Existing.JournalHandle = MoveTemp(SubDirData.JournalHandle);
+		Existing.bIsReparsePoint = SubDirData.bIsReparsePoint;
 	}
 
 	TPair<FString, FCachedDirScanDir> ScanDirPair;
@@ -392,6 +426,10 @@ void FAssetDataDiscoveryCache::QueueConsume()
 		FString& DirName = ScanDirPair.Key;
 		FCachedDirScanDir& ScanDir = ScanDirPair.Value;
 		FCachedVolumeInfo& Volume = FindOrAddVolume(DirName);
+		if (FindOrAddIsReparsePoint(DirName, Volume))
+		{
+			continue;
+		}
 		if (!Volume.bJournalAvailable && WriteEnabled != EFeatureEnabled::Always)
 		{
 			continue;
@@ -428,6 +466,10 @@ void FAssetDataDiscoveryCache::QueueConsume()
 			}
 		}
 
+		// These values get set in FindOrAddIsReparsePoint, make sure they dont get lost here
+		ScanDir.bIsReparsePoint = Existing.bIsReparsePoint;
+		ScanDir.bIsInsideReparsePoint = Existing.bIsInsideReparsePoint;
+
 		Existing = MoveTemp(ScanDir);
 		Existing.bCacheValid = true;
 	}
@@ -442,13 +484,13 @@ void FAssetDataDiscoveryCache::QueueAdd(FString DirName, FCachedDirScanDir DirDa
 	ScanQueueDirFullDatas.Enqueue(TPair<FString, FCachedDirScanDir>(MoveTemp(DirName), MoveTemp(DirData)));
 }
 
-void FAssetDataDiscoveryCache::QueueAdd(FString DirName, FFileJournalFileHandle JournalHandle)
+void FAssetDataDiscoveryCache::QueueAdd(FString DirName, FFileJournalFileHandle JournalHandle, bool bIsReparsePoint)
 {
 	if (WriteEnabled == EFeatureEnabled::Never)
 	{
 		return;
 	}
-	ScanQueueDirHandles.Enqueue(TPair<FString, FFileJournalFileHandle>(MoveTemp(DirName), MoveTemp(JournalHandle)));
+	ScanQueueDirHandles.Enqueue(FDiscoveredSubDirData(MoveTemp(DirName), MoveTemp(JournalHandle), bIsReparsePoint));
 }
 
 void FCachedVolumeInfo::ConditionalConstruct(const FString& InVolumeName)
@@ -550,4 +592,4 @@ FArchive& operator<<(FArchive& Ar, FCachedDirScanFile& Data)
 }
 
 
-} // namespace UE::AssetDataGather::Private
\ No newline at end of file
+} // namespace UE::AssetDataGather::Private
diff --git a/Engine/Source/Runtime/AssetRegistry/Private/AssetDataGathererDiscoveryCache.h b/Engine/Source/Runtime/AssetRegistry/Private/AssetDataGathererDiscoveryCache.h
--- a/Engine/Source/Runtime/AssetRegistry/Private/AssetDataGathererDiscoveryCache.h
+++ b/Engine/Source/Runtime/AssetRegistry/Private/AssetDataGathererDiscoveryCache.h
@@ -63,9 +63,28 @@ struct FCachedDirScanDir
 	TArray<FCachedDirScanFile> Files;
 	FFileJournalFileHandle JournalHandle = FileJournalFileHandleInvalid;
 	bool bCacheValid = false;
+
+	/** bIsReparse True when this directory is a reparse point (symlink/junction/etc)
+	  *			   It is optionally set to true or false when found during a directory scan
+	  *			   If it is never found or set we will have to do an explicit IsSymLink on this directory
+	*/
+	TOptional<bool> bIsReparsePoint;
+
+	/** bIsInsideReparsePoint True when this directory, or one of its parent directories are a ReparsePoint
+	  *						  It is optional set when checking if this directory contains any ReparsePoints, Parent directories included
+	*/
+	TOptional<bool> bIsInsideReparsePoint;
 };
 FArchive& operator<<(FArchive& Ar, FCachedDirScanDir& Data);
 
+/** Stores data from QueueAdd for subdirectories discovered during the directory scan, to be read during QueueConsume. */
+struct FDiscoveredSubDirData
+{
+	FString DirNameAbsPath;
+	FFileJournalFileHandle JournalHandle;
+	TOptional<bool> bIsReparsePoint;
+};
+
 /** Data about a volume used by FAssetDataDiscoveryCache to avoid the IO cost of rescanning directories on the volume. */
 class FCachedVolumeInfo
 {
@@ -129,20 +148,22 @@ public:
 
 	void QueueConsume();
 	void QueueAdd(FString DirName, FCachedDirScanDir DirData);
-	void QueueAdd(FString DirName, FFileJournalFileHandle JournalHandle);
+	void QueueAdd(FString DirName, FFileJournalFileHandle JournalHandle, bool bIsReparsePoint);
 
 private:
 	bool TryReadCacheFile();
 	void SerializeReadCacheFile(FArchive& Ar);
 	void SerializeWriteCacheFile(FArchive& Ar);
 
+	bool FindOrAddIsReparsePoint(FStringView DirName, FCachedVolumeInfo& CachedVolume);
+
 private:
 	/** Cached information about each volume on disk that has mounted directories. */
 	TSortedMap<FString, FCachedVolumeInfo> CachedVolumes;
 	/** Childpath information about directories collected from ParallelFor during scanning. */
 	TMpscQueue<TPair<FString, FCachedDirScanDir>> ScanQueueDirFullDatas;
 	/** JournalHandle information about directories collected during ParallelFor during scanning. */
-	TMpscQueue<TPair<FString, FFileJournalFileHandle>> ScanQueueDirHandles;
+	TMpscQueue<FDiscoveredSubDirData> ScanQueueDirHandles;
 
 	/** Used to implement initialization on demand. */
 	bool bInitialized = false;
diff --git a/Engine/Source/Runtime/Core/Private/Windows/WindowsPlatformFileJournal.cpp b/Engine/Source/Runtime/Core/Private/Windows/WindowsPlatformFileJournal.cpp
--- a/Engine/Source/Runtime/Core/Private/Windows/WindowsPlatformFileJournal.cpp
+++ b/Engine/Source/Runtime/Core/Private/Windows/WindowsPlatformFileJournal.cpp
@@ -254,6 +254,7 @@ bool FWindowsPlatformFile::FileJournalIterateDirectory(const TCHAR* Directory,
 				FileData.ModificationTime = WindowsFileTimeToUEDateTime(LastWriteTime);
 				FileData.bIsValid = true;
 				FileData.bIsDirectory = !!(Info->FileAttributes & FILE_ATTRIBUTE_DIRECTORY);
+				FileData.bIsReparsePoint = !!(Info->FileAttributes & FILE_ATTRIBUTE_REPARSE_POINT);
 				FileData.JournalHandle = ToFileHandle(Info->FileId);
 
 				if (!Visitor(*LocalAbsPath, FileData))
@@ -345,6 +346,7 @@ FFileJournalData FWindowsPlatformFile::FileJournalGetFileData(const TCHAR* Filen
 	}
 
 	OutResult.bIsDirectory = !!(Info->FileAttributes & FILE_ATTRIBUTE_DIRECTORY);
+	OutResult.bIsReparsePoint = !!(Info->FileAttributes & FILE_ATTRIBUTE_REPARSE_POINT);
 	FILETIME LastWriteTime;
 	LastWriteTime.dwLowDateTime = Info->LastWriteTime.LowPart;
 	LastWriteTime.dwHighDateTime = Info->LastWriteTime.HighPart;
diff --git a/Engine/Source/Runtime/Core/Public/GenericPlatform/GenericPlatformFile.h b/Engine/Source/Runtime/Core/Public/GenericPlatform/GenericPlatformFile.h
--- a/Engine/Source/Runtime/Core/Public/GenericPlatform/GenericPlatformFile.h
+++ b/Engine/Source/Runtime/Core/Public/GenericPlatform/GenericPlatformFile.h
@@ -272,6 +272,7 @@ struct FFileJournalData
 	FFileJournalFileHandle JournalHandle = FileJournalFileHandleInvalid;
 	bool bIsValid : 1;
 	bool bIsDirectory : 1;
+	bool bIsReparsePoint : 1;
 };
 
 /** Stores custom error messages from the engine along with an optional system error code that can provider more detailed information */
@@ -1017,6 +1018,7 @@ inline bool FFileJournalFileHandle::operator!=(const FFileJournalFileHandle& Oth
 inline FFileJournalData::FFileJournalData()
 	: bIsValid(false)
 	, bIsDirectory(false)
+	, bIsReparsePoint(false)
 {
 }
 
