Fix compiler error in AssetDataGatherer

diff --git a/Engine/Source/Runtime/AssetRegistry/Private/AssetDataGathererDiscoveryCache.cpp b/Engine/Source/Runtime/AssetRegistry/Private/AssetDataGathererDiscoveryCache.cpp
--- a/Engine/Source/Runtime/AssetRegistry/Private/AssetDataGathererDiscoveryCache.cpp
+++ b/Engine/Source/Runtime/AssetRegistry/Private/AssetDataGathererDiscoveryCache.cpp
@@ -490,7 +490,7 @@ void FAssetDataDiscoveryCache::QueueAdd(FString DirName, FFileJournalFileHandle
 	{
 		return;
 	}
-	ScanQueueDirHandles.Enqueue(FDiscoveredSubDirData(MoveTemp(DirName), MoveTemp(JournalHandle), bIsReparsePoint));
+	ScanQueueDirHandles.Enqueue(FDiscoveredSubDirData{MoveTemp(DirName), MoveTemp(JournalHandle), bIsReparsePoint});
 }
 
 void FCachedVolumeInfo::ConditionalConstruct(const FString& InVolumeName)
