Fixed default value for new ASSET_REGISTRY_LOAD_NAMES_IN_ORDER define (should be disabled by default to keep existing behavior)

diff --git a/Engine/Source/Runtime/AssetRegistry/Private/AssetRegistryArchive.cpp b/Engine/Source/Runtime/AssetRegistry/Private/AssetRegistryArchive.cpp
--- a/Engine/Source/Runtime/AssetRegistry/Private/AssetRegistryArchive.cpp
+++ b/Engine/Source/Runtime/AssetRegistry/Private/AssetRegistryArchive.cpp
@@ -7,7 +7,7 @@
 #include "AssetRegistry/AssetRegistryState.h"
 
 #ifndef ASSET_REGISTRY_LOAD_NAMES_IN_ORDER
-#define ASSET_REGISTRY_LOAD_NAMES_IN_ORDER 1 //0
+#define ASSET_REGISTRY_LOAD_NAMES_IN_ORDER 0
 #endif
 
 constexpr uint32 AssetRegistryNumberedNameBit = 0x80000000;
