AssetRegistry: Fix poor performance in LoadCalculatedDependencies when maps have many actors.
#rnx

diff --git a/Engine/Source/Runtime/AssetRegistry/Private/AssetRegistry.cpp b/Engine/Source/Runtime/AssetRegistry/Private/AssetRegistry.cpp
--- a/Engine/Source/Runtime/AssetRegistry/Private/AssetRegistry.cpp
+++ b/Engine/Source/Runtime/AssetRegistry/Private/AssetRegistry.cpp
@@ -5655,30 +5655,29 @@ void FAssetRegistryImpl::LoadCalculatedDependencies(FName PackageName,
 	FReadScopeLock GathererClassScopeLock(RegisteredDependencyGathererClassesLock);
 
 	TArray<UE::AssetDependencyGatherer::Private::FRegisteredAssetDependencyGatherer*, TInlineAllocator<2>> Gatherers;
+	TArray<FString> DirectoryReferences;
 	State.EnumerateAssetsByPackageName(PackageName,
-		[this, &Gatherers, &GetCompiledFilter, &bOutHadActivity, PackageName](const FAssetData* AssetData)
+		[this, &Gatherers, &GetCompiledFilter, &bOutHadActivity, PackageName, &DirectoryReferences](const FAssetData* AssetData)
 	{
 		Gatherers.Reset();
 
 		// Check the class name instead of trying to load the actual class as that is slow
 		// This code could be moved somewhere where it doesn't need to re-query the asset data, but it needs to happen after both dependencies and data are handled
 		RegisteredDependencyGathererClasses.MultiFind(AssetData->AssetClassPath, Gatherers);
+
 		for (UE::AssetDependencyGatherer::Private::FRegisteredAssetDependencyGatherer* Gatherer : Gatherers)
 		{
-			if (!bOutHadActivity)
-			{
-				RemoveDirectoryReferencer(PackageName);
-			}
-
 			bOutHadActivity = true;
 			TArray<IAssetDependencyGatherer::FGathereredDependency> GatheredDependencies;
-			TArray<FString> DirectoryReferences;
 
-			Gatherer->GatherDependencies(*AssetData, State, GetCompiledFilter, GatheredDependencies, DirectoryReferences);
+			TArray<FString> GathererDirectoryReferences;
+			Gatherer->GatherDependencies(*AssetData, State, GetCompiledFilter, GatheredDependencies, GathererDirectoryReferences);
 			
 			if (GatheredDependencies.Num())
 			{
 				FDependsNode* SourceNode = State.CreateOrFindDependsNode(FAssetIdentifier(PackageName));
+				bool bWasSourceDependencyListSorted = SourceNode->IsDependencyListSorted(EDependencyCategory::Package);
+				SourceNode->SetIsDependencyListSorted(EDependencyCategory::Package, false);
 
 				for (const IAssetDependencyGatherer::FGathereredDependency& GatheredDepencency : GatheredDependencies)
 				{
@@ -5687,45 +5686,70 @@ void FAssetRegistryImpl::LoadCalculatedDependencies(FName PackageName,
 					SourceNode->AddDependency(TargetNode, EDependencyCategory::Package, DependencyProperties);
 					TargetNode->AddReferencer(SourceNode);
 				}
-			}
 
-			for (const FString& Directory : DirectoryReferences)
-			{
-				AddDirectoryReferencer(PackageName, Directory);
+				SourceNode->SetIsDependencyListSorted(EDependencyCategory::Package, bWasSourceDependencyListSorted);
 			}
+
+			DirectoryReferences.Append(MoveTemp(GathererDirectoryReferences));
 		}
 
 		return true; // Keep iterating the assets in the package
 	});
+	SetDirectoriesWatchedByPackage(PackageName, MoveTemp(DirectoryReferences));
 }
 
-void FAssetRegistryImpl::AddDirectoryReferencer(FName PackageName, const FString& DirectoryLocalPathOrLongPackageName)
+void FAssetRegistryImpl::SetDirectoriesWatchedByPackage(FName PackageName, TConstArrayView<FString> Directories)
 {
-	FString DirectoryLocalPath;
-	if (!FPackageName::TryConvertLongPackageNameToFilename(DirectoryLocalPathOrLongPackageName, DirectoryLocalPath))
+	TArray<FString> DirectoryLocalPaths;
+	DirectoryLocalPaths.Reserve(Directories.Num());
+	for (const FString& DirectoryLocalPathOrLongPackageName : Directories)
 	{
-		UE_LOG(LogAssetRegistry, Warning,
-			TEXT("AddDirectoryReferencer called with LongPackageName %s that cannot be mapped to a local path. Ignoring it."),
-			*DirectoryLocalPathOrLongPackageName);
-		return;
+		FString DirectoryLocalPath;
+		if (!FPackageName::TryConvertLongPackageNameToFilename(DirectoryLocalPathOrLongPackageName, DirectoryLocalPath))
+		{
+			UE_LOG(LogAssetRegistry, Warning,
+				TEXT("AddDirectoryReferencer called with LongPackageName %s that cannot be mapped to a local path. Ignoring it."),
+				*DirectoryLocalPathOrLongPackageName);
+		}
+		FPaths::MakeStandardFilename(DirectoryLocalPath);
+		DirectoryLocalPaths.Add(MoveTemp(DirectoryLocalPath));
 	}
-	FPaths::MakeStandardFilename(DirectoryLocalPath);
-	DirectoryReferencers.AddUnique(DirectoryLocalPath, PackageName);
-}
 
-void FAssetRegistryImpl::RemoveDirectoryReferencer(FName PackageName)
-{
-	TArray<FString> FoundKeys;
-	for (const TPair<FString, FName>& Pair : DirectoryReferencers)
+	uint32 PackageNameHash = GetTypeHash(PackageName);
+	TArray<FString>* WatchedDirectories = DirectoriesWatchedByPackage.FindByHash(PackageNameHash, PackageName);
+
+	// Remove all old values from the reverse map
+	if (WatchedDirectories)
+	{
+		for (const FString& WatchedDirectory : *WatchedDirectories)
+		{
+			TSet<FName>* DirectoryWatchers = PackagesWatchingDirectory.Find(WatchedDirectory);
+			if (DirectoryWatchers)
+			{
+				DirectoryWatchers->RemoveByHash(PackageNameHash, PackageName);
+			}
+		}
+	}
+
+	if (Directories.IsEmpty())
 	{
-		if (Pair.Value == PackageName)
+		if (WatchedDirectories)
 		{
-			FoundKeys.Add(Pair.Key);
+			DirectoriesWatchedByPackage.RemoveByHash(PackageNameHash, PackageName);
 		}
+		return;
+	}
+
+	if (!WatchedDirectories)
+	{
+		WatchedDirectories = &DirectoriesWatchedByPackage.FindOrAddByHash(PackageNameHash, PackageName);
 	}
-	for (FString& Key : FoundKeys)
+	*WatchedDirectories = MoveTemp(DirectoryLocalPaths);
+
+	for (const FString& WatchedDirectory : *WatchedDirectories)
 	{
-		DirectoryReferencers.Remove(Key, PackageName);
+		TSet<FName>& DirectoryWatchers = PackagesWatchingDirectory.FindOrAdd(WatchedDirectory);
+		DirectoryWatchers.AddByHash(PackageNameHash, PackageName);
 	}
 }
 
@@ -7681,16 +7705,28 @@ void FAssetRegistryImpl::OnDirectoryChanged(Impl::FEventContext& EventContext,
 		if (bIsValidPackageName)
 		{
 			// If a package changes in a referenced directory, modify the Assets that monitor that directory
-			FString DirectoryPath = FPaths::CreateStandardFilename(FPaths::GetPath(File));
-			// TODO: Change DirectoryReferencers to a FileNameMap so we can do this FindParentDirectory check more quickly
+			FString ParentDirectory = FPaths::CreateStandardFilename(FPaths::GetPath(File));
 			TArray<FName, TInlineAllocator<1>> WatcherPackageNames;
-			for (const TPair<FString, FName>& Pair : DirectoryReferencers)
+			while (!ParentDirectory.IsEmpty())
 			{
-				if (FPathViews::IsParentPathOf(Pair.Key, DirectoryPath))
+				TSet<FName>* PackagesWatchingThisDirectory = PackagesWatchingDirectory.Find(ParentDirectory);
+				if (PackagesWatchingThisDirectory)
+				{
+					for (FName PackageWatchingThisDirectory : *PackagesWatchingThisDirectory)
+					{
+						WatcherPackageNames.Add(PackageWatchingThisDirectory);
+					}
+				}
+				FString NewParentDirectory = FPaths::GetPath(ParentDirectory);
+				if (ParentDirectory == NewParentDirectory)
 				{
-					WatcherPackageNames.Add(Pair.Value);
+					break;
 				}
+				ParentDirectory = MoveTemp(NewParentDirectory);
 			}
+			Algo::Sort(WatcherPackageNames, FNameFastLess());
+			WatcherPackageNames.SetNum(Algo::Unique(WatcherPackageNames));
+
 			for (FName WatcherPackageName : WatcherPackageNames)
 			{
 				// ScanModifiedAssetFiles accepts LongPackageNames as well as LocalPaths
diff --git a/Engine/Source/Runtime/AssetRegistry/Private/AssetRegistryImpl.h b/Engine/Source/Runtime/AssetRegistry/Private/AssetRegistryImpl.h
--- a/Engine/Source/Runtime/AssetRegistry/Private/AssetRegistryImpl.h
+++ b/Engine/Source/Runtime/AssetRegistry/Private/AssetRegistryImpl.h
@@ -8,6 +8,7 @@
 #include "AssetRegistry/IAssetRegistry.h"
 #include "AssetRegistry/PathTree.h"
 #include "Containers/RingBuffer.h"
+#include "Containers/Map.h"
 #include "Containers/Set.h"
 #include "HAL/CriticalSection.h"
 #include "Math/NumericLimits.h"
@@ -356,8 +357,8 @@ public:
 		bool& bOutHadActivity);
 	/** Add a watch on a directory that to modify data for a package whenever packages in the directory are modified. */
 	void AddDirectoryReferencer(FName PackageName, const FString& DirectoryLocalPathOrLongPackageName);
-	/** Remove all directory watches for PackageName. */
-	void RemoveDirectoryReferencer(FName PackageName);
+	/** Set the directories watched by PackageName. Removes previous list, call with empty array to clear. */
+	void SetDirectoriesWatchedByPackage(FName PackageName, TConstArrayView<FString> Directories);
 
 	/** Called when new gatherer is registered. Requires subsequent call to RebuildAssetDependencyGathererMapIfNeeded */
 	void OnAssetDependencyGathererRegistered() { bRegisteredDependencyGathererClassesDirty = true; }
@@ -685,8 +686,13 @@ private:
 	/** The set of object paths that have had their disk cache updated from the in memory version */
 	TSet<FSoftObjectPath> AssetDataObjectPathsUpdatedOnLoad;
 
-	/** A map from directoryname to packagename of Packages that have CalculatedDependencies on packages in the directory. */
-	TMultiMap<FString, FName> DirectoryReferencers;
+	/**
+	 * A bidirectional map of which directories are watched by which packages for CalculatedDependencies. When a
+	 * directory updates we need to rerun CalculatedDependencies for every package watching the directory. When a
+	 * package runs CalculatedDependencies, we need to clear all of its watched files.
+	 */
+	TMap<FString, TSet<FName>> PackagesWatchingDirectory;
+	TMap<FName, TArray<FString>> DirectoriesWatchedByPackage;
 
 	/**
 	 * Number of times during the startup scan that TickOnBackgroundThread was interrupted by a request for the
