AssetRegistry: When possible, fix OnDirectoryChanged and ScanPathsSynchronous to call LoadCalculatedDependencies on the changed or scanned asset before returning. This is only possible if we are not gathering other paths outside of the changed/scanned paths, because LoadCalculatedDependencies can only run when all gathering is done, and we do not wait on those other paths to complete.

This partially fixes a case relied upon by some users of the AssetRegistry: they expect that after a package save, waiting on the directory watcher and ticking the AssetRegistry once will finish the update for the saved files, and that update includes LoadCalculatedDependencies on the packages that watch the directory of the saved package, such as a map package watching its external actor directory.
#rnx

diff --git a/Engine/Source/Runtime/AssetRegistry/Private/AssetDataGatherer.cpp b/Engine/Source/Runtime/AssetRegistry/Private/AssetDataGatherer.cpp
--- a/Engine/Source/Runtime/AssetRegistry/Private/AssetDataGatherer.cpp
+++ b/Engine/Source/Runtime/AssetRegistry/Private/AssetDataGatherer.cpp
@@ -4079,7 +4079,7 @@ FAssetDataGatherer::ETickResult FAssetDataGatherer::TickInternal(double& TickSta
 		{
 			if (WaitBatchCount != -1)
 			{
-				WaitBatchCount = -1; // WaitBatchCount was set higher than FilesToSearch->GetNumAvailable(), mark it completed
+				WaitBatchCount = -1; // WaitBatchCount was set equal to or higher than FilesToSearch->GetNumAvailable(), mark it completed
 				TickResult = ETickResult::Interrupt;
 			}
 
@@ -4109,8 +4109,8 @@ FAssetDataGatherer::ETickResult FAssetDataGatherer::TickInternal(double& TickSta
 		{
 			if (WaitBatchCount == 0)
 			{
-				// We've finished executing the caller's requested batchcount (and we have restored idle if we
-				// are idle), so exit now without doing any further work.
+				// We've finished executing the caller's requested batchcount and we have fulfilled our check-for-idle
+				// contract above (and found we are not yet idle), so exit now without doing any further work.
 				WaitBatchCount = -1;
 				return ETickResult::Interrupt;
 			}
@@ -4120,12 +4120,11 @@ FAssetDataGatherer::ETickResult FAssetDataGatherer::TickInternal(double& TickSta
 			bWaitBatchCountDecremented = true;
 			NumToProcess = FMath::Min(NumToProcess, WaitBatchCount);
 			WaitBatchCount -= NumToProcess;
-			if (WaitBatchCount == 0)
-			{
-				// Mark that WaitBatchCount has been consumed and is no longer active
-				WaitBatchCount = -1;
-				TickResult = ETickResult::Interrupt;
-			}
+			// WaitBatchCount might now be 0, but we do not set TickResult == ETickResult::Interrupt until the next
+			// call to TickInternal, because we need to execute the code above to check for whether we are idle; we
+			// have a contract that we set Idle if possible before breaking out of the InnerTickLoop when processing
+			// a WaitBatchCount. This is important for the case of ScanPathsSynchronous on an asset that needs
+			// LoadCalculatedDependencies. LoadCalculatedDependencies can not run until the Gatherer reports Idle.
 		}
 		DependencyResults.Reserve(FilesToSearch->GetNumAvailable() + DependencyResults.Num());
 		check(NumToProcess > 0);
