AssetRegistry: Fix handling of directory watches on paths outside of the engine directory, which occurs for projects under C:\Users. The AssetRegistry requires that it can compare paths using string comparisons, so normalized paths need to have slashes converted even if they are paths outside of the engine directory.
#rnx

diff --git a/Engine/Source/Runtime/AssetRegistry/Private/AssetRegistry.cpp b/Engine/Source/Runtime/AssetRegistry/Private/AssetRegistry.cpp
--- a/Engine/Source/Runtime/AssetRegistry/Private/AssetRegistry.cpp
+++ b/Engine/Source/Runtime/AssetRegistry/Private/AssetRegistry.cpp
@@ -1250,12 +1250,12 @@ void UAssetRegistryImpl::InitializeEvents(UE::AssetRegistry::Impl::FInitializeCo
 			// directories.
 
 #if UE_ENABLE_DIRECTORYWATCH_ROOTS
-			const FString ProjectPluginDir = FPaths::CreateStandardFilename(FPaths::ProjectPluginsDir());
+			const FString ProjectPluginDir = UE::AssetRegistry::CreateStandardFilename(FPaths::ProjectPluginsDir());
 			if (IPlatformFile::GetPlatformPhysical().DirectoryExists(*ProjectPluginDir))
 			{
 				DirectoryWatchRoots.Add(ProjectPluginDir);
 			}
-			const FString EnginePluginDir = FPaths::CreateStandardFilename(FPaths::EnginePluginsDir());
+			const FString EnginePluginDir = UE::AssetRegistry::CreateStandardFilename(FPaths::EnginePluginsDir());
 			if (IPlatformFile::GetPlatformPhysical().DirectoryExists(*EnginePluginDir))
 			{
 				DirectoryWatchRoots.Add(EnginePluginDir);
@@ -1278,7 +1278,7 @@ void UAssetRegistryImpl::InitializeEvents(UE::AssetRegistry::Impl::FInitializeCo
 			for (TArray<FString>::TConstIterator RootPathIt(Context.RootContentPaths); RootPathIt; ++RootPathIt)
 			{
 				const FString& RootPath = *RootPathIt;
-				ContentFolder = FPaths::CreateStandardFilename(FPackageName::LongPackageNameToFilename(RootPath));
+				ContentFolder = UE::AssetRegistry::CreateStandardFilename(FPackageName::LongPackageNameToFilename(RootPath));
 				if (IsDirAlreadyWatchedByRootWatchers(ContentFolder))
 				{
 					continue;
@@ -1848,7 +1848,7 @@ void UAssetRegistryImpl::FinishDestroy()
 					for (TArray<FString>::TConstIterator RootPathIt(RootContentPaths); RootPathIt; ++RootPathIt)
 					{
 						const FString& RootPath = *RootPathIt;
-						const FString ContentFolder = FPaths::CreateStandardFilename(FPackageName::LongPackageNameToFilename(RootPath));
+						const FString ContentFolder = UE::AssetRegistry::CreateStandardFilename(FPackageName::LongPackageNameToFilename(RootPath));
 						if (!IsDirAlreadyWatchedByRootWatchers(ContentFolder))
 						{
 							DirectoryWatcher->UnregisterDirectoryChangedCallback_Handle(ContentFolder, OnDirectoryChangedDelegateHandles.FindRef(RootPath));
@@ -5709,7 +5709,7 @@ void FAssetRegistryImpl::AddDirectoryReferencer(FName PackageName, const FString
 			*DirectoryLocalPathOrLongPackageName);
 		return;
 	}
-	FPaths::MakeStandardFilename(DirectoryLocalPath);
+	DirectoryLocalPath = UE::AssetRegistry::CreateStandardFilename(DirectoryLocalPath);
 	DirectoryReferencers.AddUnique(DirectoryLocalPath, PackageName);
 }
 
@@ -7681,7 +7681,7 @@ void FAssetRegistryImpl::OnDirectoryChanged(Impl::FEventContext& EventContext,
 		if (bIsValidPackageName)
 		{
 			// If a package changes in a referenced directory, modify the Assets that monitor that directory
-			FString DirectoryPath = FPaths::CreateStandardFilename(FPaths::GetPath(File));
+			FString DirectoryPath = UE::AssetRegistry::CreateStandardFilename(FPaths::GetPath(File));
 			// TODO: Change DirectoryReferencers to a FileNameMap so we can do this FindParentDirectory check more quickly
 			TArray<FName, TInlineAllocator<1>> WatcherPackageNames;
 			for (const TPair<FString, FName>& Pair : DirectoryReferencers)
@@ -7726,7 +7726,7 @@ void FAssetRegistryImpl::OnDirectoryRescanRequired(Impl::FEventContext& EventCon
 {
 	TArray<TPair<FString,FString>> DirPathsAndPackageNames;
 	FString DirPathAsPackageName;
-	FString NormalizedDirPath = FPaths::CreateStandardFilename(DirPath);
+	FString NormalizedDirPath = UE::AssetRegistry::CreateStandardFilename(DirPath);
 	if (FPackageName::TryConvertFilenameToLongPackageName(NormalizedDirPath, DirPathAsPackageName))
 	{
 		DirPathsAndPackageNames.Emplace(DirPath, MoveTemp(DirPathAsPackageName));
@@ -7742,7 +7742,7 @@ void FAssetRegistryImpl::OnDirectoryRescanRequired(Impl::FEventContext& EventCon
 		{
 			if (FPackageName::TryGetMountPointForPath(MountedLongPackageName, UnusedPackageName, MountedFilePath, UnusedRelPath))
 			{
-				FString NormalizeMountedFilePath = FPaths::CreateStandardFilename(FString(MountedFilePath));
+				FString NormalizeMountedFilePath = UE::AssetRegistry::CreateStandardFilename(FString(MountedFilePath));
 				if (FPaths::IsUnderDirectory(NormalizeMountedFilePath, NormalizedDirPath))
 				{
 					DirPathsAndPackageNames.Emplace(MoveTemp(NormalizeMountedFilePath), MoveTemp(MountedLongPackageName));
@@ -7791,7 +7791,7 @@ void FAssetRegistryImpl::OnDirectoryRescanRequired(Impl::FEventContext& EventCon
 					// Convert Filename to a PackagePath. We know the base dir so its faster to use that than FPackageName
 					// which has to scan all mount dirs
 					FStringView RelPath;
-					FString NormalizedFilename = FPaths::CreateStandardFilename(Filename);
+					FString NormalizedFilename = UE::AssetRegistry::CreateStandardFilename(Filename);
 					if (!FPathViews::TryMakeChildPathRelativeTo(NormalizedFilename, LocalPath, RelPath))
 					{
 						return true;
@@ -8231,7 +8231,7 @@ void UAssetRegistryImpl::OnContentPathMounted(const FString& InAssetPath, const
 
 		// Listen for directory changes in this content path
 #if WITH_EDITOR
-		const FString StandardFileSystemPath = FPaths::CreateStandardFilename(FileSystemPath);
+		const FString StandardFileSystemPath = UE::AssetRegistry::CreateStandardFilename(FileSystemPath);
 		// In-game doesn't listen for directory changes
 		if (DirectoryWatcher && !IsDirAlreadyWatchedByRootWatchers(StandardFileSystemPath))
 		{
@@ -8315,7 +8315,7 @@ void UAssetRegistryImpl::OnContentPathDismounted(const FString& InAssetPath, con
 
 		// Stop listening for directory changes in this content path
 #if WITH_EDITOR
-		const FString StandardFileSystemPath = FPaths::CreateStandardFilename(FileSystemPath);
+		const FString StandardFileSystemPath = UE::AssetRegistry::CreateStandardFilename(FileSystemPath);
 		if (DirectoryWatcher && !IsDirAlreadyWatchedByRootWatchers(StandardFileSystemPath))
 		{
 			// Make sure OnDirectoryChangedDelegateHandles key is symmetrical with the one used in OnContentPathMounted
@@ -9887,4 +9887,23 @@ bool FInterruptionContext::ShouldExitEarly()
 }
 
 }
+
+FString CreateStandardFilename(const FString& InPath)
+{
+	FString Result = FPaths::CreateStandardFilename(InPath);
+
+	// Follow most of the behavior of FPaths::CreateStandardFilename, but tweak it to also
+	// replace '\' -> '/' even when the path is not convertible to an engine-relative path.
+	// And add a special case tweak of that tweak, so that a windows-style network share \\
+	// is left unchanged and remains understandable by the OS.
+	bool bKeepLeadingSlashes = InPath.StartsWith(TEXTVIEW("\\\\"));
+	Result.ReplaceInline(TEXT("\\"), TEXT("/"), ESearchCase::CaseSensitive);
+	if (bKeepLeadingSlashes)
+	{
+		Result[0] = '\\';
+		Result[1] = '\\';
+	}
+	return Result;
+}
+
 } // namespace AssetRegistry
diff --git a/Engine/Source/Runtime/AssetRegistry/Private/AssetRegistryImpl.h b/Engine/Source/Runtime/AssetRegistry/Private/AssetRegistryImpl.h
--- a/Engine/Source/Runtime/AssetRegistry/Private/AssetRegistryImpl.h
+++ b/Engine/Source/Runtime/AssetRegistry/Private/AssetRegistryImpl.h
@@ -901,12 +901,22 @@ void EnumerateMemoryAssets(const FARCompiledFilter& InFilter, TSet<FName>& OutPa
 
 } // namespace Utils within namespace UE::AssetRegistry
 
+// Normalization functions for paths. We require that normalized paths are case-insensitively string-identical for the same input
+// path specification. e.g. all slashes must be replaced with forward-slashes. Special case: a leading \\ is left unreplaced,
+// so //share will be seen as different than \\share.
+FString CreateStandardFilename(const FString& InFilename);
+
+} // namespace UE::AssetRegistry
+
 
 ///////////////////////////////////////////////////////
 // Inline implementations
 ///////////////////////////////////////////////////////
 
 
+namespace UE::AssetRegistry
+{
+
 #if WITH_EDITOR
 inline uint32& FAssetRegistryImpl::GetBackgroundTickInterruptionsCount()
 {
@@ -924,5 +934,4 @@ inline void FAssetRegistryImpl::SetGameThreadTakeOverGatherEachTick(bool bValue)
 }
 #endif
 
-
 } // namespace UE::AssetRegistry
