[AssetDataDiscovery]

- Fix crash when scanning for ReparsePoints during AssetDataDiscovery. FAssetDataDiscoveryCache::FindOrAddIsReparsePoint is a recursive function that can manipulate a map in the recursive call. Since we hold on to a reference to an element in the map in the prologue of the recursive function and then use it in the epilogue, if the map has resized, the reference we are holding will become invalid and should be fetched again. To avoid extra lookups, we keep track of when we may have modified the map across recursive calls.

- Adds overload to FCachedVolumeDir::FindOrAddDir taking an out argument to know if a new Dir was added or not.
- Replaced FindOrAddByHash with AddByHash in FCachedVolumeDir::FindOrAddDir when we have already validated that the Dir doesn't exist in the map.

diff --git a/Engine/Source/Runtime/AssetRegistry/Private/AssetDataGathererDiscoveryCache.cpp b/Engine/Source/Runtime/AssetRegistry/Private/AssetDataGathererDiscoveryCache.cpp
--- a/Engine/Source/Runtime/AssetRegistry/Private/AssetDataGathererDiscoveryCache.cpp
+++ b/Engine/Source/Runtime/AssetRegistry/Private/AssetDataGathererDiscoveryCache.cpp
@@ -385,37 +385,54 @@ FCachedDirScanDir* FAssetDataDiscoveryCache::FindDir(FStringView Path, FCachedVo
 	return Volume->FindDir(Path);
 }
 
-bool FAssetDataDiscoveryCache::FindOrAddIsReparsePoint(FStringView DirName, FCachedVolumeInfo& Volume)
+bool FindOrAddIsReparsePointRecursive(FStringView DirName, FCachedVolumeInfo& Volume, bool& bOutAdded)
 {
+	bOutAdded = false;
 	if (DirName.IsEmpty())
 	{
 		return false;
 	}
 
-	FCachedDirScanDir& Existing = Volume.FindOrAddDir(DirName);
+	// We must use a pointer here since re-assigning to Existing later as a value will invoke 
+	// copy assignment on a potentially invalid type
+	FCachedDirScanDir* Existing = &Volume.FindOrAddDir(DirName, bOutAdded);
 
-	if (!Existing.bIsInsideReparsePoint.IsSet())
+	if (!Existing->bIsInsideReparsePoint.IsSet())
 	{
 		// Usually bIsReparsePoint is already set because we calculated it when we encountered this directory in the directory scan of its parent directory.
 		// But for top-level directories or directories that were found to be modified by the journal, we did not already encounter it.
 		// Calculate it now manually if it is not already set.
-		if (!Existing.bIsReparsePoint.IsSet())
+		if (!Existing->bIsReparsePoint.IsSet())
 		{
 			IPlatformFile& PlatformFile = FPlatformFileManager::Get().GetPlatformFile();
-			Existing.bIsReparsePoint = PlatformFile.IsSymlink(DirName.GetData()) == ESymlinkResult::Symlink;
+			Existing->bIsReparsePoint = PlatformFile.IsSymlink(DirName.GetData()) == ESymlinkResult::Symlink;
 		}
 
 		// Base case for when we have found a reparse point, end the recursion and set true for all our previous calls
-		if (Existing.bIsReparsePoint.GetValue())
+		if (Existing->bIsReparsePoint.GetValue())
 		{
-			Existing.bIsInsideReparsePoint = true;
+			Existing->bIsInsideReparsePoint = true;
 			return true;
 		}
 
-		Existing.bIsInsideReparsePoint = FindOrAddIsReparsePoint(FPaths::GetPath(DirName.GetData()), Volume);
+		bool bAddedDuringRecursion = false;
+		bool bIsInsideReparsePoint = FindOrAddIsReparsePointRecursive(FPaths::GetPath(DirName.GetData()), Volume, bAddedDuringRecursion);
+		// If we added a new entry, then we may have invalidated Existing and need to fetch it again
+		if (bAddedDuringRecursion)
+		{
+			bOutAdded = true;
+			Existing = Volume.FindDir(DirName);
+		}
+		Existing->bIsInsideReparsePoint = bIsInsideReparsePoint;
 	}
 
-	return Existing.bIsInsideReparsePoint.GetValue();
+	return Existing->bIsInsideReparsePoint.GetValue();
+}
+
+bool FAssetDataDiscoveryCache::FindOrAddIsReparsePoint(FStringView DirName, FCachedVolumeInfo& Volume)
+{
+	bool bOutAdded = false;
+	return FindOrAddIsReparsePointRecursive(DirName, Volume, bOutAdded);
 }
 
 void FAssetDataDiscoveryCache::QueueConsume()
@@ -549,13 +566,20 @@ void FCachedVolumeInfo::InitializePlatformData()
 	}
 }
 
-FCachedDirScanDir& FCachedVolumeInfo::FindOrAddDir(FStringView Path)
+FCachedDirScanDir& FCachedVolumeInfo::FindOrAddDir(FStringView InPath)
 {
-	uint32 PathHash = GetTypeHash(Path);
-	FCachedDirScanDir* ScanDir = Dirs.FindByHash(PathHash, Path);
-	if (!ScanDir)
+	bool bAdded = false;
+	return FindOrAddDir(InPath, bAdded);
+}
+
+FCachedDirScanDir& FCachedVolumeInfo::FindOrAddDir(FStringView InPath, bool& bOutAdded)
+{
+	uint32 PathHash = GetTypeHash(InPath);
+	FCachedDirScanDir* ScanDir = Dirs.FindByHash(PathHash, InPath);
+	bOutAdded = ScanDir == nullptr;
+	if (bOutAdded)
 	{
-		ScanDir = &Dirs.FindOrAddByHash(PathHash, FString(Path));
+		ScanDir = &Dirs.AddByHash(PathHash, FString(InPath));
 	}
 	return *ScanDir;
 }
diff --git a/Engine/Source/Runtime/AssetRegistry/Private/AssetDataGathererDiscoveryCache.h b/Engine/Source/Runtime/AssetRegistry/Private/AssetDataGathererDiscoveryCache.h
--- a/Engine/Source/Runtime/AssetRegistry/Private/AssetDataGathererDiscoveryCache.h
+++ b/Engine/Source/Runtime/AssetRegistry/Private/AssetDataGathererDiscoveryCache.h
@@ -94,7 +94,8 @@ public:
 	void InitializePlatformData();
 	
 	/** Assumes Path has already been normalized. */
-	FCachedDirScanDir& FindOrAddDir(FStringView Path);
+	FCachedDirScanDir& FindOrAddDir(FStringView InPath);
+	FCachedDirScanDir& FindOrAddDir(FStringView InPath, bool& bOutAdded);
 	void RemoveDirs(TArray<FString>&& InPaths);
 
 	/** Assumes Path has already been normalized. */
