Disable DirectoryWatcher in AR in commandlets.

- Additions (or removal) of asset files and directories are not expected during a commandlet run.
- The change primarily targets reducing cooker (+cook workers) memory usage.

diff --git a/Engine/Source/Runtime/AssetRegistry/Private/AssetRegistry.cpp b/Engine/Source/Runtime/AssetRegistry/Private/AssetRegistry.cpp
--- a/Engine/Source/Runtime/AssetRegistry/Private/AssetRegistry.cpp
+++ b/Engine/Source/Runtime/AssetRegistry/Private/AssetRegistry.cpp
@@ -196,8 +196,8 @@ namespace UE::AssetRegistry::Impl
 
 	bool IsDirectoryWatcherEnabled()
 	{
-		// In-game doesn't listen for directory changes
-		return GIsEditor && !bDisableDirectoryWatcher;
+		// In-game and in commandlets AR doesn't listen for directory changes
+		return !bDisableDirectoryWatcher && GIsEditor && !IsRunningCommandlet();
 	}
 
 	/** Name of UObjectRedirector property */
