AssetRegistry: Whitespace changes: wrap long lines in IAssetRegistry.h.
#rnx

diff --git a/Engine/Source/Runtime/AssetRegistry/Public/AssetRegistry/IAssetRegistry.h b/Engine/Source/Runtime/AssetRegistry/Public/AssetRegistry/IAssetRegistry.h
--- a/Engine/Source/Runtime/AssetRegistry/Public/AssetRegistry/IAssetRegistry.h
+++ b/Engine/Source/Runtime/AssetRegistry/Public/AssetRegistry/IAssetRegistry.h
@@ -364,7 +364,8 @@ public:
 	 * @param OutAssetData out FAssetData 
 	 * @return Enum return code
 	 */
-	virtual UE::AssetRegistry::EExists TryGetAssetByObjectPath(const FSoftObjectPath& ObjectPath, FAssetData& OutAssetData) const = 0;
+	virtual UE::AssetRegistry::EExists TryGetAssetByObjectPath(const FSoftObjectPath& ObjectPath,
+		FAssetData& OutAssetData) const = 0;
 
 	/**
 	 * Tries to get the pacakge data for a specified path
@@ -373,7 +374,8 @@ public:
 	 * @param OutAssetPackageData out FAssetPackageData
 	 * @return Enum return code
 	 */
-	virtual UE::AssetRegistry::EExists TryGetAssetPackageData(FName PackageName, FAssetPackageData& OutAssetPackageData) const = 0;
+	virtual UE::AssetRegistry::EExists TryGetAssetPackageData(FName PackageName,
+		FAssetPackageData& OutAssetPackageData) const = 0;
 
 	/**
 	 * Tries to get the package data for the specified package name. If found, OutCorrectCasePackageName
@@ -384,7 +386,8 @@ public:
 	 * @param OutCorrectCasePackageName out FName matching filesystem casing
 	 * @return Return code enum
 	 */
-	virtual UE::AssetRegistry::EExists TryGetAssetPackageData(FName PackageName, FAssetPackageData& OutAssetPackageData, FName& OutCorrectCasePackageName) const = 0;
+	virtual UE::AssetRegistry::EExists TryGetAssetPackageData(FName PackageName,
+		FAssetPackageData& OutAssetPackageData, FName& OutCorrectCasePackageName) const = 0;
 	
 	/**
 	 * Gets asset data for all assets in the registry.
@@ -439,56 +442,84 @@ public:
 	virtual FName GetFirstPackageByName(FStringView PackageName) const = 0;
 
 	/**
-	 * Gets a list of AssetIdentifiers or FAssetDependencies that are referenced by the supplied AssetIdentifier. (On disk references ONLY)
+	 * Gets a list of AssetIdentifiers or FAssetDependencies that are referenced by the supplied AssetIdentifier.
+	 * Only returns dependencies reported in the on-disk package.
 	 *
-	 * @param AssetIdentifier	the name of the package/name for which to gather dependencies.
-	 * @param OutDependencies	a list of things that are referenced by AssetIdentifier.
-	 * @param Category	which category(ies) of dependencies to include in the output list. Dependencies matching ANY of the OR'd categories will be returned.
-	 * @param Flags	which flags are required present or not present on the dependencies. Dependencies matching ALL required and NONE excluded bits will be returned. For each potentially returned dependency, flags not applicable to their category are ignored.
+	 * @param AssetIdentifier The name of the package/name for which to gather dependencies.
+	 * @param OutDependencies A list of things that are referenced by AssetIdentifier.
+	 * @param Category Which category(ies) of dependencies to include in the output list. Dependencies matching ANY
+	 *        of the OR'd categories will be returned.
+	 * @param Flags	Which flags are required present or not present on the dependencies. Dependencies matching ALL
+	 *        required and NONE excluded bits will be returned. For each potentially returned dependency, flags not
+	 *        applicable to their category are ignored.
 	 */
-	virtual bool GetDependencies(const FAssetIdentifier& AssetIdentifier, TArray<FAssetIdentifier>& OutDependencies, UE::AssetRegistry::EDependencyCategory Category = UE::AssetRegistry::EDependencyCategory::All, const UE::AssetRegistry::FDependencyQuery& Flags = UE::AssetRegistry::FDependencyQuery()) const = 0;
-	virtual bool GetDependencies(const FAssetIdentifier& AssetIdentifier, TArray<FAssetDependency>& OutDependencies, UE::AssetRegistry::EDependencyCategory Category = UE::AssetRegistry::EDependencyCategory::All, const UE::AssetRegistry::FDependencyQuery& Flags = UE::AssetRegistry::FDependencyQuery()) const = 0;
+	virtual bool GetDependencies(const FAssetIdentifier& AssetIdentifier, TArray<FAssetIdentifier>& OutDependencies,
+		UE::AssetRegistry::EDependencyCategory Category = UE::AssetRegistry::EDependencyCategory::All,
+		const UE::AssetRegistry::FDependencyQuery& Flags = UE::AssetRegistry::FDependencyQuery()) const = 0;
+	virtual bool GetDependencies(const FAssetIdentifier& AssetIdentifier, TArray<FAssetDependency>& OutDependencies,
+		UE::AssetRegistry::EDependencyCategory Category = UE::AssetRegistry::EDependencyCategory::All,
+		const UE::AssetRegistry::FDependencyQuery& Flags = UE::AssetRegistry::FDependencyQuery()) const = 0;
 
 	/**
 	 * Gets a list of PackageNames that are referenced by the supplied package. (On disk references ONLY)
 	 *
-	 * @param PackageName		the name of the package for which to gather dependencies (eg, /Game/MyFolder/MyAsset)
-	 * @param OutDependencies	a list of packages that are referenced by the package whose path is PackageName
-	 * @param Category	which category(ies) of dependencies to include in the output list. Dependencies matching ANY of the OR'd categories will be returned.
-	 * @param Flags	which flags are required present or not present on the dependencies. Dependencies matching ALL required and NONE excluded bits will be returned. For each potentially returned dependency, flags not applicable to their category are ignored.
+	 * @param PackageName The name of the package for which to gather dependencies (eg, /Game/MyFolder/MyAsset)
+	 * @param OutDependencies A list of packages that are referenced by the package whose path is PackageName
+	 * @param Category Which category(ies) of dependencies to include in the output list. Dependencies matching
+	 *        ANY of the OR'd categories will be returned.
+	 * @param Flags	Which flags are required present or not present on the dependencies. Dependencies matching ALL
+	 *        required and NONE excluded bits will be returned. For each potentially returned dependency, flags not
+	 *        applicable to their category are ignored.
 	 */
-	virtual bool GetDependencies(FName PackageName, TArray<FName>& OutDependencies, UE::AssetRegistry::EDependencyCategory Category = UE::AssetRegistry::EDependencyCategory::Package, const UE::AssetRegistry::FDependencyQuery& Flags = UE::AssetRegistry::FDependencyQuery()) const = 0;
+	virtual bool GetDependencies(FName PackageName, TArray<FName>& OutDependencies,
+		UE::AssetRegistry::EDependencyCategory Category = UE::AssetRegistry::EDependencyCategory::Package,
+		const UE::AssetRegistry::FDependencyQuery& Flags = UE::AssetRegistry::FDependencyQuery()) const = 0;
 
 	/**
 	 * Gets a list of paths to objects that are referenced by the supplied package. (On disk references ONLY)
 	 *
-	 * @param PackageName		the name of the package for which to gather dependencies (eg, /Game/MyFolder/MyAsset)
-	 * @param DependencyOptions	which kinds of dependencies to include in the output list
-	 * @param OutDependencies	a list of packages that are referenced by the package whose path is PackageName
+	 * @param PackageName The name of the package for which to gather dependencies (eg, /Game/MyFolder/MyAsset)
+	 * @param DependencyOptions Which kinds of dependencies to include in the output list
+	 * @param OutDependencies A list of packages that are referenced by the package whose path is PackageName
 	 */
-	UFUNCTION(BlueprintCallable, BlueprintPure=false, Category = "AssetRegistry", meta=(DisplayName="Get Dependencies", ScriptName="GetDependencies"))
-	virtual bool K2_GetDependencies(FName PackageName, const FAssetRegistryDependencyOptions& DependencyOptions, TArray<FName>& OutDependencies) const;
+	UFUNCTION(BlueprintCallable, BlueprintPure=false, Category = "AssetRegistry",
+		meta=(DisplayName="Get Dependencies", ScriptName="GetDependencies"))
+	virtual bool K2_GetDependencies(FName PackageName, const FAssetRegistryDependencyOptions& DependencyOptions,
+		TArray<FName>& OutDependencies) const;
 
 	/**
-	 * Gets a list of AssetIdentifiers or FAssetDependencies that reference the supplied AssetIdentifier. (On disk references ONLY)
+	 * Gets a list of AssetIdentifiers or FAssetDependencies that reference the supplied AssetIdentifier.
+	 * Only returns referencers reported in the on-disk package.
 	 *
-	 * @param AssetIdentifier	the name of the package/name for which to gather referencers (eg, /Game/MyFolder/MyAsset)
-	 * @param OutReferencers	a list of things that reference AssetIdentifier.
-	 * @param Category	which category(ies) of referencers to include in the output list. Referencers that have a dependency matching ANY of the OR'd categories will be returned.
-	 * @param Flags	which flags are required present or not present on the referencer's dependency. Referencers that have a dependency matching ALL required and NONE excluded bits will be returned. For each potentially returned dependency, flags not applicable to their category are ignored.
+	 * @param AssetIdentifier The name of the package/name for which to gather referencers (eg, /Game/MyFolder/MyAsset)
+	 * @param OutReferencers A list of things that reference AssetIdentifier.
+	 * @param Category Which category(ies) of referencers to include in the output list. Referencers that have a
+	 *        dependency matching ANY of the OR'd categories will be returned.
+	 * @param Flags Which flags are required present or not present on the referencer's dependency. Referencers that
+	 *        have a dependency matching ALL required and NONE excluded bits will be returned. For each potentially
+	 *        returned dependency, flags not applicable to their category are ignored.
 	 */
-	virtual bool GetReferencers(const FAssetIdentifier& AssetIdentifier, TArray<FAssetIdentifier>& OutReferencers, UE::AssetRegistry::EDependencyCategory Category = UE::AssetRegistry::EDependencyCategory::All, const UE::AssetRegistry::FDependencyQuery& Flags = UE::AssetRegistry::FDependencyQuery()) const = 0;
-	virtual bool GetReferencers(const FAssetIdentifier& AssetIdentifier, TArray<FAssetDependency>& OutReferencers, UE::AssetRegistry::EDependencyCategory Category = UE::AssetRegistry::EDependencyCategory::All, const UE::AssetRegistry::FDependencyQuery& Flags = UE::AssetRegistry::FDependencyQuery()) const = 0;
+	virtual bool GetReferencers(const FAssetIdentifier& AssetIdentifier, TArray<FAssetIdentifier>& OutReferencers,
+		UE::AssetRegistry::EDependencyCategory Category = UE::AssetRegistry::EDependencyCategory::All,
+		const UE::AssetRegistry::FDependencyQuery& Flags = UE::AssetRegistry::FDependencyQuery()) const = 0;
+	virtual bool GetReferencers(const FAssetIdentifier& AssetIdentifier, TArray<FAssetDependency>& OutReferencers,
+		UE::AssetRegistry::EDependencyCategory Category = UE::AssetRegistry::EDependencyCategory::All,
+		const UE::AssetRegistry::FDependencyQuery& Flags = UE::AssetRegistry::FDependencyQuery()) const = 0;
 
 	/**
 	 * Gets a list of PackageNames that reference the supplied package. (On disk references ONLY)
 	 *
-	 * @param PackageName		the name of the package for which to gather dependencies (eg, /Game/MyFolder/MyAsset)
-	 * @param OutReferencers	a list of packages that reference the package whose path is PackageName
-	 * @param Category	which category(ies) of referencers to include in the output list. Referencers that have a dependency matching ANY of the OR'd categories will be returned.
-	 * @param Flags	which flags are required present or not present on the referencer's dependency. Referencers that have a dependency matching ALL required and NONE excluded bits will be returned. For each potentially returned dependency, flags not applicable to their category are ignored.
+	 * @param PackageName The name of the package for which to gather dependencies (eg, /Game/MyFolder/MyAsset)
+	 * @param OutReferencers A list of packages that reference the package whose path is PackageName
+	 * @param Category Which category(ies) of referencers to include in the output list. Referencers that have a
+	 *        dependency matching ANY of the OR'd categories will be returned.
+	 * @param Flags Which flags are required present or not present on the referencer's dependency. Referencers that
+	 *         have a dependency matching ALL required and NONE excluded bits will be returned. For each potentially
+	 *         returned dependency, flags not applicable to their category are ignored.
 	 */
-	virtual bool GetReferencers(FName PackageName, TArray<FName>& OutReferencers, UE::AssetRegistry::EDependencyCategory Category = UE::AssetRegistry::EDependencyCategory::Package, const UE::AssetRegistry::FDependencyQuery& Flags = UE::AssetRegistry::FDependencyQuery()) const = 0;
+	virtual bool GetReferencers(FName PackageName, TArray<FName>& OutReferencers,
+		UE::AssetRegistry::EDependencyCategory Category = UE::AssetRegistry::EDependencyCategory::Package,
+		const UE::AssetRegistry::FDependencyQuery& Flags = UE::AssetRegistry::FDependencyQuery()) const = 0;
 
 	/**
 	 * Gets a list of packages that reference the supplied package. (On disk references ONLY)
@@ -497,8 +528,10 @@ public:
 	 * @param ReferenceOptions	which kinds of references to include in the output list
 	 * @param OutReferencers	a list of packages that reference the package whose path is PackageName
 	 */
-	UFUNCTION(BlueprintCallable, BlueprintPure=false, Category = "AssetRegistry", meta=(DisplayName="Get Referencers", ScriptName="GetReferencers"))
-	virtual bool K2_GetReferencers(FName PackageName, const FAssetRegistryDependencyOptions& ReferenceOptions, TArray<FName>& OutReferencers) const;
+	UFUNCTION(BlueprintCallable, BlueprintPure=false, Category = "AssetRegistry",
+		meta=(DisplayName="Get Referencers", ScriptName="GetReferencers"))
+	virtual bool K2_GetReferencers(FName PackageName, const FAssetRegistryDependencyOptions& ReferenceOptions,
+		TArray<FName>& OutReferencers) const;
 
 	/** Finds Package data for a package name. This data is only updated on save and can only be accessed for valid packages */
 	virtual TOptional<FAssetPackageData> GetAssetPackageDataCopy(FName PackageName) const = 0;
@@ -521,18 +554,32 @@ public:
 	* @param OutCorrectCasePackageName 	out package name matching the case on disk
 	* @param OutExtension				out extension for packages containing assets (dot included)
 	*/
-	virtual bool DoesPackageExistOnDisk(FName PackageName, FString* OutCorrectCasePackageName = nullptr, FString* OutExtension = nullptr) const = 0;
+	virtual bool DoesPackageExistOnDisk(FName PackageName, FString* OutCorrectCasePackageName = nullptr,
+		FString* OutExtension = nullptr) const = 0;
 
-	/** Uses the asset registry to look for ObjectRedirectors. This will follow the chain of redirectors. It will return the original path if no redirectors are found */
+	/**
+	 * Uses the asset registry to look for ObjectRedirectors. This will follow the chain of redirectors. It will
+	 * return the original path if no redirectors are found.
+	 */
 	virtual FSoftObjectPath GetRedirectedObjectPath(const FSoftObjectPath& ObjectPath) = 0;
 
-	/** Returns true if the specified ClassName's ancestors could be found. If so, OutAncestorClassNames is a list of all its ancestors. This can be slow if temporary caching mode is not on */
-	UFUNCTION(BlueprintCallable, BlueprintPure = false, Category = "AssetRegistry", meta = (DisplayName = "GetAncestorClassNames", ScriptName = "GetAncestorClassNames"))
-	virtual bool GetAncestorClassNames(FTopLevelAssetPath ClassPathName, TArray<FTopLevelAssetPath>& OutAncestorClassNames) const = 0;
+	/**
+	 * Returns true if the specified ClassName's ancestors could be found. If so, OutAncestorClassNames is a list of
+	 * all its ancestors. This can be slow if temporary caching mode is not on.
+	 */
+	UFUNCTION(BlueprintCallable, BlueprintPure = false, Category = "AssetRegistry",
+		meta=(DisplayName = "GetAncestorClassNames", ScriptName = "GetAncestorClassNames"))
+	virtual bool GetAncestorClassNames(FTopLevelAssetPath ClassPathName,
+		TArray<FTopLevelAssetPath>& OutAncestorClassNames) const = 0;
 
-	/** Returns the names of all classes derived by the supplied class names, excluding any classes matching the excluded class names. This can be slow if temporary caching mode is not on */
-	UFUNCTION(BlueprintCallable, BlueprintPure = false, Category = "AssetRegistry", meta = (DisplayName = "GetDerivedClassNames", ScriptName = "GetDerivedClassNames"))
-	virtual void GetDerivedClassNames(const TArray<FTopLevelAssetPath>& ClassNames, const TSet<FTopLevelAssetPath>& ExcludedClassNames, TSet<FTopLevelAssetPath>& OutDerivedClassNames) const = 0;
+	/**
+	 * Returns the names of all classes derived by the supplied class names, excluding any classes matching the
+	 * excluded class names. This can be slow if temporary caching mode is not on.
+	 */
+	UFUNCTION(BlueprintCallable, BlueprintPure = false, Category = "AssetRegistry",
+		meta=(DisplayName = "GetDerivedClassNames", ScriptName = "GetDerivedClassNames"))
+	virtual void GetDerivedClassNames(const TArray<FTopLevelAssetPath>& ClassNames,
+		const TSet<FTopLevelAssetPath>& ExcludedClassNames, TSet<FTopLevelAssetPath>& OutDerivedClassNames) const = 0;
 	
 	/** Gets a list of all paths that are currently cached */
 	UFUNCTION(BlueprintCallable, BlueprintPure=false, Category = "AssetRegistry")
@@ -577,7 +624,10 @@ public:
 	/** Modifies passed in filter optimize it for query and expand any recursive paths and classes */
 	virtual void CompileFilter(const FARFilter& InFilter, FARCompiledFilter& OutCompiledFilter) const = 0;
 
-	/** Enables or disable temporary search caching, when this is enabled scanning/searching is faster because we assume no objects are loaded between scans. Disabling frees any caches created */
+	/**
+	 * Enables or disable temporary search caching, when this is enabled scanning/searching is faster because we
+	 * assume no objects are loaded between scans. Disabling frees any caches created.
+	 */
 	virtual void SetTemporaryCachingMode(bool bEnable) = 0;
 	/**
 	 * Mark that the temporary cached needs to be updated before being used again, because e.g. a new class was loaded.
@@ -601,7 +651,8 @@ public:
 	 * @param FAssetData the asset to check for progress status
 	 * @param ReportType the type of report to query.
 	 */
-	virtual float GetAssetAvailabilityProgress(const FAssetData& AssetData, EAssetAvailabilityProgressReportingType::Type ReportType) const = 0;
+	virtual float GetAssetAvailabilityProgress(const FAssetData& AssetData,
+		EAssetAvailabilityProgressReportingType::Type ReportType) const = 0;
 
 	/**
 	 * @param ReportType The report type to query.
@@ -629,57 +680,75 @@ public:
 	 * Gets paths for all Verse files in the supplied folder
 	 *
 	 * @param PackagePath the path to query asset data in (e.g. /Game/MyFolder)
-	 * @param OutFilePaths the list of Verse files in this path, as pseudo UE LongPackagePaths with extension (e.g. /Game/MyFolder/MyVerseFile.verse)
+	 * @param OutFilePaths the list of Verse files in this path, as pseudo UE LongPackagePaths with extension
+	 *       (e.g. /Game/MyFolder/MyVerseFile.verse)
 	 * @param bRecursive if true, the supplied path will be searched recursively
 	 */
 	virtual bool GetVerseFilesByPath(FName PackagePath, TArray<FName>& OutFilePaths, bool bRecursive = false) const = 0;
 
-	/** Adds the specified path to the set of cached paths. These will be returned by GetAllCachedPaths(). Returns true if the path was actually added and false if it already existed. */
+	/**
+	 * Adds the specified path to the set of cached paths. These will be returned by GetAllCachedPaths(). Returns
+	 * true if the path was actually added and false if it already existed.
+	 */
 	virtual bool AddPath(const FString& PathToAdd) = 0;
 
-	/** Attempts to remove the specified path to the set of cached paths. This will only succeed if there are no assets left in the specified path. */
+	/**
+	 * Attempts to remove the specified path to the set of cached paths. This will only succeed if there are no
+	 * assets left in the specified path.
+	 */
 	virtual bool RemovePath(const FString& PathToRemove) = 0;
 
 	/** Queries whether the given path exists in the set of cached paths */
 	virtual bool PathExists(const FString& PathToTest) const = 0;
 	virtual bool PathExists(const FName PathToTest) const = 0;
 
-	/** Scan the supplied paths recursively right now and populate the asset registry. If bForceRescan is true, the paths will be scanned again, even if they were previously scanned */
+	/**
+	 * Scan the supplied paths recursively right now and populate the asset registry. If bForceRescan is true, the
+	 * paths will be scanned again, even if they were previously scanned.
+	 */
 	UFUNCTION(BlueprintCallable, Category = "AssetRegistry")
-	virtual void ScanPathsSynchronous(const TArray<FString>& InPaths, bool bForceRescan = false, bool bIgnoreDenyListScanFilters = false) = 0;
+	virtual void ScanPathsSynchronous(const TArray<FString>& InPaths, bool bForceRescan = false,
+		bool bIgnoreDenyListScanFilters = false) = 0;
 
-	/** Scan the specified individual files right now and populate the asset registry. If bForceRescan is true, the paths will be scanned again, even if they were previously scanned */
+	/**
+	 * Scan the specified individual files right now and populate the asset registry. If bForceRescan is true, the
+	 * paths will be scanned again, even if they were previously scanned.
+	 */
 	UFUNCTION(BlueprintCallable, Category = "AssetRegistry")
 	virtual void ScanFilesSynchronous(const TArray<FString>& InFilePaths, bool bForceRescan = false) = 0;
 
 	/** Scan the supplied paths and files recursively right now and populate the asset registry. */
-	virtual void ScanSynchronous(const TArray<FString>& InPaths, const TArray<FString>& InFilePaths, UE::AssetRegistry::EScanFlags InScanFlags = UE::AssetRegistry::EScanFlags::None) = 0;
+	virtual void ScanSynchronous(const TArray<FString>& InPaths, const TArray<FString>& InFilePaths,
+		UE::AssetRegistry::EScanFlags InScanFlags = UE::AssetRegistry::EScanFlags::None) = 0;
 
 	/** Look for all assets on disk (can be async or synchronous) */
 	UFUNCTION(BlueprintCallable, Category = "AssetRegistry")
 	virtual void SearchAllAssets(bool bSynchronousSearch) = 0;
 
 	/**
-	 * Whether SearchAllAssets has been called, or was auto-called at startup. When async (editor or cooking), if SearchAllAssets has ever been called,
-	 * any newly-mounted directory will be automatically searched.
+	 * Whether SearchAllAssets has been called, or was auto-called at startup. When async (editor or cooking), if
+	 * SearchAllAssets has ever been called, any newly-mounted directory will be automatically searched.
 	 */
 	UFUNCTION(BlueprintCallable, Category = "AssetRegistry")
 	virtual bool IsSearchAllAssets() const = 0;
 
-	/** Whether searching is done async (and was started at startup), or synchronously and on-demand, requiring ScanPathsSynchronous or SearchAllAssets. */
+	/**
+	 * Whether searching is done async (and was started at startup), or synchronously and on-demand, requiring
+	 * ScanPathsSynchronous or SearchAllAssets.
+	 */
 	UFUNCTION(BlueprintCallable, Category = "AssetRegistry")
 	virtual bool IsSearchAsync() const = 0;
 
 	/**
-	 * Wait for scan to be complete. If called during editor startup before OnPostEngineInit, and there are any assets that use classes in 
-	 * not-yet-loaded plugin modules, WaitForCompletion will return silently with those assets still ungathered.
+	 * Wait for scan to be complete. If called during editor startup before OnPostEngineInit, and there are assets with
+	 * classes in not-yet-loaded plugin modules, WaitForCompletion returns silently with those assets still ungathered.
 	 */
 	UFUNCTION(BlueprintCallable, Category = "AssetRegistry")
 	virtual void WaitForCompletion() = 0;
 
 	/**
-	 * Wait for premade asset registry load to be completed, for use when querying the asset registry prior to flushing in ScanPathSynchronous.
-	 * Will no-op if no such asset registry exists.
+	 * Wait for premade asset registry load to be completed, for use when querying the asset registry prior to flushing
+	 * in ScanPathSynchronous. Will no-op if no such asset registry exists.
 	 */
 	virtual void WaitForPremadeAssetRegistry() = 0;
 
@@ -693,7 +762,10 @@ public:
 	UFUNCTION(BlueprintCallable, Category = "AssetRegistry")
 	virtual void WaitForPackage(const FString& PackageName) = 0;
 
-	/** If assets are currently being asynchronously scanned in the specified path, this will cause them to be scanned before other assets. */
+	/**
+	 * If assets are currently being asynchronously scanned in the specified path, this will cause them to be scanned
+	 * before other assets.
+	 */
 	UFUNCTION(BlueprintCallable, Category = "AssetRegistry")
 	virtual void PrioritizeSearchPath(const FString& PathToPrioritize) = 0;
 
@@ -732,7 +804,10 @@ public:
 	UE_DEPRECATED(5.2, "Use the new AssetsSaved function that takes FAssetData.")
 	virtual void AssetSaved(const UObject& SavedAsset) = 0;
 
-	/** Called during SavePackage to update the AssetRegistry's copy of the AssetDatas in the package to match the newly saved values. */
+	/**
+	 * Called during SavePackage to update the AssetRegistry's copy of the AssetDatas in the package to match the
+	 * newly saved values.
+	 */
 	virtual void AssetsSaved(TArray<FAssetData>&& Assets) = 0;
 
 	/**
@@ -765,8 +840,9 @@ public:
 	DECLARE_TS_MULTICAST_DELEGATE_OneParam( FAssetsEvent, TConstArrayView<FAssetData> );
 
 	/** 
-	* EXPERIMENTAL: event for when duplicated assets are found and needs to be resolved, this is experimental as it's a counter measure for
-	* a situation that shouldn't happen. Also, DO NOT CALL AssetRegistry functions from delegate subscribers; doing so will create a deadlock.
+	* EXPERIMENTAL: event for when duplicated assets are found and needs to be resolved, this is experimental as it's
+	* a counter measure for a situation that shouldn't happen. Also, DO NOT CALL AssetRegistry functions from delegate
+	* subscribers; doing so will create a deadlock.
 	*/
 	DECLARE_TS_MULTICAST_DELEGATE_ThreeParams( FAssetCollisionEvent, FAssetData&, FAssetData&, FAssetData*& );
 	virtual FAssetCollisionEvent& OnAssetCollision_Private() = 0;
@@ -837,7 +913,8 @@ public:
 	/** Payload data for a file progress update */
 	struct FFileLoadProgressUpdateData
 	{
-		FFileLoadProgressUpdateData(int32 InNumTotalAssets, int32 InNumAssetsProcessedByAssetRegistry, int32 InNumAssetsPendingDataLoad, bool InIsDiscoveringAssetFiles)
+		FFileLoadProgressUpdateData(int32 InNumTotalAssets, int32 InNumAssetsProcessedByAssetRegistry,
+			int32 InNumAssetsPendingDataLoad, bool InIsDiscoveringAssetFiles)
 			: NumTotalAssets(InNumTotalAssets)
 			, NumAssetsProcessedByAssetRegistry(InNumAssetsProcessedByAssetRegistry)
 			, NumAssetsPendingDataLoad(InNumAssetsPendingDataLoad)
@@ -910,10 +987,14 @@ public:
 #if ASSET_REGISTRY_STATE_DUMPING_ENABLED
 	/**
 	 * Writes out the state in textual form. Use arguments to control which segments to emit.
-	 * @param Arguments List of segments to emit. Possible values: 'ObjectPath', 'PackageName', 'Path', 'Class', 'Tag', 'Dependencies' and 'PackageData'
-	 * @param OutPages Textual representation will be written to this array; each entry will have LinesPerPage lines of the full dump.
-	 * @param LinesPerPage - how many lines should be combined into each string element of OutPages, for e.g. breaking up the dump into separate files.
-	 *        To facilitate diffing between similar-but-different registries, the actual number of lines per page will be slightly less than LinesPerPage; we introduce partially deterministic pagebreaks near the end of each page.
+	 * @param Arguments List of segments to emit. Possible values: 'ObjectPath', 'PackageName', 'Path', 'Class',
+	 *        'Tag', 'Dependencies' and 'PackageData'
+	 * @param OutPages Textual representation will be written to this array; each entry will have LinesPerPage lines
+	 *        of the full dump.
+	 * @param LinesPerPage - how many lines should be combined into each string element of OutPages, for e.g. breaking
+	 *        up the dump into separate files. To facilitate diffing between similar-but-different registries, the
+	 *        actual number of lines per page will be slightly less than LinesPerPage; we introduce partially
+	 *        deterministic pagebreaks near the end of each page.
 	 */
 	virtual void DumpState(const TArray<FString>& Arguments, TArray<FString>& OutPages, int32 LinesPerPage = 1) const = 0;
 #endif
@@ -927,7 +1008,9 @@ public:
 	virtual bool ContainsTag(FName TagName) const = 0;
 
 	/** Fills in FAssetRegistrySerializationOptions from ini, optionally using a target platform ini name */
-	virtual void InitializeSerializationOptions(FAssetRegistrySerializationOptions& Options, const FString& PlatformIniName = FString(), UE::AssetRegistry::ESerializationTarget Target = UE::AssetRegistry::ESerializationTarget::ForGame) const = 0;
+	virtual void InitializeSerializationOptions(FAssetRegistrySerializationOptions& Options,
+		const FString& PlatformIniName = FString(),
+		UE::AssetRegistry::ESerializationTarget Target = UE::AssetRegistry::ESerializationTarget::ForGame) const = 0;
 
 	struct FLoadPackageRegistryData
 	{
@@ -961,7 +1044,8 @@ public:
 	 * To avoid copies, the callback is called from within the ReadLock.
 	 * DO NOT CALL AssetRegistry functions from the callback; doing so will create a deadlock.
 	 */
-	virtual void ReadLockEnumerateAllTagToAssetDatas(TFunctionRef<bool(FName TagName, FEnumerateAssetDatasFunc EnumerateAssets)> Callback) const = 0;
+	virtual void ReadLockEnumerateAllTagToAssetDatas(TFunctionRef<bool(FName TagName,
+		FEnumerateAssetDatasFunc EnumerateAssets)> Callback) const = 0;
 
 	/**
 	 * Predicate called to decide whether to recurse into a reference when setting manager references
@@ -972,8 +1056,10 @@ public:
 	 * @param DependencyType	Type of dependencies to recurse over
 	 * @param Flags				Flags describing this particular set attempt
 	 */
-	typedef TFunction<EAssetSetManagerResult::Type(const FAssetIdentifier& Manager, const FAssetIdentifier& Source, const FAssetIdentifier& Target,
-		UE::AssetRegistry::EDependencyCategory Category, UE::AssetRegistry::EDependencyProperty Properties, EAssetSetManagerFlags::Type Flags)> ShouldSetManagerPredicate;
+	typedef TFunction<EAssetSetManagerResult::Type(const FAssetIdentifier& Manager, const FAssetIdentifier& Source,
+		const FAssetIdentifier& Target, UE::AssetRegistry::EDependencyCategory Category,
+		UE::AssetRegistry::EDependencyProperty Properties, EAssetSetManagerFlags::Type Flags)>
+		ShouldSetManagerPredicate;
 
 	/** 
 	 *	Indicates if path should be beautified before presented to the user.
@@ -1002,17 +1088,23 @@ protected:
 	friend class UAssetManager;
 
 	/**
-	 * Specifies a list of manager mappings, optionally recursing to dependencies. These mappings can then be queried later to see which assets "manage" other assets
-	 * This function is only meant to be called by the AssetManager, calls from anywhere else will conflict and lose data
+	 * Specifies a list of manager mappings, optionally recursing to dependencies. These mappings can then be queried
+	 * later to see which assets "manage" other assets. This function is only meant to be called by the AssetManager,
+	 * calls from anywhere else will conflict and lose data.
 	 *
-	 * @param ManagerMap		Map of Managing asset to Managed asset. This will construct Manager references and clear existing 
-	 * @param bClearExisting	If true, will clear any existing manage dependencies
-	 * @param RecurseType		Dependency types to recurse into, from the value of the manager map
-	 * @param RecursePredicate	Predicate that is called on recursion if bound, return true if it should recurse into that node
+	 * @param ManagerMap Map of Managing asset to Managed asset. This will construct Manager references and clear existing
+	 * @param bClearExisting If true, will clear any existing manage dependencies
+	 * @param RecurseType Dependency types to recurse into, from the value of the manager map
+	 * @param RecursePredicate Predicate that is called on recursion if bound, return true if it should recurse into that node
 	 */
-	virtual void SetManageReferences(const TMultiMap<FAssetIdentifier, FAssetIdentifier>& ManagerMap, bool bClearExisting, UE::AssetRegistry::EDependencyCategory RecurseType, TSet<FDependsNode*>& ExistingManagedNodes, ShouldSetManagerPredicate ShouldSetManager = nullptr) = 0;
+	virtual void SetManageReferences(const TMultiMap<FAssetIdentifier, FAssetIdentifier>& ManagerMap,
+		bool bClearExisting, UE::AssetRegistry::EDependencyCategory RecurseType,
+		TSet<FDependsNode*>& ExistingManagedNodes, ShouldSetManagerPredicate ShouldSetManager = nullptr) = 0;
 
-	/** Sets the PrimaryAssetId for a specific asset. This should only be called by the AssetManager, and is needed when the AssetManager is more up to date than the on disk Registry */
+	/**
+	 * Sets the PrimaryAssetId for a specific asset. This should only be called by the AssetManager, and is needed
+	 * when the AssetManager is more up to date than the on disk Registry.
+	 */
 	virtual bool SetPrimaryAssetIdForObjectPath(const FSoftObjectPath& ObjectPath, FPrimaryAssetId PrimaryAssetId) = 0;
 };
 
@@ -1031,7 +1123,8 @@ namespace UE::AssetRegistry
 		int64 DependencyDataOffset = INDEX_NONE;
 		int32 ObjectCount = 0;
 
-		ASSETREGISTRY_API bool DoSerialize(FArchive& BinaryArchive, const FPackageFileSummary& PackageFileSummary, EReadPackageDataMainErrorCode& OutError);
+		ASSETREGISTRY_API bool DoSerialize(FArchive& BinaryArchive, const FPackageFileSummary& PackageFileSummary,
+			EReadPackageDataMainErrorCode& OutError);
 	};
 
 	struct FDeserializeObjectPackageData
@@ -1051,10 +1144,13 @@ namespace UE::AssetRegistry
 		ASSETREGISTRY_API bool DoSerialize(FArchive& BinaryArchive, EReadPackageDataMainErrorCode& OutError);
 	};
 
-	// Functions to read and write the data used by the AssetRegistry in each package; the format of this data is separate from the format of the data in the asset registry
-	// WritePackageData is declared in AssetRegistryInterface.h, in the CoreUObject module, because it is needed by SavePackage in CoreUObject
-	ASSETREGISTRY_API bool ReadPackageDataMain(FArchive& BinaryArchive, const FString& PackageName, const FPackageFileSummary& PackageFileSummary,
-		int64& OutDependencyDataOffset, TArray<FAssetData*>& OutAssetDataList, EReadPackageDataMainErrorCode& OutError,
+	// Functions to read and write the data used by the AssetRegistry in each package; the format of this data is
+	// separate from the format of the data in the asset registry.
+	// WritePackageData is declared in AssetRegistryInterface.h, in the CoreUObject module, because it is needed by
+	// SavePackage in CoreUObject
+	ASSETREGISTRY_API bool ReadPackageDataMain(FArchive& BinaryArchive, const FString& PackageName,
+		const FPackageFileSummary& PackageFileSummary, int64& OutDependencyDataOffset,
+		TArray<FAssetData*>& OutAssetDataList, EReadPackageDataMainErrorCode& OutError,
 		const TArray<FObjectImport>* InImports = nullptr, const TArray<FObjectExport>* InExports = nullptr);
 	struct FReadPackageDataDependenciesArgs
 	{
@@ -1073,14 +1169,16 @@ namespace UE::AssetRegistry
 	};
 	ASSETREGISTRY_API bool ReadPackageDataDependencies(FReadPackageDataDependenciesArgs& Args);
 	UE_DEPRECATED(5.5, "Use version that takes FReadPackageDataArgs")
-	ASSETREGISTRY_API bool ReadPackageDataDependencies(FArchive& BinaryArchive, TBitArray<>& OutImportUsedInGame, TBitArray<>& OutSoftPackageUsedInGame);
+	ASSETREGISTRY_API bool ReadPackageDataDependencies(FArchive& BinaryArchive, TBitArray<>& OutImportUsedInGame,
+		TBitArray<>& OutSoftPackageUsedInGame);
 
 	/**
 	 * Given a list of packages, gather the most important assets for each package.
 	 * If multiple assets are in a package, the most important asset will be added.
 	 * If a package does not exist or does not have any assets, no entry will be added for that package name.
 	 */
-	ASSETREGISTRY_API void GetAssetForPackages(TConstArrayView<FName> PackageNames, TMap<FName, FAssetData>& OutPackageToAssetData);
+	ASSETREGISTRY_API void GetAssetForPackages(TConstArrayView<FName> PackageNames, TMap<FName,
+		FAssetData>& OutPackageToAssetData);
 
 	enum class EGetMostImportantAssetFlags
 	{
@@ -1101,17 +1199,18 @@ namespace UE::AssetRegistry
 	 * where you need to relate information about a package to an asset.
 	 * 
 	 * Usually there is only 1 asset per package so this is straightforward, however in the multiple asset case it:
-	 *	Tries to find the "UAsset" via the FAssetData::IsUAsset() function. (i.e. asset name matches package name)
-	 *	If none exist, tries to find a "Top Level Asset" using FAssetData::IsToplevelAsset(). (i.e. outer == package)
-	 *		If only one exists, use that.
-	 *		Otherwise, if EGetMostImportantAssetFlags::RequireOneTopLevelAsset isn't set, gather the set of possibles and return the first sorted on asset class then name.
-	 *			If no top level assets, all package assets
-	 *			If multiple top level assets, all top level assets
+	 *  Tries to find the "UAsset" via the FAssetData::IsUAsset() function. (i.e. asset name matches package name)
+	 *  If none exist, tries to find a "Top Level Asset" using FAssetData::IsToplevelAsset(). (i.e. outer == package)
+	 *  If only one exists, use that.
+	 *  Otherwise, if EGetMostImportantAssetFlags::RequireOneTopLevelAsset isn't set, gather the set of possibles and
+	 *  return the first sorted on asset class then name.
+	 *	    If no top level assets, all package assets
+	 *	    If multiple top level assets, all top level assets
 	 * 
 	 * A good source for PackageAssetDatas is FAssetRegistryState::GetAssetsByPackageName.
-	 * 
 	 */
-	ASSETREGISTRY_API const FAssetData* GetMostImportantAsset(TConstArrayView<const FAssetData*> PackageAssetDatas, EGetMostImportantAssetFlags InFlags = EGetMostImportantAssetFlags::None);
+	ASSETREGISTRY_API const FAssetData* GetMostImportantAsset(TConstArrayView<const FAssetData*> PackageAssetDatas,
+		EGetMostImportantAssetFlags InFlags = EGetMostImportantAssetFlags::None);
 
 	/*
 	* Returns true if the asset registry should start searching all assets on startup
@@ -1131,12 +1230,12 @@ namespace UE::AssetRegistry
 	* NOTE some platforms do not compress in UnrealPak. For those platforms, all compressed sizes are actually
 	* uncompressed!
 	* 
-	* Stage_ChunkCountFName						Total iostore chunks in the package.
-	* Stage_ChunkSizeFName						The total uncompressed size of the chunks in the package.
-	* Stage_ChunkCompressedSizeFName			The total compressed size of the chunks in the package. SEE NOTE ABOVE!
-	* Stage_ChunkInstalledSizeFName				The compressed size of all chunks that must be installed with the game/plugin.
-	* Stage_ChunkStreamingSizeFName				The compressed size of all chunks that are delivered using IAS.
-	* Stage_ChunkOptionalSizeFName				The compressed size of all chunks that are placed in an optional container.
+	* Stage_ChunkCountFName				Total iostore chunks in the package.
+	* Stage_ChunkSizeFName				The total uncompressed size of the chunks in the package.
+	* Stage_ChunkCompressedSizeFName	The total compressed size of the chunks in the package. SEE NOTE ABOVE!
+	* Stage_ChunkInstalledSizeFName		The compressed size of all chunks that must be installed with the game/plugin.
+	* Stage_ChunkStreamingSizeFName		The compressed size of all chunks that are delivered using IAS.
+	* Stage_ChunkOptionalSizeFName		The compressed size of all chunks that are placed in an optional container.
 	* 
 	* See CookMetadata.h - UE::Cook::EPluginSizeTypes for more information on size types.
 	*/
