AssetRegistry: Add IAssetRegistry.ContainsDependency function, for use by the cooker when checking whether loaded packages were unsolicited or expected from the load of the referencer package.
Also remove duplication in the two IterateDependencyList helper functions for FDependsNode by pulling the shared EnumerateMatchingDependencyFlags code into a separate function.
#rnx
#incrementalcook

diff --git a/Engine/Source/Runtime/AssetRegistry/Private/AssetRegistry.cpp b/Engine/Source/Runtime/AssetRegistry/Private/AssetRegistry.cpp
--- a/Engine/Source/Runtime/AssetRegistry/Private/AssetRegistry.cpp
+++ b/Engine/Source/Runtime/AssetRegistry/Private/AssetRegistry.cpp
@@ -3190,6 +3190,14 @@ bool UAssetRegistryImpl::GetDependencies(FName PackageName, TArray<FName>& OutDe
 	return true;
 }
 
+bool UAssetRegistryImpl::ContainsDependency(FName PackageName, FName QueryDependencyName,
+	UE::AssetRegistry::EDependencyCategory Category, const UE::AssetRegistry::FDependencyQuery& Flags) const
+{
+	UE::AssetRegistry::FInterfaceReadScopeLock InterfaceScopeLock(InterfaceLock);
+	return GuardedData.GetState().ContainsDependency(FAssetIdentifier(PackageName),
+		FAssetIdentifier(QueryDependencyName), Category, Flags);
+}
+
 bool IAssetRegistry::K2_GetDependencies(FName PackageName, const FAssetRegistryDependencyOptions& DependencyOptions, TArray<FName>& OutDependencies) const
 {
 	UE::AssetRegistry::FDependencyQuery Flags;
diff --git a/Engine/Source/Runtime/AssetRegistry/Private/AssetRegistry.h b/Engine/Source/Runtime/AssetRegistry/Private/AssetRegistry.h
--- a/Engine/Source/Runtime/AssetRegistry/Private/AssetRegistry.h
+++ b/Engine/Source/Runtime/AssetRegistry/Private/AssetRegistry.h
@@ -61,6 +61,7 @@ public:
 	virtual bool GetDependencies(const FAssetIdentifier& AssetIdentifier, TArray<FAssetIdentifier>& OutDependencies, UE::AssetRegistry::EDependencyCategory Category = UE::AssetRegistry::EDependencyCategory::All, const UE::AssetRegistry::FDependencyQuery& Flags = UE::AssetRegistry::FDependencyQuery()) const override;
 	virtual bool GetDependencies(const FAssetIdentifier& AssetIdentifier, TArray<FAssetDependency>& OutDependencies, UE::AssetRegistry::EDependencyCategory Category = UE::AssetRegistry::EDependencyCategory::All, const UE::AssetRegistry::FDependencyQuery& Flags = UE::AssetRegistry::FDependencyQuery()) const override;
 	virtual bool GetDependencies(FName PackageName, TArray<FName>& OutDependencies, UE::AssetRegistry::EDependencyCategory Category = UE::AssetRegistry::EDependencyCategory::Package, const UE::AssetRegistry::FDependencyQuery& Flags = UE::AssetRegistry::FDependencyQuery()) const override; //-V1101
+	virtual bool ContainsDependency(FName PackageName, FName QueryDependencyName, UE::AssetRegistry::EDependencyCategory Category = UE::AssetRegistry::EDependencyCategory::Package, const UE::AssetRegistry::FDependencyQuery& Flags = UE::AssetRegistry::FDependencyQuery()) const override;
 	virtual bool GetReferencers(const FAssetIdentifier& AssetIdentifier, TArray<FAssetIdentifier>& OutReferencers, UE::AssetRegistry::EDependencyCategory Category = UE::AssetRegistry::EDependencyCategory::All, const UE::AssetRegistry::FDependencyQuery& Flags = UE::AssetRegistry::FDependencyQuery()) const override;
 	virtual bool GetReferencers(const FAssetIdentifier& AssetIdentifier, TArray<FAssetDependency>& OutReferencers, UE::AssetRegistry::EDependencyCategory Category = UE::AssetRegistry::EDependencyCategory::All, const UE::AssetRegistry::FDependencyQuery& Flags = UE::AssetRegistry::FDependencyQuery()) const override;
 	virtual bool GetReferencers(FName PackageName, TArray<FName>& OutReferencers, UE::AssetRegistry::EDependencyCategory Category = UE::AssetRegistry::EDependencyCategory::Package, const UE::AssetRegistry::FDependencyQuery& Flags = UE::AssetRegistry::FDependencyQuery()) const override; //-V1101
diff --git a/Engine/Source/Runtime/AssetRegistry/Private/AssetRegistryState.cpp b/Engine/Source/Runtime/AssetRegistry/Private/AssetRegistryState.cpp
--- a/Engine/Source/Runtime/AssetRegistry/Private/AssetRegistryState.cpp
+++ b/Engine/Source/Runtime/AssetRegistry/Private/AssetRegistryState.cpp
@@ -1554,6 +1554,20 @@ bool FAssetRegistryState::GetDependencies(const FAssetIdentifier& AssetIdentifie
 	}
 }
 
+bool FAssetRegistryState::ContainsDependency(const FAssetIdentifier& AssetIdentifier,
+	const FAssetIdentifier& QueryAsset,
+	UE::AssetRegistry::EDependencyCategory Category, const UE::AssetRegistry::FDependencyQuery& Flags) const
+{
+	const FDependsNode* const* NodePtr = CachedDependsNodes.Find(AssetIdentifier);
+	const FDependsNode* const* QueryNodePtr = CachedDependsNodes.Find(QueryAsset);
+	if (!NodePtr || !QueryNodePtr)
+	{
+		return false;
+	}
+
+	return (*NodePtr)->ContainsDependency(*QueryNodePtr, Category, Flags);
+}
+
 bool FAssetRegistryState::GetReferencers(const FAssetIdentifier& AssetIdentifier,
 	TArray<FAssetIdentifier>& OutReferencers,
 	UE::AssetRegistry::EDependencyCategory Category, const UE::AssetRegistry::FDependencyQuery& Flags) const
diff --git a/Engine/Source/Runtime/AssetRegistry/Private/DependsNode.cpp b/Engine/Source/Runtime/AssetRegistry/Private/DependsNode.cpp
--- a/Engine/Source/Runtime/AssetRegistry/Private/DependsNode.cpp
+++ b/Engine/Source/Runtime/AssetRegistry/Private/DependsNode.cpp
@@ -28,8 +28,71 @@ void FDependsNode::PrintReferencers() const
 	PrintReferencersRecursive(TEXT(""), VisitedNodes);
 }
 
-template <uint32 FlagWidth>
-void IterateDependencyList(const FDependsNode::FIterateDependenciesCallback& InCallback,
+/**
+ * Helper function for IterateDependencyList, which itself is a helper function for FDependsNode functions.
+ * This function is given an index into one of the lists. All of the data for the list and how to interpret it is also
+ * passed in.
+ * The function looks of the set of DependencyProperty flag groups for the index, and iterates over that set, and
+ * calls the callback on each flag group that matches the requested SearchFlags.
+ */
+template <uint32 FlagWidth, typename VisitorType>
+void EnumerateMatchingDependencyFlags(const TBitArray<>* FlagBits, UE::AssetRegistry::EDependencyProperty CategoryMask,
+	UE::AssetRegistry::EDependencyProperty(*ByteToProperties)(uint8),
+	int32 ListIndex, const UE::AssetRegistry::FDependencyQuery& SearchFlags, VisitorType&& Visitor)
+{
+	using namespace UE::AssetRegistry;
+
+	typedef TPropertyCombinationSet<FlagWidth> FCombinationSet;
+	constexpr uint32 FlagSetWidth = FCombinationSet::StorageBitCount;
+
+	EDependencyProperty RequiredProperties = SearchFlags.Required & CategoryMask;
+	EDependencyProperty ExcludedProperties = SearchFlags.Excluded & CategoryMask;
+
+	FCombinationSet DependsNodeFlagsSet(*FlagBits, ListIndex * FlagSetWidth);
+	bool bDuplicate = false;
+	for (uint32 DependencyFlagBits : DependsNodeFlagsSet)
+	{
+		EDependencyProperty DependencyProperties = ByteToProperties(static_cast<uint8>(DependencyFlagBits));
+		if (!EnumHasAllFlags(DependencyProperties, RequiredProperties))
+		{
+			continue;
+		}
+		if (EnumHasAnyFlags(DependencyProperties, ExcludedProperties))
+		{
+			continue;
+		}
+		bool bPassesRequiredUnions = true;
+		for (EDependencyProperty RequiredUnion : SearchFlags.RequiredUnions)
+		{
+			EDependencyProperty RequiredUnionProperty = RequiredUnion & CategoryMask;
+			if (RequiredUnionProperty != EDependencyProperty::None &&
+				!EnumHasAnyFlags(DependencyProperties, RequiredUnionProperty))
+			{
+				bPassesRequiredUnions = false;
+				break;
+			}
+		}
+		if (!bPassesRequiredUnions)
+		{
+			continue;
+		}
+
+		Visitor(DependencyProperties, bDuplicate);
+		bDuplicate = true;
+	}
+}
+
+/**
+ * Helper function for GetDependencies, IterateOverDependencies, and others that need to find dependencies in
+ * each of an FDependsNode's list of dependencies, while filtering by the list of DependencyProperties associated with
+ * each dependency.
+ * This function is given all of the data for the list and how to interpret it, and a requested SearchFlags.
+ * The function iterates over all dependencies in the list, and iterates over the set of DependencyProperty flag
+ * groups for each dependency, and calls the callback with the dependency and with each flag group on the dependency
+ * that matches the requested SearchFlags.
+ */
+template <uint32 FlagWidth, typename CallbackType>
+void IterateDependencyList(CallbackType&& InCallback,
 	UE::AssetRegistry::EDependencyCategory SearchCategory,
 	const UE::AssetRegistry::FDependencyQuery& SearchFlags,
 	UE::AssetRegistry::EDependencyCategory ListCategory,
@@ -54,51 +117,31 @@ void IterateDependencyList(const FDependsNode::FIterateDependenciesCallback& InC
 	}
 	else
 	{
-		EDependencyProperty RequiredProperties = SearchFlags.Required & CategoryMask;
-		EDependencyProperty ExcludedProperties = SearchFlags.Excluded & CategoryMask;
-		typedef TPropertyCombinationSet<FlagWidth> FCombinationSet;
-		constexpr uint32 FlagSetWidth = FCombinationSet::StorageBitCount;
 		for (int32 ListIndex = 0; ListIndex < Dependencies.Num(); ++ListIndex)
 		{
 			FDependsNode* DependsNode = Dependencies[ListIndex];
-			FCombinationSet DependsNodeFlagsSet(*FlagBits, ListIndex * FlagSetWidth);
-			bool bDuplicate = false;
-			for (uint32 DependencyFlagBits : DependsNodeFlagsSet)
-			{
-				EDependencyProperty DependencyProperties = ByteToProperties(static_cast<uint8>(DependencyFlagBits));
-				if (!EnumHasAllFlags(DependencyProperties, RequiredProperties))
-				{
-					continue;
-				}
-				if (EnumHasAnyFlags(DependencyProperties, ExcludedProperties))
+			EnumerateMatchingDependencyFlags<FlagWidth>(
+				FlagBits, CategoryMask, ByteToProperties, ListIndex, SearchFlags,
+				[&InCallback, DependsNode, ListCategory](EDependencyProperty MatchingDependencyProperties, bool bDuplicate)
 				{
-					continue;
-				}
-				bool bPassesRequiredUnions = true;
-				for (EDependencyProperty RequiredUnion : SearchFlags.RequiredUnions)
-				{
-					EDependencyProperty RequiredUnionProperty = RequiredUnion & CategoryMask;
-					if (RequiredUnionProperty != EDependencyProperty::None &&
-						!EnumHasAnyFlags(DependencyProperties, RequiredUnionProperty))
-					{
-						bPassesRequiredUnions = false;
-						break;
-					}
-				}
-				if (!bPassesRequiredUnions)
-				{
-					continue;
-				}
-
-				InCallback(DependsNode, ListCategory, DependencyProperties, bDuplicate);
-				bDuplicate = true;
-			}
+					InCallback(DependsNode, ListCategory, MatchingDependencyProperties, bDuplicate);
+				});
 		}
 	}
 }
 
-template <uint32 FlagWidth>
-void IterateDependencyList(const FDependsNode::FIterateDependenciesCallback& InCallback,
+/**
+ * Helper function for IterateOverDependencies, ContainsDependency, and others that need to find a given dependency in
+ * one of an FDependsNode's list of dependencies, while filtering by the list of DependencyProperties associated with
+ * each dependency.
+ * This function is given all of the data for the list and how to interpret it, and a requested SearchNode and
+ * SearchFlags.
+ * The function looks up the SearchNode in the list of dependencies, and iterates over the set of DependencyProperty flag
+ * groups for each dependency, and calls the callback with each flag group on the dependency that matches the requested
+ * SearchFlags.
+ */
+template <uint32 FlagWidth, typename CallbackType>
+void IterateDependencyList(CallbackType&& InCallback,
 	const FDependsNode* SearchNode,
 	UE::AssetRegistry::EDependencyCategory SearchCategory,
 	const UE::AssetRegistry::FDependencyQuery& SearchFlags,
@@ -123,42 +166,13 @@ void IterateDependencyList(const FDependsNode::FIterateDependenciesCallback& InC
 		}
 		else
 		{
-			EDependencyProperty RequiredProperties = SearchFlags.Required & CategoryMask;
-			EDependencyProperty ExcludedProperties = SearchFlags.Excluded & CategoryMask;
-			typedef TPropertyCombinationSet<FlagWidth> FCombinationSet;
-			constexpr uint32 FlagSetWidth = FCombinationSet::StorageBitCount;
 			FDependsNode* DependsNode = Dependencies[ListIndex];
-			FCombinationSet DependsNodeFlagsSet(*FlagBits, ListIndex * FlagSetWidth);
-			bool bDuplicate = false;
-			for (uint32 DependencyFlagBits : DependsNodeFlagsSet)
-			{
-				EDependencyProperty DependencyProperties = ByteToProperties(static_cast<uint8>(DependencyFlagBits));
-				if (!EnumHasAllFlags(DependencyProperties, RequiredProperties))
-				{
-					continue;
-				}
-				if (EnumHasAnyFlags(DependencyProperties, ExcludedProperties))
-				{
-					continue;
-				}
-				bool bPassesRequiredUnions = true;
-				for (EDependencyProperty RequiredUnion : SearchFlags.RequiredUnions)
+			EnumerateMatchingDependencyFlags<FlagWidth>(
+				FlagBits, CategoryMask, ByteToProperties, ListIndex, SearchFlags,
+				[&InCallback, DependsNode, ListCategory](EDependencyProperty MatchingDependencyProperties, bool bDuplicate)
 				{
-					EDependencyProperty RequiredUnionProperty = RequiredUnion & CategoryMask;
-					if (RequiredUnionProperty != EDependencyProperty::None &&
-						!EnumHasAnyFlags(DependencyProperties, RequiredUnionProperty))
-					{
-						bPassesRequiredUnions = false;
-						break;
-					}
-				}
-				if (!bPassesRequiredUnions)
-				{
-					continue;
-				}
-				InCallback(DependsNode, ListCategory, DependencyProperties, bDuplicate);
-				bDuplicate = true;
-			}
+					InCallback(DependsNode, ListCategory, MatchingDependencyProperties, bDuplicate);
+				});
 		}
 	};
 
@@ -632,31 +646,56 @@ void FDependsNode::RemoveLinks(const TUniqueFunction<bool(const FDependsNode*)>&
 }
 
 
-bool FDependsNode::ContainsDependency(FDependsNode* InDependency,
-	UE::AssetRegistry::EDependencyCategory Category) const
+bool FDependsNode::ContainsDependency(const FDependsNode* InDependency,
+	UE::AssetRegistry::EDependencyCategory Category, const UE::AssetRegistry::FDependencyQuery& Flags) const
 {
 	using namespace UE::AssetRegistry;
-	auto ListContains = [InDependency](const TArray<FDependsNode*>& List, bool IsSorted)
-	{
-		return (IsSorted ? Algo::BinarySearch(List, InDependency) : List.Find(InDependency)) != INDEX_NONE;
-	};
+
+	bool bExists = false;
 	if (!!(Category & EDependencyCategory::Package))
 	{
-		if (ListContains(PackageDependencies, PackageIsSorted))
+		IterateDependencyList<PackageFlagWidth>([&bExists]
+		(const FDependsNode* DependsNode, EDependencyCategory MatchingCategory,
+			EDependencyProperty MatchingDependencyProperties, bool bDuplicate)
+			{
+				bExists = true;
+			},
+			InDependency, Category, Flags, EDependencyCategory::Package,
+			EDependencyProperty::PackageMask, PackageDependencies, &PackageFlags,
+			ByteToPackageProperties, PackageIsSorted);
+		if (bExists)
 		{
 			return true;
 		}
 	}
 	if (!!(Category & EDependencyCategory::SearchableName))
 	{
-		if (ListContains(NameDependencies, SearchableNameIsSorted))
+		IterateDependencyList<SearchableNameFlagWidth>([&bExists]
+		(const FDependsNode* DependsNode, EDependencyCategory MatchingCategory,
+			EDependencyProperty MatchingDependencyProperties, bool bDuplicate)
+			{
+				bExists = true;
+			},
+			InDependency, Category, Flags, EDependencyCategory::SearchableName,
+			EDependencyProperty::SearchableNameMask, NameDependencies, nullptr,
+			nullptr, SearchableNameIsSorted);
+		if (bExists)
 		{
 			return true;
 		}
 	}
 	if (!!(Category & EDependencyCategory::Manage))
 	{
-		if (ListContains(ManageDependencies, ManageIsSorted))
+		IterateDependencyList<ManageFlagWidth>([&bExists]
+		(const FDependsNode* DependsNode, EDependencyCategory MatchingCategory,
+			EDependencyProperty MatchingDependencyProperties, bool bDuplicate)
+			{
+				bExists = true;
+			},
+			InDependency, Category, Flags, EDependencyCategory::Manage,
+			EDependencyProperty::ManageMask, ManageDependencies, &ManageFlags,
+			ByteToManageProperties, ManageIsSorted);
+		if (bExists)
 		{
 			return true;
 		}
diff --git a/Engine/Source/Runtime/AssetRegistry/Private/DependsNode.h b/Engine/Source/Runtime/AssetRegistry/Private/DependsNode.h
--- a/Engine/Source/Runtime/AssetRegistry/Private/DependsNode.h
+++ b/Engine/Source/Runtime/AssetRegistry/Private/DependsNode.h
@@ -72,8 +72,9 @@ public:
 	void RemoveReferencer(FDependsNode* InReferencer);
 	/** Removes any referencers that no longer have this node as a dependency */
 	void RefreshReferencers();
-	bool ContainsDependency(FDependsNode* InDependency,
-		UE::AssetRegistry::EDependencyCategory Category = UE::AssetRegistry::EDependencyCategory::All) const;
+	bool ContainsDependency(const FDependsNode* InDependency,
+		UE::AssetRegistry::EDependencyCategory Category = UE::AssetRegistry::EDependencyCategory::All,
+		const UE::AssetRegistry::FDependencyQuery& Flags = UE::AssetRegistry::FDependencyQuery()) const;
 	/** Clear all dependency records from this node */
 	void ClearDependencies(UE::AssetRegistry::EDependencyCategory Category = UE::AssetRegistry::EDependencyCategory::All);
 	void ClearReferencers();
diff --git a/Engine/Source/Runtime/AssetRegistry/Public/AssetRegistry/AssetRegistryState.h b/Engine/Source/Runtime/AssetRegistry/Public/AssetRegistry/AssetRegistryState.h
--- a/Engine/Source/Runtime/AssetRegistry/Public/AssetRegistry/AssetRegistryState.h
+++ b/Engine/Source/Runtime/AssetRegistry/Public/AssetRegistry/AssetRegistryState.h
@@ -305,7 +305,6 @@ public:
 
 	/**
 	 * Appends a list of packages and searchable names that are referenced by the supplied package or name.
-	 * (On disk references ONLY)
 	 *
 	 * @param AssetIdentifier	the name of the package/name for which to gather dependencies
 	 * @param OutDependencies	a list of things that are referenced by AssetIdentifier
@@ -324,9 +323,24 @@ public:
 		UE::AssetRegistry::EDependencyCategory Category = UE::AssetRegistry::EDependencyCategory::All,
 		const UE::AssetRegistry::FDependencyQuery& Flags = UE::AssetRegistry::FDependencyQuery()) const;
 
+	/**
+	 * Return whether the supplied package references the query package.
+	 *
+	 * @param PackageName The name of the package for which to inspect dependencies (eg, /Game/MyFolder/MyAsset)
+	 * @param QueryAsset Which package to look for in PackageName's dependencies
+	 * @param Category Which category(ies) of dependencies to look in for QueryAsset. Dependencies matching
+	 *        ANY of the OR'd categories will be considered.
+	 * @param Flags Which flags are required present or not present on the dependencies. Dependencies matching ALL
+	 *        required and NONE excluded bits will be considered. For each potentially considered dependency, flags not
+	 *        applicable to their category are ignored.
+	 */
+	ASSETREGISTRY_API bool ContainsDependency(const FAssetIdentifier& AssetIdentifier,
+		const FAssetIdentifier& QueryAsset,
+		UE::AssetRegistry::EDependencyCategory Category = UE::AssetRegistry::EDependencyCategory::All,
+		const UE::AssetRegistry::FDependencyQuery& Flags = UE::AssetRegistry::FDependencyQuery()) const;
+
 	/**
 	 * Appends a list of packages and searchable names that reference the supplied package or name.
-	 * (On disk references ONLY)
 	 *
 	 * @param AssetIdentifier	the name of the package/name for which to gather dependencies
 	 * @param OutReferencers	a list of things that reference AssetIdentifier
diff --git a/Engine/Source/Runtime/AssetRegistry/Public/AssetRegistry/IAssetRegistry.h b/Engine/Source/Runtime/AssetRegistry/Public/AssetRegistry/IAssetRegistry.h
--- a/Engine/Source/Runtime/AssetRegistry/Public/AssetRegistry/IAssetRegistry.h
+++ b/Engine/Source/Runtime/AssetRegistry/Public/AssetRegistry/IAssetRegistry.h
@@ -461,7 +461,8 @@ public:
 		const UE::AssetRegistry::FDependencyQuery& Flags = UE::AssetRegistry::FDependencyQuery()) const = 0;
 
 	/**
-	 * Gets a list of PackageNames that are referenced by the supplied package. (On disk references ONLY)
+	 * Gets a list of PackageNames that are referenced by the supplied package.
+	 * Only returns dependencies reported in the on-disk package.
 	 *
 	 * @param PackageName The name of the package for which to gather dependencies (eg, /Game/MyFolder/MyAsset)
 	 * @param OutDependencies A list of packages that are referenced by the package whose path is PackageName
@@ -475,6 +476,22 @@ public:
 		UE::AssetRegistry::EDependencyCategory Category = UE::AssetRegistry::EDependencyCategory::Package,
 		const UE::AssetRegistry::FDependencyQuery& Flags = UE::AssetRegistry::FDependencyQuery()) const = 0;
 
+	/**
+	 * Return whether the supplied package references the query package. Only considers dependencies reported in the
+	 * on-disk package.
+	 *
+	 * @param PackageName The name of the package for which to inspect dependencies (eg, /Game/MyFolder/MyAsset)
+	 * @param QueryDependencyName Which package to look for in PackageName's dependencies
+	 * @param Category Which category(ies) of dependencies to look in for QueryDependencyName. Dependencies matching
+	 *        ANY of the OR'd categories will be considered.
+	 * @param Flags Which flags are required present or not present on the dependencies. Dependencies matching ALL
+	 *        required and NONE excluded bits will be considered. For each potentially considered dependency, flags not
+	 *        applicable to their category are ignored.
+	 */
+	virtual bool ContainsDependency(FName PackageName, FName QueryDependencyName,
+		UE::AssetRegistry::EDependencyCategory Category = UE::AssetRegistry::EDependencyCategory::Package,
+		const UE::AssetRegistry::FDependencyQuery& Flags = UE::AssetRegistry::FDependencyQuery()) const = 0;
+
 	/**
 	 * Gets a list of paths to objects that are referenced by the supplied package. (On disk references ONLY)
 	 *
