AssetRegistry: Fix performance of EnumerateAsset calls when the ARFilter has TagsAndValues: FilterAssets was not correctly calculating the cost of TArray copies and was using a suboptimal intersection algorithm when a large number of assets in the depot had the requested tag.
#rnx

diff --git a/Engine/Source/Runtime/AssetRegistry/Private/AssetRegistryState.cpp b/Engine/Source/Runtime/AssetRegistry/Private/AssetRegistryState.cpp
--- a/Engine/Source/Runtime/AssetRegistry/Private/AssetRegistryState.cpp
+++ b/Engine/Source/Runtime/AssetRegistry/Private/AssetRegistryState.cpp
@@ -1034,19 +1034,19 @@ void FilterAssets(TArray<const FAssetData*>&InOutResults,
 	}
 }
 
-bool AssetDataMatchesTag(const FAssetData* AssetData, const TPair<FName, TOptional<FString>>& TagPair)
+bool AssetDataMatchesTag(const FAssetData* AssetData, FName TagName, const TOptional<FString>& TagValue)
 {
 	if (!AssetData)
 	{
 		return false;
 	}
-	if (!TagPair.Value.IsSet())
+	if (!TagValue.IsSet())
 	{
-		return AssetData->TagsAndValues.Contains(TagPair.Key);
+		return AssetData->TagsAndValues.Contains(TagName);
 	}
 	else
 	{
-		return AssetData->TagsAndValues.ContainsKeyValue(TagPair.Key, TagPair.Value.GetValue());
+		return AssetData->TagsAndValues.ContainsKeyValue(TagName, TagValue.GetValue());
 	}
 }
 
@@ -1055,48 +1055,55 @@ void FilterAssets(TArray<const FAssetData*>& InOutResults, const AccelerationMap
 	const TMultiMap<FName, TOptional<FString>>& TagsAndValues,
 	CallbackType&& FunctionToKeepAsset, int32 FilterComplexity, const FAssetDataMap& CachedAssets)
 {
-	TArray<TArray<const FAssetData*>, TInlineAllocator<10>> Matches;
+	// AccelerationMapValueTypePtr is either TArray<FAssetData*>* or TArray<FAssetDataPtrIndex>*
+	using AccelerationMapValueTypePtr = decltype(AccelerationMap.Find(FName()));
+	struct FMatchData
+	{
+		FName TagName;
+		const TOptional<FString>* TagValuePtr;
+		AccelerationMapValueTypePtr AssetsWithTag;
+	};
+
+	TArray<FMatchData, TInlineAllocator<10>> Matches;
 	Matches.Reserve(TagsAndValues.Num());
-	uint32 TotalMatches = 0;
+	uint32 EstimateOfTotalMatches = 0;
 
 	for (const TPair<FName, TOptional<FString>>& TagPair : TagsAndValues)
 	{
-		TArray<const FAssetData*>& Results = Matches.Emplace_GetRef();
-		if (const auto* TagAssets = AccelerationMap.Find(TagPair.Key))
+		AccelerationMapValueTypePtr AssetsWithTag = AccelerationMap.Find(TagPair.Key);
+		if (AssetsWithTag)
+		{
+			Matches.Add({ TagPair.Key, &TagPair.Value, AssetsWithTag });
+			EstimateOfTotalMatches += AssetsWithTag->Num();
+		}
+	}
+
+	if (InOutResults.IsEmpty())
+	{
+		// No previous Results; set Results equal to the values found in AccelerationMap
+		for (FMatchData& MatchData : Matches)
 		{
-			Results.Reserve(TagAssets->Num());
 #if !UE_ASSETREGISTRY_INDIRECT_ASSETDATA_POINTERS
-			for (FAssetData* AssetData : *TagAssets)
+			for (FAssetData* AssetData : *MatchData.AssetsWithTag)
 			{
 #else
-			for (FAssetDataPtrIndex AssetIndex : *TagAssets)
+			for (FAssetDataPtrIndex AssetIndex : *MatchData.AssetsWithTag)
 			{
 				FAssetData* AssetData = CachedAssets[AssetIndex];
 #endif
-				if (AssetDataMatchesTag(AssetData, TagPair))
+				if (AssetDataMatchesTag(AssetData, MatchData.TagName, *MatchData.TagValuePtr))
 				{
-					Results.Add(AssetData);
+					InOutResults.Add(AssetData);
 				}
 			}
 		}
-		TotalMatches += Results.Num();
-	}
-
-	if (InOutResults.IsEmpty())
-	{
-		// No previous Results; set Results equal to the values found in AccelerationMap
-		InOutResults.Reserve(TotalMatches);
-		for (TArray<const FAssetData*>& Assets : Matches)
-		{
-			InOutResults.Append(Assets);
-		}
 		// Remove duplicates
 		Algo::Sort(InOutResults);
 		InOutResults.SetNum(Algo::Unique(InOutResults));
 	}
 	else
 	{
-		bool bUseFiltering = DecideIntersectionMethod(InOutResults.Num(), TotalMatches, FilterComplexity);
+		bool bUseFiltering = DecideIntersectionMethod(InOutResults.Num(), EstimateOfTotalMatches, FilterComplexity);
 		if (bUseFiltering)
 		{
 			InOutResults.RemoveAllSwap([&FunctionToKeepAsset](const FAssetData* AssetData)
@@ -1106,10 +1113,34 @@ void FilterAssets(TArray<const FAssetData*>& InOutResults, const AccelerationMap
 		}
 		else
 		{
+			TArray<TArray<const FAssetData*>, TInlineAllocator<10>> MatchArrays;
+			MatchArrays.Reserve(Matches.Num());
+
+			int32 TotalMatches = 0;
+			for (FMatchData& MatchData : Matches)
+			{
+				TArray<const FAssetData*>& MatchArray = MatchArrays.Emplace_GetRef();
+				MatchArray.Reserve(MatchData.AssetsWithTag->Num());
+#if !UE_ASSETREGISTRY_INDIRECT_ASSETDATA_POINTERS
+				for (FAssetData* AssetData : *MatchData.AssetsWithTag)
+				{
+#else
+				for (FAssetDataPtrIndex AssetIndex : *MatchData.AssetsWithTag)
+				{
+					FAssetData* AssetData = CachedAssets[AssetIndex];
+#endif
+					if (AssetDataMatchesTag(AssetData, MatchData.TagName, *MatchData.TagValuePtr))
+					{
+						MatchArray.Add(AssetData);
+						++TotalMatches;
+					}
+				}
+			}
+
 			// Convert Array of Arrays into format required by ArrayIntersection: Array of ArrayViews
 			TArray<TConstArrayView<const FAssetData*>, TInlineAllocator<10>> ArrayViewMatches;
-			ArrayViewMatches.Reserve(Matches.Num());
-			for (TArray<const FAssetData*>& MatchesElement : Matches)
+			ArrayViewMatches.Reserve(MatchArrays.Num());
+			for (TArray<const FAssetData*>& MatchesElement : MatchArrays)
 			{
 				ArrayViewMatches.Emplace(MatchesElement);
 			}
@@ -1169,7 +1200,7 @@ void FilterAssetsByCachedClassesByTag(TArray<const FAssetData*>& InOutResults,
 		for (TArray<const FAssetData*>::TIterator Iter(Results); Iter; ++Iter)
 		{
 			const FAssetData* AssetData = *Iter;
-			if (!AssetDataMatchesTag(AssetData, TagPair))
+			if (!AssetDataMatchesTag(AssetData, TagPair.Key, TagPair.Value))
 			{
 				Iter.RemoveCurrentSwap();
 			}
@@ -1338,7 +1369,7 @@ bool FAssetRegistryState::EnumerateAssets(const FARCompiledFilter& Filter, const
 			{
 				for (const TPair<FName, TOptional<FString>>& TagPair : Filter.TagsAndValues)
 				{
-					if (AssetDataMatchesTag(AssetData, TagPair))
+					if (AssetDataMatchesTag(AssetData, TagPair.Key, TagPair.Value))
 					{
 						return true; // keep
 					}
