AssetRegistry - AssetDataGatherer - Optimize ScanPathsSynchronous case when the path has already been scanned by avoiding reading IO unnecessarily in the AssetDataDiscovery tree.
#rnx

diff --git a/Engine/Source/Runtime/AssetRegistry/Private/AssetDataGatherer.cpp b/Engine/Source/Runtime/AssetRegistry/Private/AssetDataGatherer.cpp
--- a/Engine/Source/Runtime/AssetRegistry/Private/AssetDataGatherer.cpp
+++ b/Engine/Source/Runtime/AssetRegistry/Private/AssetDataGatherer.cpp
@@ -35,6 +35,7 @@
 #include "Misc/ScopeExit.h"
 #include "Misc/TrackedActivity.h"
 #include "Serialization/Archive.h"
+#include "String/Find.h"
 #include "String/LexFromString.h"
 #include "Tasks/Task.h"
 #include "TelemetryRouter.h"
@@ -743,10 +744,12 @@ FScanDir* FScanDir::GetControllingDir(FStringView InRelPath, bool bIsDirectory,
 	}
 }
 
-bool FScanDir::TrySetDirectoryProperties(FStringView InRelPath, FInherited& ParentData,
-	const FSetPathProperties& InProperties, bool bConfirmedExists, FScanDirAndParentData& OutControllingDir)
+void FScanDir::TrySetDirectoryProperties(FPathExistence& QueryPath, FStringView InRelPath, FInherited& ParentData,
+	const FSetPathProperties& InProperties, FScanDirAndParentData& OutControllingDir,
+	FStringView& OutControllingDirRelPath, bool& bInOutMadeChanges)
 {
-	// TrySetDirectoryProperties can only be called on valid ScanDirs, which we rely on so we can call FindOrAddSubDir which requires that
+	// TrySetDirectoryProperties can only be called on valid ScanDirs, which we rely on because we call FindOrAddSubDir
+	// which relies on that.
 	check(IsValid()); 
 
 	if (InRelPath.IsEmpty())
@@ -754,6 +757,7 @@ bool FScanDir::TrySetDirectoryProperties(FStringView InRelPath, FInherited& Pare
 		// The properties apply to this entire directory
 		if (InProperties.IsOnAllowList.IsSet() && DirectData.bIsOnAllowList != *InProperties.IsOnAllowList)
 		{
+			bInOutMadeChanges = true;
 			SetComplete(false);
 			if (bScanInFlight)
 			{
@@ -787,6 +791,7 @@ bool FScanDir::TrySetDirectoryProperties(FStringView InRelPath, FInherited& Pare
 		if ((InProperties.MatchesDenyList.IsSet() && DirectData.bMatchesDenyList != *InProperties.MatchesDenyList) ||
 			(InProperties.IgnoreDenyList.IsSet() && DirectData.bIgnoreDenyList != *InProperties.IgnoreDenyList))
 		{
+			bInOutMadeChanges = true;
 			SetComplete(false);
 			if (InProperties.MatchesDenyList.IsSet())
 			{
@@ -822,6 +827,7 @@ bool FScanDir::TrySetDirectoryProperties(FStringView InRelPath, FInherited& Pare
 		}
 		if (InProperties.HasScanned.IsSet())
 		{
+			bInOutMadeChanges = true;
 			SetComplete(false);
 			bool bNewValue = *InProperties.HasScanned;
 			auto SetProperties = [bNewValue](FScanDir& ScanDir)
@@ -839,7 +845,8 @@ bool FScanDir::TrySetDirectoryProperties(FStringView InRelPath, FInherited& Pare
 
 		OutControllingDir.ScanDir = this;
 		OutControllingDir.ParentData = ParentData;
-		return true;
+		OutControllingDirRelPath = FStringView();
+		return;
 	}
 	else
 	{
@@ -850,12 +857,6 @@ bool FScanDir::TrySetDirectoryProperties(FStringView InRelPath, FInherited& Pare
 			// If this directory is set to be monitored, all Monitored flags of its children are unused, are guaranteed set to false, and should not be changed
 			ModifiedProperties = *Properties;
 			ModifiedProperties->IsOnAllowList.Reset();
-			if (!ModifiedProperties->IsSet())
-			{
-				OutControllingDir.ScanDir = this;
-				OutControllingDir.ParentData = ParentData;
-				return false;
-			}
 			Properties = &ModifiedProperties.GetValue();
 		}
 
@@ -863,22 +864,29 @@ bool FScanDir::TrySetDirectoryProperties(FStringView InRelPath, FInherited& Pare
 		FStringView Remainder;
 		FPathViews::SplitFirstComponent(InRelPath, FirstComponent, Remainder);
 
+		FInherited ParentDataForSubDir(ParentData, DirectData);
+
 		FScanDir* SubDir = nullptr;
 		if (bHasScanned &&
 			(!Properties->HasScanned.IsSet() || *Properties->HasScanned == true) &&
-			(!Properties->IsOnAllowList.IsSet()) &&
-			(!Properties->IgnoreDenyList.IsSet()) &&
-			(!Properties->MatchesDenyList.IsSet()))
-		{
-			// If this parent directory has already been scanned and we are not changing the target directory's has-been-scanned value,
-			// and the next child subdirectory does not exist, then the child directory has already been scanned and we do not need to set the properties on it.
-			// Therefore call FindSubDir instead of FindOrAddSubDir, and abort if we do not find it
+			(!Properties->IsOnAllowList.IsSet() || *Properties->IsOnAllowList == ParentDataForSubDir.bIsOnAllowList) &&
+			(!Properties->IgnoreDenyList.IsSet() || *Properties->IgnoreDenyList == ParentDataForSubDir.bIgnoreDenyList) &&
+			(!Properties->MatchesDenyList.IsSet() || *Properties->MatchesDenyList == ParentDataForSubDir.bMatchesDenyList)
+			)
+		{
+			// If this parent directory has already been scanned and we are not changing any values on the target
+			// path to a different value than the current directory, and the next child subdirectory is not recorded
+			// on *this, then one of these is true and we can early exit and report *this as the ControllingDir.
+			//     * The QueryPath is a file path instead of a directory and we don't need to take any action for it
+			//     * The next child directory towards the QueryPath has already been completed and we do not need
+			//       to set any properties on it.
 			SubDir = FindSubDir(FirstComponent);
 			if (!SubDir)
 			{
 				OutControllingDir.ScanDir = this;
 				OutControllingDir.ParentData = ParentData;
-				return false;
+				OutControllingDirRelPath = InRelPath;
+				return;
 			}
 		}
 		else
@@ -886,22 +894,33 @@ bool FScanDir::TrySetDirectoryProperties(FStringView InRelPath, FInherited& Pare
 			SubDir = FindSubDir(FirstComponent);
 			if (!SubDir)
 			{
-				if (!bConfirmedExists)
+				if (!QueryPath.HasExistenceData())
 				{
-					TStringBuilder<256> LocalAbsPath;
-					AppendLocalAbsPath(LocalAbsPath);
-					FPathViews::AppendPath(LocalAbsPath, InRelPath);
-					FFileStatData StatData = IFileManager::Get().GetStatData(LocalAbsPath.ToString());
-					if (!StatData.bIsValid || !StatData.bIsDirectory)
-					{
-						UE_LOG(LogAssetRegistry, Warning, TEXT("SetDirectoryProperties called on %s path %.*s. Ignoring the call."),
-							StatData.bIsValid ? TEXT("file") : TEXT("non-existent"), LocalAbsPath.Len(), LocalAbsPath.GetData());
-						OutControllingDir.ScanDir = this;
-						OutControllingDir.ParentData = ParentData;
-						return false;
-					}
-					bConfirmedExists = true;
+					QueryPath.LoadExistenceData();
+					// InRelPath might have been invalidated by the call to LoadExistenceData, and
+					// it's capitalization might have changed anyway. Recreate InRelPath from the QueryPath's
+					// relative path from this.
+					FString ThisDirAbsPath = GetLocalAbsPath();
+					// TryMakeChildPathRelativeTo should succeed because our caller promises that QueryPath's relative
+					// path from this exists in the old InRelPath.
+					ensure(FPathViews::TryMakeChildPathRelativeTo(QueryPath.GetLocalAbsPath(), ThisDirAbsPath, InRelPath));
+					FPathViews::SplitFirstComponent(InRelPath, FirstComponent, Remainder);
 				}
+
+				// If the path does not exist on disk, then contractually we are not required to set any properties for it
+				// and we can return this, the closest existing scandir to its path.
+				// If the path is a file, then we need to create directories down to its parent directory, and set the
+				// requested properties on its parent directory, and we exit out when we have found its parent directory
+				// which we detect by SplitFirstComponent returns a single component with no remainder.
+				if (QueryPath.GetType() == FPathExistence::EType::MissingParentDir
+					|| (Remainder.IsEmpty() && QueryPath.GetType() != FPathExistence::EType::Directory))
+				{
+					OutControllingDir.ScanDir = this;
+					OutControllingDir.ParentData = ParentData;
+					OutControllingDirRelPath = InRelPath;
+					return;
+				}
+
 				SubDir = &FindOrAddSubDir(FirstComponent);
 				// If the current directory has already been scanned then the SubDir we just created must have been
 				// previously discovered, or it was created on disk after the last time we scanned the current
@@ -913,24 +932,24 @@ bool FScanDir::TrySetDirectoryProperties(FStringView InRelPath, FInherited& Pare
 				// then we will set it back to bHasScanned=false in TrySetDirectoryProperties below.
 				if (this->bHasScanned)
 				{
-					FInherited ParentDataForSubDir(ParentData, DirectData);
 					if (SubDir->IsMonitored(ParentDataForSubDir))
 					{
 						SubDir->bHasScanned = true;
 					}
 				}
+				bInOutMadeChanges = true;
 				SetComplete(false);
 			}
 		}
 
-		FInherited ParentDataForSubDir(ParentData, DirectData);
-		bool bResult = SubDir->TrySetDirectoryProperties(Remainder, ParentDataForSubDir, *Properties, bConfirmedExists,
-			OutControllingDir);
+		SubDir->TrySetDirectoryProperties(QueryPath, Remainder, ParentDataForSubDir, *Properties,
+			OutControllingDir, OutControllingDirRelPath, bInOutMadeChanges);
 		if (OutControllingDir.ScanDir && !OutControllingDir.ScanDir->IsComplete())
 		{
+			bInOutMadeChanges = true;
 			SetComplete(false);
 		}
-		return bResult;
+		return;
 	}
 }
 
@@ -1135,7 +1154,6 @@ void FScanDir::SetComplete(bool bInIsComplete)
 	{
 		FAssetDataDiscovery& Discovery = MountDir->GetDiscovery();
 		Discovery.NumDirectoriesToScan.Increment();
-		UE_CLOG(Discovery.bIsIdle, LogAssetRegistry, Warning, TEXT("AssetDataGatherer: SetComplete(false) is called while bIsIdle=true."));
 	}
 }
 
@@ -1322,17 +1340,27 @@ bool FMountDir::IsMonitored(FStringView InLocalAbsPath) const
 	return MonitorData.IsMonitored();
 }
 
-bool FMountDir::TrySetDirectoryProperties(FStringView InLocalAbsPath, const FSetPathProperties& InProperties, bool bConfirmedExists,
-	FScanDirAndParentData* OutControllingDir)
+void FMountDir::TrySetDirectoryProperties(FPathExistence& QueryPath, const FSetPathProperties& InProperties,
+	FScanDirAndParentData* OutControllingDir, FStringView* OutControllingDirRelPath, bool* bOutMadeChanges)
 {
+	if (OutControllingDir)
+	{
+		OutControllingDir->ScanDir.SafeRelease();
+		OutControllingDir->ParentData = FScanDir::FInherited();
+	}
+	if (OutControllingDirRelPath)
+	{
+		*OutControllingDirRelPath = FStringView();
+	}
+	if (bOutMadeChanges)
+	{
+		*bOutMadeChanges = false;
+	}
+
 	FStringView RelPath;
-	if (!ensure(FPathViews::TryMakeChildPathRelativeTo(InLocalAbsPath, GetLocalAbsPath(), RelPath)))
+	if (!ensure(FPathViews::TryMakeChildPathRelativeTo(QueryPath.GetLocalAbsPath(), GetLocalAbsPath(), RelPath)))
 	{
-		if (OutControllingDir)
-		{
-			OutControllingDir->ScanDir.SafeRelease();
-		}
-		return false;
+		return;
 	}
 	if (InProperties.IgnoreDenyList.IsSet())
 	{
@@ -1344,24 +1372,28 @@ bool FMountDir::TrySetDirectoryProperties(FStringView InLocalAbsPath, const FSet
 			// FAssetDataDiscovery would call it on the child mount dir instead of this parent mountdir
 			FSetPathProperties NewProperties(InProperties);
 			NewProperties.IgnoreDenyList.Reset();
-			if (!NewProperties.IsSet())
-			{
-				if (OutControllingDir)
-				{
-					OutControllingDir->ScanDir.SafeRelease();
-				}
-				return false;
-			}
-			return TrySetDirectoryProperties(InLocalAbsPath, NewProperties, bConfirmedExists, OutControllingDir);
+			return TrySetDirectoryProperties(QueryPath, NewProperties, OutControllingDir,
+				OutControllingDirRelPath, bOutMadeChanges);
 		}
 	}
 	FScanDirAndParentData PlaceHolderOutControllingDir;
+	FStringView PlaceholderOutControllingDirRelPath;
+	bool bPlaceholderOutMadeChanges = false;
 	if (!OutControllingDir)
 	{
 		OutControllingDir = &PlaceHolderOutControllingDir;
 	}
+	if (!OutControllingDirRelPath)
+	{
+		OutControllingDirRelPath = &PlaceholderOutControllingDirRelPath;
+	}
+	if (!bOutMadeChanges)
+	{
+		bOutMadeChanges = &bPlaceholderOutMadeChanges;
+	}
 	FScanDir::FInherited ParentData;
-	return Root->TrySetDirectoryProperties(RelPath, ParentData, InProperties, bConfirmedExists, *OutControllingDir);
+	Root->TrySetDirectoryProperties(QueryPath, RelPath, ParentData, InProperties, *OutControllingDir,
+		*OutControllingDirRelPath, *bOutMadeChanges);
 }
 
 void FMountDir::UpdateDenyList()
@@ -1413,15 +1445,33 @@ void FMountDir::UpdateDenyList()
 		if (FileManager.DirectoryExists(AbsPathDenyList.ToString()))
 		{
 			FScanDirAndParentData UnusedControllingDir;
-			Root->TrySetDirectoryProperties(RelPath, ParentData, ChangeDenyList, true /* bConfirmedExists */, UnusedControllingDir);
+			FStringView UnusedControllingDirRelPath;
+			bool bUnusedMadeChanges = false;
+
+			FPathExistence QueryPath(AbsPathDenyList);
+			QueryPath.SetConfirmedExists(true);
+
+			Root->TrySetDirectoryProperties(QueryPath, RelPath, ParentData, ChangeDenyList,
+				UnusedControllingDir, UnusedControllingDirRelPath, bUnusedMadeChanges);
 		}
 	}
 	ChangeDenyList.MatchesDenyList = false;
 	for (const FString& RelPath : RemovedDenyLists)
 	{
-		// We don't need to check for existence when setting the removal property, because the scandir already exists
 		FScanDirAndParentData UnusedControllingDir;
-		Root->TrySetDirectoryProperties(RelPath, ParentData, ChangeDenyList, false /* bConfirmedExists */, UnusedControllingDir);
+		FStringView UnusedControllingDirRelPath;
+		bool bUnusedMadeChanges = false;
+
+		AbsPathDenyList.Reset();
+		AbsPathDenyList << LocalAbsPath;
+		FPathViews::AppendPath(AbsPathDenyList, RelPath);
+
+		// We don't need to check for existence on QueryPath when setting the removal property, because the scandir
+		// already exists
+		FPathExistence QueryPath(AbsPathDenyList);
+
+		Root->TrySetDirectoryProperties(QueryPath, RelPath, ParentData, ChangeDenyList,
+			UnusedControllingDir, UnusedControllingDirRelPath, bUnusedMadeChanges);
 	}
 }
 
@@ -2441,14 +2491,19 @@ void FPathExistence::LoadExistenceData()
 	bHasExistenceData = true;
 }
 
+bool FPathExistence::HasExistenceData() const
+{
+	return bHasExistenceData;
+}
+
+void FPathExistence::SetConfirmedExists(bool bValue)
+{
+	bHasExistenceData = bValue;
+}
+
 void FAssetDataDiscovery::SetPropertiesAndWait(TArrayView<FPathExistence> QueryPaths, bool bAddToAllowList,
 	bool bForceRescan, bool bIgnoreDenyListScanFilters)
 {
-	for (FPathExistence& QueryPath : QueryPaths)
-	{
-		QueryPath.LoadExistenceData();
-	}
-
 	struct FScanDirAndQueryPath
 	{
 		TRefCountPtr<FScanDir> ScanDir;
@@ -2461,20 +2516,12 @@ void FAssetDataDiscovery::SetPropertiesAndWait(TArrayView<FPathExistence> QueryP
 		FGathererScopeLock TreeScopeLock(&TreeLock);
 		for (FPathExistence& QueryPath : QueryPaths)
 		{
-			// We might have been asked to wait on a filename missing the extension, in which case QueryPath.GetType() == MissingButDirExists
-			// We need to handle Directory, File, and MissingButDirExists in unique ways
-			FPathExistence::EType PathType = QueryPath.GetType();
-			if (PathType == FPathExistence::EType::MissingParentDir)
-			{
-				// SetPropertiesAndWait is called for every ScanPathsSynchronous, and this is the first spot that checks for existence. 
-				// Some systems call ScanPathsSynchronous speculatively to scan whatever is present, so this log is verbose-only.
-				UE_LOG(LogAssetRegistry, Verbose, TEXT("SetPropertiesAndWait called on non-existent path %s. Call will be ignored."),
-					*QueryPath.GetLocalAbsPath());
-				continue;
-			}
-			FStringView SearchPath = QueryPath.GetLowestExistingPath();
+			// Performance note: It is important not to call any functions on QueryPath that access its existence data
+			// until after the early exit for if (ScanDir->IsComplete() && !bForceRescan). Avoiding them lets us skip
+			// the IO cost of reading existence data in the case where we have been asked to add a directory to the
+			// AssetRegistry's list of scanned directories, but it already exists.
 
-			FMountDir* MountDir = FindContainingMountPoint(SearchPath);
+			FMountDir* MountDir = FindContainingMountPoint(QueryPath.GetLocalAbsPath());
 			if (!MountDir)
 			{
 				UE_LOG(LogAssetRegistry, Log, TEXT("SetPropertiesAndWait called on %s which is not in a mounted directory. Call will be ignored."),
@@ -2482,47 +2529,91 @@ void FAssetDataDiscovery::SetPropertiesAndWait(TArrayView<FPathExistence> QueryP
 				continue;
 			}
 
-			if (PathType == FPathExistence::EType::Directory)
+			FSetPathProperties Properties;
+			if (bAddToAllowList)
 			{
-				FSetPathProperties Properties;
-				if (bAddToAllowList)
-				{
-					Properties.IsOnAllowList = bAddToAllowList;
-				}
-				if (bForceRescan)
-				{
-					Properties.HasScanned = false;
-				}
-				if (bIgnoreDenyListScanFilters)
-				{
-					Properties.IgnoreDenyList = true;
-				}
-				if (Properties.IsSet())
-				{
-					SetIsIdle(false);
-					MountDir->TrySetDirectoryProperties(SearchPath, Properties, true /* bConfirmedExists */, nullptr);
-				}
+				Properties.IsOnAllowList = bAddToAllowList;
+			}
+			if (bForceRescan)
+			{
+				Properties.HasScanned = false;
+			}
+			if (bIgnoreDenyListScanFilters)
+			{
+				Properties.IgnoreDenyList = true;
+			}
+			FScanDirAndParentData ControlDirData;
+			FStringView RelPathFromControlDir;
+			bool bMadeChanges;
+			MountDir->TrySetDirectoryProperties(QueryPath, Properties, &ControlDirData, &RelPathFromControlDir, &bMadeChanges);
+			if (bMadeChanges)
+			{
+				SetIsIdle(false);
 			}
 
-			FString RelPath;
+			TRefCountPtr<FScanDir>& ScanDir = ControlDirData.ScanDir;
 			FScanDir::FInherited MonitorData;
-			bool bSearchPathIsDirectory = PathType == FPathExistence::EType::Directory || PathType == FPathExistence::EType::MissingButDirExists;
-			TRefCountPtr<FScanDir> ScanDir = MountDir->GetControllingDir(SearchPath, bSearchPathIsDirectory, MonitorData, RelPath);
-			bool bIsAllowedInThisCall = MonitorData.IsOnAllowList() || bAddToAllowList;
-			bool bIsDeniedInThisCall = MonitorData.IsOnDenyList() && !bIgnoreDenyListScanFilters;
-			bool bIsMonitoredInThisCall = bIsAllowedInThisCall && !bIsDeniedInThisCall;
+			bool bIsAllowedInThisCall = false;
+			bool bIsDeniedInThisCall = false;
+			bool bIsMonitoredInThisCall = false;
+			if (ScanDir)
+			{
+				ScanDir->GetMonitorData(FStringView(), ControlDirData.ParentData, MonitorData);
+				bIsAllowedInThisCall = MonitorData.IsOnAllowList() || bAddToAllowList;
+				bIsDeniedInThisCall = MonitorData.IsOnDenyList() && !bIgnoreDenyListScanFilters;
+				bIsMonitoredInThisCall = bIsAllowedInThisCall && !bIsDeniedInThisCall;
+			}
 			if (!ScanDir || !bIsMonitoredInThisCall)
 			{
-				UE_LOG(LogAssetRegistry, Log, TEXT("SetPropertiesAndWait called on %s which is not monitored. Call will be ignored."),
+				UE_LOG(LogAssetRegistry, Verbose, TEXT("SetPropertiesAndWait called on %s which is not monitored. Call will be ignored."),
+					*QueryPath.GetLocalAbsPath());
+				continue;
+			}
+
+			if (ScanDir->IsComplete() && !bForceRescan)
+			{
+				// The requested path (if a directory) or its closest parent directory already had the settings we were asked to apply,
+				// and it is already complete, and (applicable if the requested path is a file) we were not asked to rescan it.
+				// Therefore there is no new scanned data to gather.
+				continue;
+			}
+
+			// After this point we are no longer in the common case of having no work to do, so now we can pay the IO
+			// cost of loading the querypath's existence data.
+
+			// We might have been asked to wait on a filename missing the extension, in which case QueryPath.GetType() == MissingButDirExists
+			// We need to handle Directory, File, and MissingButDirExists in unique ways.
+			FPathExistence::EType PathType = QueryPath.GetType();
+			if (PathType == FPathExistence::EType::MissingParentDir)
+			{
+				// SetPropertiesAndWait is called for every ScanPathsSynchronous, and this is the first spot that checks for existence. 
+				// Some systems call ScanPathsSynchronous speculatively to scan whatever is present, so this log is verbose-only.
+				UE_LOG(LogAssetRegistry, Verbose, TEXT("SetPropertiesAndWait called on non-existent path %s. Call will be ignored."),
 					*QueryPath.GetLocalAbsPath());
 				continue;
 			}
 
+			bool bSearchPathIsDirectory = PathType == FPathExistence::EType::Directory || PathType == FPathExistence::EType::MissingButDirExists;
 			if (bSearchPathIsDirectory)
 			{
-				// If Relpath from the controlling dir to the requested dir is not empty then we have found a parent directory rather than the requested directory.
-				// This can only occur for a monitored directory when the requested directory is already complete and we do not need to wait on it.
-				if (!RelPath.IsEmpty() || ScanDir->IsComplete())
+				if (ScanDir->IsComplete())
+				{
+					// The requested path (if a directory) or its closest parent directory already had the settings we
+					// were asked to apply, and it is already complete, and (applicable if the requested path is a
+					// file) the file did not exist so we can ignore the request to rescan the file. Therefore there
+					// is no new scanned data to gather. This is similar to the early exit made above, but we also
+					// early exit now even if bForceRescan is true.
+					continue;
+				}
+
+				// If Relpath from the controlling dir to the requested dir (or the parent of the requested file if
+				// the file is missing) is not empty then we have found a parent directory rather than the requested
+				// directory.
+				// This can only occur for a monitored directory when the requested directory is already complete and
+				// we do not need to wait on it.
+				if ((PathType == FPathExistence::EType::Directory && !RelPathFromControlDir.IsEmpty())
+					|| (PathType == FPathExistence::EType::MissingButDirExists 
+						&& UE::String::FindFirstOfAnyChar(RelPathFromControlDir, { '/', '\\' }) != INDEX_NONE))
 				{
 					continue;
 				}
@@ -2537,7 +2628,7 @@ void FAssetDataDiscovery::SetPropertiesAndWait(TArrayView<FPathExistence> QueryP
 					PriorityData = &PriorityScanDirs.Add_GetRef(FPriorityScanDirData{ ScanDir });
 				}
 				++PriorityData->RequestCount;
-				PriorityData->ParentData = MonitorData;
+				PriorityData->ParentData = ControlDirData.ParentData;
 			}
 			else
 			{
@@ -2545,22 +2636,23 @@ void FAssetDataDiscovery::SetPropertiesAndWait(TArrayView<FPathExistence> QueryP
 				bool bAlreadyScanned = ScanDir->HasScanned() && MonitorData.IsMonitored();
 				if (!bAlreadyScanned || bForceRescan)
 				{
-					FStringView RelPathFromParentDir = FPathViews::GetCleanFilename(RelPath);
+					FStringView RelPathFromParentDir = FPathViews::GetCleanFilename(RelPathFromControlDir);
 					EGatherableFileType FileType = GetFileType(RelPathFromParentDir);
 					if (FileType != EGatherableFileType::Invalid)
 					{
-						FStringView FileRelPathNoExt = FPathViews::GetBaseFilenameWithPath(RelPath);
+						FStringView FileRelPathNoExt = FPathViews::GetBaseFilenameWithPath(RelPathFromControlDir);
 						if (!DoesPathContainInvalidCharacters(FileType, FileRelPathNoExt))
 						{
 							TStringBuilder<256> LongPackageName;
 							LongPackageName << MountDir->GetLongPackageName();
 							FPathViews::AppendPath(LongPackageName, ScanDir->GetMountRelPath());
 							FPathViews::AppendPath(LongPackageName, FileRelPathNoExt);
-							AddDiscoveredFile(FDiscoveredPathData(SearchPath, LongPackageName, RelPathFromParentDir, QueryPath.GetModificationTime(), FileType));
-							if (FPathViews::IsPathLeaf(RelPath) && !ScanDir->HasScanned())
+							AddDiscoveredFile(FDiscoveredPathData(QueryPath.GetLocalAbsPath(), LongPackageName,
+								RelPathFromParentDir, QueryPath.GetModificationTime(), FileType));
+							if (FPathViews::IsPathLeaf(RelPathFromControlDir) && !ScanDir->HasScanned())
 							{
 								SetIsIdle(false);
-								ScanDir->MarkFileAlreadyScanned(RelPath);
+								ScanDir->MarkFileAlreadyScanned(RelPathFromControlDir);
 							}
 						}
 					}
@@ -2662,7 +2754,8 @@ void FAssetDataDiscovery::PrioritizeSearchPath(const FString& LocalAbsPath, EPri
 
 	FSetPathProperties EmptyProperties;
 	FScanDirAndParentData ScanDirAndParent;
-	MountDir->TrySetDirectoryProperties(LocalAbsPath, EmptyProperties, false, &ScanDirAndParent);
+	FPathExistence QueryPath(LocalAbsPath);
+	MountDir->TrySetDirectoryProperties(QueryPath, EmptyProperties, &ScanDirAndParent);
 	FScanDir* ScanDir = ScanDirAndParent.ScanDir.GetReference();
 	if (ScanDir && ScanDir->IsValid() && !ScanDir->IsComplete())
 	{
@@ -2681,33 +2774,40 @@ void FAssetDataDiscovery::PrioritizeSearchPath(const FString& LocalAbsPath, EPri
 	}
 }
 
-bool FAssetDataDiscovery::TrySetDirectoryProperties(const FString& LocalAbsPath, const FSetPathProperties& InProperties, bool bConfirmedExists)
+void FAssetDataDiscovery::TrySetDirectoryProperties(const FString& LocalAbsPath,
+	const FSetPathProperties& InProperties, bool bConfirmedExists)
 {
 	if (!InProperties.IsSet())
 	{
-		return false;
+		return;
 	}
 	CHECK_IS_NOT_LOCKED_CURRENT_THREAD(ResultsLock);
+	FPathExistence QueryPath(LocalAbsPath);
+	QueryPath.SetConfirmedExists(bConfirmedExists);
 	FGathererScopeLock TreeScopeLock(&TreeLock);
-	SetIsIdle(false);
-	if (!TrySetDirectoryPropertiesInternal(LocalAbsPath, InProperties, bConfirmedExists))
-	{
-		return false;
-	}
-	return true;
+	TrySetDirectoryPropertiesInternal(QueryPath, InProperties);
 }
 
-bool FAssetDataDiscovery::TrySetDirectoryPropertiesInternal(const FString& LocalAbsPath, const FSetPathProperties& InProperties, bool bConfirmedExists)
+void FAssetDataDiscovery::TrySetDirectoryPropertiesInternal(FPathExistence& QueryPath,
+	const FSetPathProperties& InProperties)
 {
 	CHECK_IS_LOCKED_CURRENT_THREAD(TreeLock);
-	FMountDir* MountDir = FindContainingMountPoint(LocalAbsPath);
+	FMountDir* MountDir = FindContainingMountPoint(QueryPath.GetLocalAbsPath());
 	if (!MountDir)
 	{
-		UE_LOG(LogAssetRegistry, Warning, TEXT("FAssetDataGatherer::SetDirectoryProperties called on unmounted path %.*s. Call will be ignored."), LocalAbsPath.Len(), *LocalAbsPath);
-		return false;
+		UE_LOG(LogAssetRegistry, Warning,
+			TEXT("FAssetDataGatherer::SetDirectoryProperties called on unmounted path %s. Call will be ignored."),
+			*QueryPath.GetLocalAbsPath());
+		return;
 	}
 
-	return MountDir->TrySetDirectoryProperties(LocalAbsPath, InProperties, bConfirmedExists, nullptr);
+	bool bMadeChanges;
+	MountDir->TrySetDirectoryProperties(QueryPath, InProperties, nullptr /* OutControllingDir */,
+		nullptr /* OutControllingDirRelPath */, &bMadeChanges);
+	if (bMadeChanges)
+	{
+		SetIsIdle(false);
+	}
 }
 
 bool FAssetDataDiscovery::IsOnAllowList(FStringView LocalAbsPath) const
@@ -5508,21 +5608,7 @@ void FAssetDataGatherer::SortPathsByPriority(
 
 	for (UE::AssetDataGather::Private::FPathExistence& QueryPath : LocalAbsPathsToPrioritize)
 	{
-		switch (QueryPath.GetType())
-		{
-		case FPathExistence::EType::Directory:
-			FilesToSearch->PrioritizeDirectory(QueryPath.GetLocalAbsPath(), Priority);
-			break;
-		case FPathExistence::EType::File:
-			// Intentional fall-through
-		case FPathExistence::EType::MissingButDirExists:
-			// We assume MissingButDirExists is a file missing the extension;
-			// PrioritizeFile handles searching for BaseNameWithPath.* and does nothing if no matches are found
-			FilesToSearch->PrioritizeFile(QueryPath.GetLocalAbsPath(), Priority);
-			break;
-		default:
-			break;
-		}
+		FilesToSearch->PrioritizePath(QueryPath, Priority);
 	}
 	OutNumPaths = FilesToSearch->NumBlockingFiles();
 }
@@ -5695,7 +5781,7 @@ void FFilesToSearch::PopFront(TArray<FGatheredPathData, AllocatorType>& Out, int
 	check(AvailableFilesNum >= 0);
 }
 
-void FFilesToSearch::PrioritizeDirectory(FStringView DirAbsPath, EPriority Priority)
+void FFilesToSearch::PrioritizePath(FPathExistence& QueryPath, EPriority Priority)
 {
 	// We may need to prioritize a LaterRetryFile that is now loadable, so add them all into the Root
 	RetryLaterRetryFiles();
@@ -5709,28 +5795,30 @@ void FFilesToSearch::PrioritizeDirectory(FStringView DirAbsPath, EPriority Prior
 		// array linearly when looking for files to accomodate a blocking priority request
 		return;
 	}
-	FTreeNode* TreeNode = Root.FindNode(DirAbsPath);
-	if (TreeNode)
+
+	// We want to look for the querypath without knowing whether it is a file or directory,
+	// because to know whether it is a file or directory we would have to do some IO.
+	// So handle both possibilities.
+	const FString& LocalAbsPath = QueryPath.GetLocalAbsPath();
+	FStringView RemainingRelPath;
+	FTreeNode& TreeNode = Root.FindClosestParentNode(LocalAbsPath, RemainingRelPath);
+	if (RemainingRelPath.IsEmpty())
 	{
-		TreeNode->PopAllFiles(BlockingFiles);
+		// The queryPath matches an existing directory in our tree; prioritize that directory by pushing all of its
+		// files (and subdirectory's files) into BlockingFiles
+		TreeNode.PopAllFiles(BlockingFiles);
 	}
-}
-
-void FFilesToSearch::PrioritizeFile(FStringView FileAbsPathExtOptional, EPriority Priority)
-{
-	// We may need to prioritize a LaterRetryFile that is now loadable, so add them all into the Root
-	RetryLaterRetryFiles();
-
-	if (Priority > EPriority::Blocking)
+	else if (UE::String::FindFirstOfAnyChar(RemainingRelPath, { '/', '\\' }) != INDEX_NONE)
 	{
-		// TODO: Implement High Priority; see note in PrioritizeDirectory
-		return;
+		// We found a parent directory of a parent directory. Neither QueryPath itself, nor its parent if it is a file,
+		// is in our tree, so there is nothing we need to prioritize for it
 	}
-	FTreeNode* TreeNode = Root.FindNode(FPathViews::GetPath(FileAbsPathExtOptional));
-	if (TreeNode)
+	else
 	{
+		// If QueryPath is a directory, it does not exist in our tree and we don't need to prioritize anything for it.
+		// But it might be a file, and its parent directory exists so look for a matching file in the parent directory.
 		int32 BeforeSize = BlockingFiles.Num();
-		TreeNode->PopMatchingDirectFiles(BlockingFiles, FileAbsPathExtOptional);
+		TreeNode.PopMatchingDirectFiles(BlockingFiles, LocalAbsPath);
 	}
 }
 
@@ -5828,6 +5916,26 @@ FFilesToSearch::FTreeNode* FFilesToSearch::FTreeNode::FindNode(FStringView InRel
 	return SubDir->FindNode(RemainingPath);
 }
 
+FFilesToSearch::FTreeNode& FFilesToSearch::FTreeNode::FindClosestParentNode(FStringView InRelPath, FStringView& OutRemainingRelPath)
+{
+	if (InRelPath.IsEmpty())
+	{
+		OutRemainingRelPath = FStringView();
+		return *this;
+	}
+	FStringView FirstComponent;
+	FStringView RemainingPath;
+	FPathViews::SplitFirstComponent(InRelPath, FirstComponent, RemainingPath);
+	FTreeNode* SubDir = FindSubDir(FirstComponent);
+	if (!SubDir)
+	{
+		OutRemainingRelPath = InRelPath;
+		return *this;
+	}
+
+	return SubDir->FindClosestParentNode(RemainingPath, OutRemainingRelPath);
+}
+
 FFilesToSearch::FTreeNode& FFilesToSearch::FTreeNode::FindOrAddSubDir(FStringView SubDirBaseName)
 {
 	int32 Index = FindLowerBoundSubDir(SubDirBaseName);
diff --git a/Engine/Source/Runtime/AssetRegistry/Private/AssetDataGathererPrivate.h b/Engine/Source/Runtime/AssetRegistry/Private/AssetDataGathererPrivate.h
--- a/Engine/Source/Runtime/AssetRegistry/Private/AssetDataGathererPrivate.h
+++ b/Engine/Source/Runtime/AssetRegistry/Private/AssetDataGathererPrivate.h
@@ -150,13 +150,11 @@ public:
 	template <typename AllocatorType>
 	void PopFront(TArray<FGatheredPathData, AllocatorType>& Out, int32 NumToPop);
 
-	/** Search the container for all files with the given directory and set their priority. */
-	void PrioritizeDirectory(FStringView DirAbsPath, EPriority Priority);
 	/**
-	 * Search the container for all files with the given Path and set their priority.
-	 * Path may be BaseNameWithPath or the full BaseNameWithPath.Extension.
-	 */
-	void PrioritizeFile(FStringView FileAbsPathExtOptional, EPriority Priority);
+	 * Search the container for all files at or under the given path and set their priority.
+	 * Path might be a directory name or a file BaseNameWithPath, or a file BaseNameWithPath.Extension.
+	*/
+	void PrioritizePath(FPathExistence& QueryPath, EPriority Priority);
 	/** How many files with blocking priority are in the container. */
 	int32 NumBlockingFiles() const;
 
@@ -188,6 +186,11 @@ private:
 		 * return nullptr if the node or any parent is not found.
 		 */
 		FTreeNode* FindNode(FStringView InRelPath);
+		/**
+		 * Recursively search the tree to find the closest parent of the given relative directory name;
+		 * returns *this if the first component of the relativepath is not in any child node.
+		 */
+		FTreeNode& FindClosestParentNode(FStringView InRelPath, FStringView& OutRemainingRelPath);
 
 		/** Add the given FilePaths as direct file children of the Node's directory. */
 		void AddFiles(TArray<FGatheredPathData>&& FilePaths);
@@ -255,6 +258,8 @@ struct FPathExistence
 
 	/** Issue the file system query if not already done. */
 	void LoadExistenceData();
+	bool HasExistenceData() const;
+	void SetConfirmedExists(bool bValue);
 
 private:
 	FString LocalAbsPath;
@@ -345,11 +350,11 @@ public:
 
 	/**
 	 * Set values of fields on the given directory indicated by InRelPath for all properties existing on InProperties.
-	 * Returns whether the directory was found and its property was changed; returns false if InRelPath was not a
-	 * directory or the property did not need to be changed.
+	 * @param bInOutMadeChanges Set to true if any changes are made. Caller is responsible for initializing to false.
 	 */
-	bool TrySetDirectoryProperties(FStringView InRelPath, FInherited& ParentData,
-		const FSetPathProperties& InProperties, bool bConfirmedExists, FScanDirAndParentData& OutControllingDir);
+	void TrySetDirectoryProperties(FPathExistence& QueryPath, FStringView InRelPath, FInherited& ParentData,
+		const FSetPathProperties& InProperties, FScanDirAndParentData& OutControllingDir,
+		FStringView& OutControllingDirRelPath, bool& bInOutMadeChanges);
 	/**
 	 * Mark that the given file has already been scanned, so that it will not be double reporting in the upcoming
 	 * directory scan, if one is upcoming.
@@ -493,11 +498,11 @@ public:
 
 	/*
 	 * Set values of fields on the given directory for all of the properties existing on InProperties.
-	 * Returns whether the directory was foundand its property was changed; returns false if LocalAbsPath was not a
-	 * directory under this MountDir or the property did not need to be changed.
 	 */
-	bool TrySetDirectoryProperties(FStringView LocalAbsPath, const FSetPathProperties& InProperties,
-		bool bConfirmedExists, FScanDirAndParentData* OutControllingDir);
+	void TrySetDirectoryProperties(FPathExistence& QueryPath, const FSetPathProperties& InProperties,
+		FScanDirAndParentData* OutControllingDir = nullptr,
+		FStringView* OutControllingDirRelPath = nullptr,
+		bool* bOutMadeChanges = nullptr);
 	/** Update all incomplete ScanDirs under this MountDir and add any that need to be scanned to OutScanRequests. */
 	void Update(TArray<FScanDirAndParentData>& OutScanRequests);
 
@@ -629,8 +634,8 @@ public:
 	/** Raise the priority until completion of scans of the given path and its subdirs. */
 	void PrioritizeSearchPath(const FString& LocalAbsPath, EPriority Priority);
 	/** Set properties on the directory, called when files are requested to be on an allow/deny list or rescanned. */
-	bool TrySetDirectoryProperties(const FString& LocalAbsPath,
-		const UE::AssetDataGather::Private::FSetPathProperties& Properties, bool bConfirmedExists);
+	void TrySetDirectoryProperties(const FString& LocalAbsPath, const FSetPathProperties& Properties,
+		bool bConfirmedExists);
 	/** Event called from the directory watcher when a directory is created. It will be scanned if IsMonitored. */
 	void OnDirectoryCreated(FStringView LocalAbsPath);
 	/** Event called from the directory watcher when files are created. Each will be reported if IsMonitored. */
@@ -674,8 +679,7 @@ private:
 	/** Handle the actions necessary for a single created file. */
 	void OnFileCreated(const FString& LocalPaths);
 	/** Helper for TrySetDirectoryProperties, called from within the critical section. */
-	bool TrySetDirectoryPropertiesInternal(const FString& LocalAbsPath,
-		const UE::AssetDataGather::Private::FSetPathProperties& Properties, bool bConfirmedExists);
+	void TrySetDirectoryPropertiesInternal(FPathExistence& QueryPath, const FSetPathProperties& Properties);
 
 	/** Add the given path as a MountPoint and update child registrations. */
 	void AddMountPointInternal(const FString& LocalAbsPath, FStringView LongPackageName, bool& bOutAlreadyExisted);
