Early out in FAssetDataDiscovery::AddMountPointInternal if the mount dir already exists (via a binary search)

#rnx

diff --git a/Engine/Source/Runtime/AssetRegistry/Private/AssetDataGatherer.cpp b/Engine/Source/Runtime/AssetRegistry/Private/AssetDataGatherer.cpp
--- a/Engine/Source/Runtime/AssetRegistry/Private/AssetDataGatherer.cpp
+++ b/Engine/Source/Runtime/AssetRegistry/Private/AssetDataGatherer.cpp
@@ -2957,17 +2957,19 @@ void FAssetDataDiscovery::AddMountPoint(const FString& LocalAbsPath, FStringView
 void FAssetDataDiscovery::AddMountPointInternal(const FString& LocalAbsPath, FStringView LongPackageName, bool& bOutAlreadyExisted)
 {
 	CHECK_IS_LOCKED_CURRENT_THREAD(TreeLock);
+
+	bOutAlreadyExisted = false;
+	if (FMountDir* ExistingMount = FindMountPoint(LocalAbsPath))
+	{
+		bOutAlreadyExisted = true;
+		return;
+	}
+
 	TArray<FMountDir*> ChildMounts;
 	FMountDir* ParentMount = nullptr;
-	bOutAlreadyExisted = false;
 	for (TUniquePtr<FMountDir>& ExistingMount : MountDirs)
 	{
-		if (FPathViews::Equals(ExistingMount->GetLocalAbsPath(), LocalAbsPath))
-		{
-			bOutAlreadyExisted = true;
-			break;
-		}
-		else if (FPathViews::IsParentPathOf(ExistingMount->GetLocalAbsPath(), LocalAbsPath))
+		if (FPathViews::IsParentPathOf(ExistingMount->GetLocalAbsPath(), LocalAbsPath))
 		{
 			// Overwrite any earlier ParentMount; later mounts are more direct parents than earlier mounts
 			ParentMount = ExistingMount.Get();
@@ -2983,10 +2985,6 @@ void FAssetDataDiscovery::AddMountPointInternal(const FString& LocalAbsPath, FSt
 			}
 		}
 	}
-	if (bOutAlreadyExisted)
-	{
-		return;
-	}
 	SetIsIdle(false);
 
 	FMountDir& Mount = FindOrAddMountPoint(LocalAbsPath, LongPackageName);
