[AssetRegistry]

- Moves the FAssetDataDiscovery into the FPreloader to allow discovery to begin as soon as possible if IsDiscoveryCacheReadEnabled(). When the Gatherer starts, it will take ownership of the FAssetDataDiscovery from the Preloader
- Removes passed in configuration parameters from FAssetDataDiscovery and FAssetDataGatherer constructors and instead the constructors self-initialize. To avoid passthrough arguments across API layers
- Renames FPreloadSettings to FGatherSettings to consolidate the global settings used for gathering (including preloading the gather and discovery step)
- Moves the AssetDataGathererConstants into this struct
- Added IsAsyncEnabled() and GetEditorGameScanMode() settings to the struct to allow as a common place for these settings since they are used in multiple places (and were regathering config state). The initialization for these new settings have been moved into FGatherSettings::Initialize

#rnx

diff --git a/Engine/Source/Runtime/AssetRegistry/Private/AssetDataGatherer.cpp b/Engine/Source/Runtime/AssetRegistry/Private/AssetDataGatherer.cpp
--- a/Engine/Source/Runtime/AssetRegistry/Private/AssetDataGatherer.cpp
+++ b/Engine/Source/Runtime/AssetRegistry/Private/AssetDataGatherer.cpp
@@ -44,45 +44,6 @@
 #include "Trace/Trace.h"
 #include "Trace/Trace.inl"
 
-namespace AssetDataGathererConstants
-{
-	constexpr int32 SingleThreadFilesPerBatch = 3;
-	constexpr int32 ExpectedMaxBatchSize = 100;
-	constexpr int32 MinSecondsToElapseBeforeCacheWrite = 30;
-	constexpr int32 MinAssetReadsBeforeCacheWrite = 100000;
-	constexpr int32 CacheShardAssetCount = 250000;
-	static constexpr uint32 CacheSerializationMagic = 0xA46B0FDA; // Versioning and integrity checking
-	static constexpr uint64 CurrentVersion = FAssetRegistryVersion::LatestVersion | (uint64(CacheSerializationMagic) << 32);
-	static const FTimespan CachePruneAge = FTimespan(7, 0, 0, 0);
-	/**
-	 * The maximum number of threads used in the ParallelFor in FAssetDataDiscovery::TickInternal.
-	 * Higher values cause more time spent waking up workers. Lower values cause more time spent running the DirectoryEnumeration on each thread.
-	 * Current value was chosen based on experiments on an SSD drive with a warm disk cache.
-	 */
-	static int32 GARDiscoverThreads = 15;
-	/**
-	 * Minimum batch sized used used in the ParallelFor in FAssetDataDiscovery::TickInternal.
-	 * In theory higher minimum values could reduce the cost of waking up workers, but in practice it does not have a significant effect.
-	 */
-	static int32 GARDiscoverMinBatchSize = 1;
-	/**
-	 * The maximum number of threads used in the ParallelFor around ReadAssetFile in FAssetDataGatherer::TickInternal.
-	 * Very high numbers of threads does not reduce wall time of the gather because the threads become IO limited.
-	 * Capping the number of threads allows those threads to be used in other async operations in Engine startup.
-	 */
-	static int32 GARGatherThreads = 60;
-
-	/**
-	 * The maximum number of the gatherer's cache files that can be loaded in parallel in LoadCacheFiles.
-	 * The amount of memory used when reading the cache rises proportional to this value.
-	 * Other than the memory cost, it should be set as high as possible to load the cache faster.
-	 */
-	static int32 GARGatherCacheParallelism = 16;
-
-	/** How long we wait on the gatherer thread between asking the discovery thread for no results. */
-	constexpr float PollDiscoveryPeriodSeconds = .005f;
-}
-
 namespace UE::AssetDataGather::Private
 {
 
@@ -178,7 +139,38 @@ void LexFromString(EFeatureEnabledReadWrite& OutValue, FStringView Text)
 	return;
 }
 
-void FPreloadSettings::Initialize()
+EEditorGameScanMode GetEditorGameScanModeFromConfig()
+{
+#if !WITH_EDITOR
+	FString EditorGameScansAR = TEXT("Sync");
+	EEditorGameScanMode Result = EEditorGameScanMode::Sync;
+
+	GConfig->GetString(TEXT("AssetRegistry"), TEXT("EditorGameScansAR"), EditorGameScansAR, GEngineIni);
+	if (EditorGameScansAR.Equals(TEXT("Async"), ESearchCase::IgnoreCase))
+	{
+		Result = EEditorGameScanMode::Async;
+	}
+	else if (EditorGameScansAR.Equals(TEXT("False"), ESearchCase::IgnoreCase))
+	{
+		Result = EEditorGameScanMode::None;
+	}
+	else if (EditorGameScansAR.Equals(TEXT("True"), ESearchCase::IgnoreCase))
+	{
+		Result = EEditorGameScanMode::Sync;
+	}
+	else
+	{
+		ensureMsgf(EditorGameScansAR.Equals(TEXT("Sync"), ESearchCase::IgnoreCase),
+			TEXT("Valid values for EditorGameScansAR are: true|false|sync|async. Received %s"), *EditorGameScansAR);
+	}
+
+	return Result;
+#else
+	return EEditorGameScanMode::None;
+#endif
+}
+
+void FGatherSettings::Initialize()
 {
 	if (bInitialized)
 	{
@@ -277,74 +269,47 @@ void FPreloadSettings::Initialize()
 		break;
 	}
 
-	bool bAsyncEnabled = FPlatformProcess::SupportsMultithreading() && FTaskGraphInterface::IsRunning();
+	EditorGameScanMode = GetEditorGameScanModeFromConfig();
+
+	// If EditorGameScanMode is set to allow async then we will use async for any editor build (editor, game, or server)
+	// Otherwise we will use async for commandlets and editor proper
+	bAsyncEnabled = (!IsRunningGame() && !IsRunningDedicatedServer())
+#if WITH_EDITOR
+		|| (EditorGameScanMode == EEditorGameScanMode::Async)
+#endif
+		;
+
+#if WITH_EDITOR || !UE_BUILD_SHIPPING
+	bool bCommandlineSynchronous;
+	if (FParse::Bool(FCommandLine::Get(), TEXT("AssetGatherSync="), bCommandlineSynchronous))
+	{
+		bAsyncEnabled = !bCommandlineSynchronous;
+	}
+#endif // WITH_EDITOR || !UE_BUILD_SHIPPING
+	if (bAsyncEnabled && (!FPlatformProcess::SupportsMultithreading() || !FTaskGraphInterface::IsRunning()))
+	{
+		bAsyncEnabled = false;
+		UE_LOG(LogAssetRegistry, Warning, TEXT("Requested asynchronous asset data gather, but threading support is disabled. Performing a synchronous gather instead!"));
+	}
 
 	CacheBaseFilename = AssetRegistryCacheRootFolder / (bGatherDependsData ? TEXT("CachedAssetRegistry") : TEXT("CachedAssetRegistryNoDeps"));
 #if UE_EDITOR // See note on FPreloader for why we only allow preloading if UE_EDITOR
-	bPreloadCache = bAsyncEnabled && UE::AssetRegistry::ShouldSearchAllAssetsAtStart();
+	bPreloadGatherCache = bAsyncEnabled && UE::AssetRegistry::ShouldSearchAllAssetsAtStart();
 #else
-	bPreloadCache = false;
+	bPreloadGatherCache = false;
 #endif
 
-	FParse::Value(FCommandLine::Get(), TEXT("-ARDiscoverThreads="), AssetDataGathererConstants::GARDiscoverThreads);
-	FParse::Value(FCommandLine::Get(), TEXT("-ARDiscoverMinBatchSize="), AssetDataGathererConstants::GARDiscoverMinBatchSize);
-	FParse::Value(FCommandLine::Get(), TEXT("-ARGatherThreads="), AssetDataGathererConstants::GARGatherThreads);
-	FParse::Value(FCommandLine::Get(), TEXT("-ARGatherCacheParallelism="), AssetDataGathererConstants::GARGatherCacheParallelism);
-	AssetDataGathererConstants::GARDiscoverThreads = FMath::Max(0, AssetDataGathererConstants::GARDiscoverThreads);
-	AssetDataGathererConstants::GARDiscoverMinBatchSize = FMath::Max(1, AssetDataGathererConstants::GARDiscoverMinBatchSize);
-	AssetDataGathererConstants::GARGatherThreads = FMath::Max(0, AssetDataGathererConstants::GARGatherThreads);
-	AssetDataGathererConstants::GARGatherCacheParallelism = FMath::Max(1, AssetDataGathererConstants::GARGatherCacheParallelism);
-}
-bool FPreloadSettings::IsGatherCacheReadEnabled() const
-{
-	return bGatherCacheReadEnabled;
-}
-bool FPreloadSettings::IsGatherCacheWriteEnabled() const
-{
-	return bGatherCacheWriteEnabled;
-}
-bool FPreloadSettings::IsDiscoveryCacheReadEnabled() const
-{
-	return bDiscoveryCacheReadEnabled;
-}
-EFeatureEnabled FPreloadSettings::IsDiscoveryCacheWriteEnabled() const
-{
-	return DiscoveryCacheWriteEnabled;
-}
-bool FPreloadSettings::IsDiscoveryCacheInvalidateEnabled() const
-{
-	return bDiscoveryCacheInvalidateEnabled;
-}
-bool FPreloadSettings::IsPreloadCache() const
-{
-	return bGatherCacheReadEnabled && bPreloadCache;
-}
-bool FPreloadSettings::IsGatherDependsData() const
-{
-	return bGatherDependsData;
-}
-bool FPreloadSettings::IsForceDependsGathering() const
-{
-	return bForceDependsGathering;
-}
-FString FPreloadSettings::GetLegacyCacheFilename() const
-{
-	return CacheBaseFilename + TEXT(".bin");
-}
-FString FPreloadSettings::GetLegacyNonMonolithicCacheDirectory() const
-{
-	return UE::AssetDataGather::Private::GPreloadSettings.GetAssetRegistryCacheRootFolder() / TEXT("AssetRegistryCache");
-}
-const FString& FPreloadSettings::GetCacheBaseFilename() const
-{
-	return CacheBaseFilename;
-}
-const FString& FPreloadSettings::GetAssetRegistryCacheRootFolder() const
-{
-	return AssetRegistryCacheRootFolder;
+	FParse::Value(FCommandLine::Get(), TEXT("-ARDiscoverThreads="), FGatherSettings::GARDiscoverThreads);
+	FParse::Value(FCommandLine::Get(), TEXT("-ARDiscoverMinBatchSize="), FGatherSettings::GARDiscoverMinBatchSize);
+	FParse::Value(FCommandLine::Get(), TEXT("-ARGatherThreads="), FGatherSettings::GARGatherThreads);
+	FParse::Value(FCommandLine::Get(), TEXT("-ARGatherCacheParallelism="), FGatherSettings::GARGatherCacheParallelism);
+	FGatherSettings::GARDiscoverThreads = FMath::Max(0, FGatherSettings::GARDiscoverThreads);
+	FGatherSettings::GARDiscoverMinBatchSize = FMath::Max(1, FGatherSettings::GARDiscoverMinBatchSize);
+	FGatherSettings::GARGatherThreads = FMath::Max(0, FGatherSettings::GARGatherThreads);
+	FGatherSettings::GARGatherCacheParallelism = FMath::Max(1, FGatherSettings::GARGatherCacheParallelism);
 }
 
-TArray<FString> FPreloadSettings::FindShardedCacheFiles() const
+TArray<FString> FGatherSettings::FindShardedCacheFiles() const
 {
 	TArray<FString> CachePaths;
 	IFileManager::Get().FindFiles(CachePaths, *(GetCacheBaseFilename() + TEXT("_*.bin")), /* Files */ true, /* Directories */ false);
@@ -359,7 +324,7 @@ TArray<FString> FPreloadSettings::FindShardedCacheFiles() const
 	return CachePaths;
 }
 
-FPreloadSettings GPreloadSettings;
+FGatherSettings GGatherSettings;
 
 /** A structure to hold serialized cache data from async loads before adding it to the Gatherer's main cache. */
 struct FCachePayload
@@ -1620,26 +1585,25 @@ bool FMountDir::IsChildMountPath(FStringView MountRelPath) const
 }
 
 
-FAssetDataDiscovery::FAssetDataDiscovery(const TArray<FString>& InLongPackageNamesDenyList,
-	const TArray<FString>& InMountRelativePathsDenyList, bool bInAsyncEnabled)
-	: LongPackageNamesDenyList(InLongPackageNamesDenyList)
-	, MountRelativePathsDenyList(InMountRelativePathsDenyList)
+FAssetDataDiscovery::FAssetDataDiscovery()
+	: LongPackageNamesDenyList()
+	, MountRelativePathsDenyList()
 	, Thread(nullptr)
-	, bAsyncEnabled(bInAsyncEnabled)
 	, bIsIdle(false)
 	, IsStopped(0)
 	, IsPaused(0)
 	, NumDirectoriesToScan(0)
 {
-	if (bAsyncEnabled && !FPlatformProcess::SupportsMultithreading())
+	GGatherSettings.Initialize();
+	bAsyncEnabled = GGatherSettings.IsAsyncEnabled();
+
+	if (FConfigFile* EngineIni = GConfig->FindConfigFile(GEngineIni))
 	{
-		bAsyncEnabled = false;
-		UE_LOG(LogAssetRegistry, Warning, TEXT("Requested asyncronous asset data discovery, but threading support is disabled. Performing a synchronous discovery instead!"));
+		EngineIni->GetArray(TEXT("AssetRegistry"), TEXT("BlacklistPackagePathScanFilters"), LongPackageNamesDenyList);
+		EngineIni->GetArray(TEXT("AssetRegistry"), TEXT("BlacklistContentSubPathScanFilters"), MountRelativePathsDenyList);
 	}
 
 	PriorityDataUpdated->Trigger();
-
-	GPreloadSettings.Initialize();
 }
 
 FAssetDataDiscovery::~FAssetDataDiscovery()
@@ -2073,9 +2037,9 @@ void FAssetDataDiscovery::TickInternal(bool bTickAll)
 
 		// Otherwise look on disk in parallel for all of the DirToScans
 		int32 NumThreads = FMath::Max(FTaskGraphInterface::Get().GetNumWorkerThreads(), 1);
-		if (AssetDataGathererConstants::GARDiscoverThreads > 0)
+		if (FGatherSettings::GARDiscoverThreads > 0)
 		{
-			NumThreads = FMath::Min(NumThreads, AssetDataGathererConstants::GARDiscoverThreads);
+			NumThreads = FMath::Min(NumThreads, FGatherSettings::GARDiscoverThreads);
 		}
 		int32 DirToScanBuffersNum = FMath::Min(NumThreads, DirToScanDatasNum);
 		if (DirToScanBuffers.Num() < DirToScanBuffersNum)
@@ -2088,7 +2052,7 @@ void FAssetDataDiscovery::TickInternal(bool bTickAll)
 			ScanBuffer.Reset();
 		}
 
-		ParallelForWithExistingTaskContext(LocalScanBuffers, DirToScanDatasNum, AssetDataGathererConstants::GARDiscoverMinBatchSize,
+		ParallelForWithExistingTaskContext(LocalScanBuffers, DirToScanDatasNum, FGatherSettings::GARDiscoverMinBatchSize,
 		[this](FDirToScanBuffer& ScanBuffer, int32 DirToScanDatasIndex)
 		{
 			FDirToScanData& Data = DirToScanDatas[DirToScanDatasIndex];
@@ -3658,13 +3622,14 @@ private:
 };
 
 #if UE_EDITOR
-/** A class to preload the cache used by FAssetDataGatherer. Preloading the cache allows us to
- * start very early in editor startup, so that we have time to finish the cache load before the engine starts
- * making package load requests that need to use the gathered data.
+/** A class to preload the AssetData cache and start the AssetDataDiscovery thread (which does its own cache loading)
+ * used by the FAssetDataGatherer. Preloading allows us to start very early in editor startup, so that we 
+ * have time to finish the cache loads before the engine starts making package load requests that need to 
+ * use the discovery and gatherer data.
  *
  * In UE_EDITOR, we know the values we need to decide whether we can preload early enough that it is useful to preload
  * In other configurations we do not know those parameters for sure until EDelayedRegisterRunPhase::ShaderTypesReady,
- *  which occurs around the same time as UAssetRegistryImpl is created, so it is not useful to preload.
+ * which occurs around the same time as UAssetRegistryImpl is created, so it is not useful to preload.
  */
 class FPreloader : public FDelayedAutoRegisterHelper
 {
@@ -3687,25 +3652,41 @@ public:
 	TArray<FCachePayload> Consume()
 	{
 		check(bInitialized);
-		if (bConsumed)
+		if (bConsumedGatherCache)
 		{
 			return {};
 		}
-		bConsumed = true;
+		bConsumedGatherCache = true;
 		PreloadReady.Wait();
 		PreloadReady.Reset();
-		return MoveTemp(Payloads);
+		return MoveTemp(GatherCachePayloads);
+	}
+
+	TUniquePtr<FAssetDataDiscovery> ConsumeDiscovery()
+	{
+		check(bInitialized);
+		if (bConsumedDiscoveryCache)
+		{
+			return nullptr;
+		}
+		bConsumedDiscoveryCache = true;
+		return MoveTemp(Discovery);
 	}
 
 private:
 
 	void DelayedInitialize()
 	{
-		GPreloadSettings.Initialize();
+		GGatherSettings.Initialize();
 
-		if (GPreloadSettings.IsPreloadCache())
+		if (GGatherSettings.IsPreloadDiscoveryCache())
 		{
-			TArray<FString> CachePaths = GPreloadSettings.FindShardedCacheFiles();
+			Discovery = MakeUnique<FAssetDataDiscovery>();
+			Discovery->StartAsync();
+		}
+		if (GGatherSettings.IsPreloadGatherCache())
+		{
+			TArray<FString> CachePaths = GGatherSettings.FindShardedCacheFiles();
 			if (CachePaths.Num())
 			{
 				PreloadReady = Async(EAsyncExecution::TaskGraph, [this, CachePaths=MoveTemp(CachePaths)]() { LoadAsync(CachePaths); });
@@ -3717,24 +3698,25 @@ private:
 	void LoadAsync(const TArray<FString>& Paths)
 	{
 		LLM_SCOPE(ELLMTag::AssetRegistry);
-		Payloads = LoadCacheFiles(Paths);
+		GatherCachePayloads = LoadCacheFiles(Paths);
 	}
 
 	TFuture<void> PreloadReady;
-	TArray<FCachePayload> Payloads;
+	TArray<FCachePayload> GatherCachePayloads;
+	TUniquePtr<UE::AssetDataGather::Private::FAssetDataDiscovery> Discovery;
 	bool bInitialized = false;
-	bool bConsumed = false;
+	bool bConsumedGatherCache = false;
+	bool bConsumedDiscoveryCache = false;
 };
 FPreloader GPreloader;
 #endif
 
 } // namespace UE::AssetDataGather::Private
 
-FAssetDataGatherer::FAssetDataGatherer(const TArray<FString>& InLongPackageNamesDenyList,
-	const TArray<FString>& InMountRelativePathsDenyList, bool bInAsyncEnabled, UE::AssetRegistry::FAssetRegistryImpl& InRegistryImpl)
+FAssetDataGatherer::FAssetDataGatherer(UE::AssetRegistry::FAssetRegistryImpl& InRegistryImpl)
 	: AssetRegistry(InRegistryImpl)
 	, Thread(nullptr)
-	, bAsyncEnabled(bInAsyncEnabled)
+	, bAsyncEnabled(false)
 	, GatherStartTime(FDateTime::Now())
 	, IsStopped(0)
 	, IsGatheringPaused(0)
@@ -3760,37 +3742,36 @@ FAssetDataGatherer::FAssetDataGatherer(const TArray<FString>& InLongPackageNames
 	using namespace UE::AssetDataGather::Private;
 
 	int32 NumGatherThreads = FMath::Max(1, FTaskGraphInterface::Get().GetNumWorkerThreads());
-	GPreloadSettings.Initialize();
-	if (AssetDataGathererConstants::GARGatherThreads > 0)
+	GGatherSettings.Initialize();
+	if (FGatherSettings::GARGatherThreads > 0)
 	{
-		NumGatherThreads = FMath::Min(NumGatherThreads, AssetDataGathererConstants::GARGatherThreads);
+		NumGatherThreads = FMath::Min(NumGatherThreads, FGatherSettings::GARGatherThreads);
 	}
-	TickInternalBatchSize = NumGatherThreads * AssetDataGathererConstants::SingleThreadFilesPerBatch;
+	TickInternalBatchSize = NumGatherThreads * FGatherSettings::SingleThreadFilesPerBatch;
 
 	constexpr bool bEditorExecutable = WITH_EDITOR;
-	bGatherAssetPackageData = bEditorExecutable || GPreloadSettings.IsForceDependsGathering();
-	bGatherDependsData = GPreloadSettings.IsGatherDependsData();
-	bCacheReadEnabled = GPreloadSettings.IsGatherCacheReadEnabled();
-	bCacheWriteEnabled = GPreloadSettings.IsGatherCacheWriteEnabled();
+	bGatherAssetPackageData = bEditorExecutable || GGatherSettings.IsForceDependsGathering();
+	bGatherDependsData = GGatherSettings.IsGatherDependsData();
+	bCacheReadEnabled = GGatherSettings.IsGatherCacheReadEnabled();
+	bCacheWriteEnabled = GGatherSettings.IsGatherCacheWriteEnabled();
+	bAsyncEnabled = GGatherSettings.IsAsyncEnabled();
 	LastCacheWriteTime = FPlatformTime::Seconds();
 
-#if WITH_EDITOR || !UE_BUILD_SHIPPING
-	bool bCommandlineSynchronous;
-	if (FParse::Bool(FCommandLine::Get(), TEXT("AssetGatherSync="), bCommandlineSynchronous))
+	// Tick is synchronous until StartAsync is called, even if bAsyncEnabled
+	bSynchronousTick = true; 
+
+#if WITH_EDITOR
+	// If the preloader has already created the FAssetDataDiscovery, take ownership of it now
+	if (TUniquePtr<FAssetDataDiscovery> PreloadedDiscovery = GPreloader.ConsumeDiscovery())
 	{
-		bAsyncEnabled = !bCommandlineSynchronous;
+		Discovery = MoveTemp(PreloadedDiscovery);
 	}
-#endif // WITH_EDITOR || !UE_BUILD_SHIPPING
-	if (bAsyncEnabled && !FPlatformProcess::SupportsMultithreading())
+	else
+#endif
 	{
-		bAsyncEnabled = false;
-		UE_LOG(LogAssetRegistry, Warning, TEXT("Requested asynchronous asset data gather, but threading support is disabled. Performing a synchronous gather instead!"));
+		Discovery = MakeUnique<UE::AssetDataGather::Private::FAssetDataDiscovery>();
 	}
-	// Tick is Synchronous until StartAsync is called, even if bAsyncEnabled
-	bSynchronousTick = true; 
 
-	Discovery = MakeUnique<UE::AssetDataGather::Private::FAssetDataDiscovery>(InLongPackageNamesDenyList,
-		InMountRelativePathsDenyList, bAsyncEnabled);
 	FilesToSearch = MakeUnique<UE::AssetDataGather::Private::FFilesToSearch>();
 }
 
@@ -3893,10 +3874,11 @@ uint32 FAssetDataGatherer::Run()
 
 			if (TickResult == ETickResult::PollDiscovery)
 			{
+				using namespace UE::AssetDataGather::Private;
 				// The gatherer thread is waiting on results from the discovery thread, and we should sleep rather
 				// than busy wait on it, to reduce contention. TODO: Change this to an event triggered by discovery
 				// results, rather than sleeping for a fixed time interval.
-				FPlatformProcess::Sleep(AssetDataGathererConstants::PollDiscoveryPeriodSeconds);
+				FPlatformProcess::Sleep(FGatherSettings::PollDiscoveryPeriodSeconds);
 			}
 		}
 	}
@@ -3951,7 +3933,7 @@ FAssetDataGatherer::ETickResult FAssetDataGatherer::InnerTickLoop(bool bInSynchr
 			{
 				// Poll discovery every so often to reduce super-linear costs
 				float TimeSinceLastPoll = static_cast<float>(CurrentTimeSeconds - LastPollTimeSeconds);
-				bPollDiscovery = TimeSinceLastPoll > AssetDataGathererConstants::PollDiscoveryPeriodSeconds;
+				bPollDiscovery = TimeSinceLastPoll > FGatherSettings::PollDiscoveryPeriodSeconds;
 			}
 			else
 			{
@@ -4012,7 +3994,7 @@ void FAssetDataGatherer::SaveCacheFile(const TArray<TPair<FName,FDiskCachedAsset
 	}
 
 	const int32 MaxShards = FPlatformMisc::NumberOfCores();
-	constexpr int32 CacheShardAssetCount = AssetDataGathererConstants::CacheShardAssetCount;
+	constexpr int32 CacheShardAssetCount = FGatherSettings::CacheShardAssetCount;
 	const int32 CacheShards = FMath::RoundUpToPowerOfTwo(FMath::Min((AssetsToSave.Num() + CacheShardAssetCount-1) / CacheShardAssetCount, MaxShards));
 	const int32 ShardMask = CacheShards - 1;
 	std::atomic<int64> TotalCacheSize{0};
@@ -4046,21 +4028,21 @@ void FAssetDataGatherer::SaveCacheFile(const TArray<TPair<FName,FDiskCachedAsset
 	ParallelFor(CacheShards, [this, &AssetsToSave, &TotalCacheSize, &DataPerShard, ShardMask](int32 Shard) 
 	{
 		TArray<TPair<FName, FDiskCachedAssetData*>>& ShardData = DataPerShard[Shard];
-		FString Filename = FString::Printf(TEXT("%s_%d.bin"), *GPreloadSettings.GetCacheBaseFilename(), Shard);
+		FString Filename = FString::Printf(TEXT("%s_%d.bin"), *GGatherSettings.GetCacheBaseFilename(), Shard);
 		int64 CacheSize = SaveCacheFileInternal(Filename, ShardData);
 		TotalCacheSize += CacheSize;
 	}, EParallelForFlags::BackgroundPriority);
 
 	UE_LOG(LogAssetRegistry, Display, TEXT("Asset registry cache written as %.1f MiB to %s_*.bin"), 
-		static_cast<float>(TotalCacheSize)/1024.f/1024.f, *GPreloadSettings.GetCacheBaseFilename());
+		static_cast<float>(TotalCacheSize)/1024.f/1024.f, *GGatherSettings.GetCacheBaseFilename());
 	// Delete old name of monolithic cache file and old non-monolithic cache directory if they exist
-	IFileManager::Get().Delete(*GPreloadSettings.GetLegacyCacheFilename(), false /* bRequireExists */,
+	IFileManager::Get().Delete(*GGatherSettings.GetLegacyCacheFilename(), false /* bRequireExists */,
 		true /* EvenReadOnly */, true /* Quiet */);
-	IFileManager::Get().DeleteDirectory(*GPreloadSettings.GetLegacyNonMonolithicCacheDirectory(),
+	IFileManager::Get().DeleteDirectory(*GGatherSettings.GetLegacyNonMonolithicCacheDirectory(),
 		false /* bRequireExists */, true /* Tree */);
 
 	// Delete any other shards if number of shards was reduced
-	TArray<FString> CacheFiles = GPreloadSettings.FindShardedCacheFiles();
+	TArray<FString> CacheFiles = GGatherSettings.FindShardedCacheFiles();
 	for (const FString& CacheFile : CacheFiles)
 	{
 		FStringView BaseName = FPathViews::GetBaseFilename(CacheFile);
@@ -4141,7 +4123,7 @@ FAssetDataGatherer::ETickResult FAssetDataGatherer::TickInternal(double& TickSta
 {
 	using namespace UE::AssetDataGather::Private;
 
-	typedef TInlineAllocator<AssetDataGathererConstants::ExpectedMaxBatchSize> FBatchInlineAllocator;
+	typedef TInlineAllocator<FGatherSettings::ExpectedMaxBatchSize> FBatchInlineAllocator;
 
 	TArray<FGatheredPathData, FBatchInlineAllocator> LocalFilesToSearch;
 	TArray<TUniquePtr<FAssetData>, FBatchInlineAllocator> LocalAssetResults;
@@ -4249,19 +4231,19 @@ FAssetDataGatherer::ETickResult FAssetDataGatherer::TickInternal(double& TickSta
 		double CacheLoadStartTime = FPlatformTime::Seconds();
 		TArray<FCachePayload> Payloads;
 #if UE_EDITOR
-		if (GPreloadSettings.IsPreloadCache())
+		if (GGatherSettings.IsPreloadGatherCache())
 		{
 			Payloads = GPreloader.Consume();
 		}
 		else
 #endif
 		{
-			TArray<FString> CachePaths = GPreloadSettings.FindShardedCacheFiles();
+			TArray<FString> CachePaths = GGatherSettings.FindShardedCacheFiles();
 			Payloads = UE::AssetDataGather::Private::LoadCacheFiles(CachePaths);
 		}
 		ConsumeCacheFiles(MoveTemp(Payloads));
 		UE_LOG(LogAssetRegistry, Display, TEXT("AssetDataGatherer spent %.3fs loading caches %s_*.bin."),
-			FPlatformTime::Seconds() - CacheLoadStartTime, *GPreloadSettings.GetCacheBaseFilename());
+			FPlatformTime::Seconds() - CacheLoadStartTime, *GGatherSettings.GetCacheBaseFilename());
 
 		FGathererScopeLock ResultsScopeLock(&ResultsLock);
 		bHasLoadedCache = true;
@@ -4419,7 +4401,7 @@ FAssetDataGatherer::ETickResult FAssetDataGatherer::TickInternal(double& TickSta
 	// NumberOfThreads == ParallelForNum/MinBatchSize == TickInternalBatchSize/SingleThreadFilesPerBatch
 	//                 == Min(WorkerThreads, GARGatherThreads)*SingleThreadFilesPerBatch / SingleThreadFilesPerBatch
 	//                 <= GARGatherThreads
-	const int32 MinBatchSize = AssetDataGathererConstants::SingleThreadFilesPerBatch;
+	const int32 MinBatchSize = FGatherSettings::SingleThreadFilesPerBatch;
 
 	ParallelFor(TEXT("AssetDataGatherReadAssetFile"), ReadContexts.Num(), MinBatchSize,
 		[this, &ReadContexts](int32 Index)
@@ -4574,8 +4556,8 @@ FAssetDataGatherer::ETickResult FAssetDataGatherer::TickInternal(double& TickSta
 
 		const int32 NumAssetsReadSinceLastCacheWrite = NumUncachedAssetFiles - LastCacheSaveNumUncachedAssetFiles;
 		if (bCacheWriteEnabled && !bIsSavingAsyncCache 
-			&& FPlatformTime::Seconds() - LocalLastCacheWriteTime >= AssetDataGathererConstants::MinSecondsToElapseBeforeCacheWrite
-			&& NumAssetsReadSinceLastCacheWrite >= AssetDataGathererConstants::MinAssetReadsBeforeCacheWrite)
+			&& FPlatformTime::Seconds() - LocalLastCacheWriteTime >= FGatherSettings::MinSecondsToElapseBeforeCacheWrite
+			&& NumAssetsReadSinceLastCacheWrite >= FGatherSettings::MinAssetReadsBeforeCacheWrite)
 		{
 			bSaveAsyncCacheTriggered = true;
 			TickResult = ETickResult::Interrupt;
@@ -5068,6 +5050,7 @@ void FAssetDataGatherer::GetCacheAssetsToSave(TArray<TPair<FName,FDiskCachedAsse
 
 	for (const TPair<FName, FDiskCachedAssetData*>& Pair : DiskCachedAssetDataMap)
 	{
+		using namespace UE::AssetDataGather::Private;
 		if (NewCachedAssetDataMap.Contains(Pair.Key))
 		{
 			continue; // Data was replaced when populating NewCachedAssetDataMap
@@ -5075,7 +5058,7 @@ void FAssetDataGatherer::GetCacheAssetsToSave(TArray<TPair<FName,FDiskCachedAsse
 		FTimespan Age = GatherStartTime - Pair.Value->LastGatheredTime;
 		// Conservatively persist cached data until final save to avoid pruning the cache too much if discovery is interrupted 
 		// when revisiting a workspace after some delay 
-		if (!bFinishedInitialDiscovery || Age < AssetDataGathererConstants::CachePruneAge)
+		if (!bFinishedInitialDiscovery || Age < FGatherSettings::CachePruneAge)
 		{
 			OutAssetsToSave.Add(Pair);
 		}
@@ -5107,7 +5090,8 @@ int64 FAssetDataGatherer::SaveCacheFileInternal(const FString& CacheFilename, co
 	TUniquePtr<FArchive> FileAr(IFileManager::Get().CreateFileWriter(*TempFilename, 0));
 	if (FileAr)
 	{
-		uint64 CurrentVersion = AssetDataGathererConstants::CurrentVersion;
+		using namespace UE::AssetDataGather::Private;
+		uint64 CurrentVersion = FGatherSettings::CurrentVersion;
 		*FileAr << CurrentVersion;
 
 #if ALLOW_NAME_BATCH_SAVING
@@ -5285,7 +5269,7 @@ TArray<FCachePayload> LoadCacheFiles(TConstArrayView<FString> InCacheFilenames)
 
 		// We want to restrict the number of threads to reduce maximum memory use due to having multiple intermediate cache files
 		// in parallel. But ParallelFor only provides an API for restricting MinBatchSize.
-		const int32 MaxNumThreads = AssetDataGathererConstants::GARGatherCacheParallelism * TasksPerResult;
+		const int32 MaxNumThreads = FGatherSettings::GARGatherCacheParallelism * TasksPerResult;
 		const int32 MinBatchSize = FMath::DivideAndRoundUp(NumTasks, MaxNumThreads);
 
 		ParallelFor(TEXT("AssetDataGatherer::LoadCacheFiles"), NumTasks, MinBatchSize,
@@ -5357,7 +5341,7 @@ TArray<FCachePayload> LoadCacheFiles(TConstArrayView<FString> InCacheFilenames)
 				FCachePayload Payload;
 				FMemoryViewReader FileReader(FileData->File->View());
 				TOptional<uint64> Version = FileReader.TryLoad<uint64>();
-				if (Version == AssetDataGathererConstants::CurrentVersion)
+				if (Version == FGatherSettings::CurrentVersion)
 				{
 					FChecksumViewReader ChecksummingReader(MoveTemp(FileReader), *CacheFilename);
 					Payload = DoLoad(ChecksummingReader);
@@ -5399,7 +5383,7 @@ TArray<FCachePayload> LoadCacheFiles(TConstArrayView<FString> InCacheFilenames)
 			{
 				uint64 Version = 0;
 				*FileAr << Version;
-				if (Version == AssetDataGathererConstants::CurrentVersion)
+				if (Version == FGatherSettings::CurrentVersion)
 				{
 					FChecksumArchiveReader ChecksummingReader(*FileAr);
 					Payload = DoLoad(ChecksummingReader);
@@ -5465,7 +5449,7 @@ SIZE_T FAssetDataGatherer::GetAllocatedSize() const
 	Result += BlockedResults.GetAllocatedSize();
 	Result += SearchTimes.GetAllocatedSize();
 	Result += GetArrayRecursiveAllocatedSize(DiscoveredPaths);
-	Result += GPreloadSettings.GetCacheBaseFilename().GetAllocatedSize();
+	Result += GGatherSettings.GetCacheBaseFilename().GetAllocatedSize();
 
 	Result += NewCachedAssetData.GetAllocatedSize();
 	for (const FDiskCachedAssetData* Value : NewCachedAssetData)
diff --git a/Engine/Source/Runtime/AssetRegistry/Private/AssetDataGatherer.h b/Engine/Source/Runtime/AssetRegistry/Private/AssetDataGatherer.h
--- a/Engine/Source/Runtime/AssetRegistry/Private/AssetDataGatherer.h
+++ b/Engine/Source/Runtime/AssetRegistry/Private/AssetDataGatherer.h
@@ -78,8 +78,7 @@ struct FAssetGatherDiagnostics
 class FAssetDataGatherer : public FRunnable
 {
 public:
-	FAssetDataGatherer(const TArray<FString>& InLongPackageNamesDenyList,
-		const TArray<FString>& InMountRelativePathsDenyList, bool bInAsyncEnabled, UE::AssetRegistry::FAssetRegistryImpl& InRegistryImpl);
+	FAssetDataGatherer(UE::AssetRegistry::FAssetRegistryImpl& InRegistryImpl);
 	virtual ~FAssetDataGatherer();
 
 	void OnInitialSearchCompleted();
diff --git a/Engine/Source/Runtime/AssetRegistry/Private/AssetDataGathererDiscoveryCache.cpp b/Engine/Source/Runtime/AssetRegistry/Private/AssetDataGathererDiscoveryCache.cpp
--- a/Engine/Source/Runtime/AssetRegistry/Private/AssetDataGathererDiscoveryCache.cpp
+++ b/Engine/Source/Runtime/AssetRegistry/Private/AssetDataGathererDiscoveryCache.cpp
@@ -35,7 +35,7 @@ namespace UE::AssetDataGather::Private
 
 FString FAssetDataDiscoveryCache::GetCacheFileName() const
 {
-	return FPaths::Combine(GPreloadSettings.GetAssetRegistryCacheRootFolder(),
+	return FPaths::Combine(GGatherSettings.GetAssetRegistryCacheRootFolder(),
 		TEXT("CachedAssetRegistryDiscovery.bin"));
 }
 
@@ -48,9 +48,9 @@ void FAssetDataDiscoveryCache::LoadAndUpdateCache()
 	TRACE_CPUPROFILER_EVENT_SCOPE(AssetDataGatherLoadDiscoveryCache);
 	bInitialized = true;
 
-	WriteEnabled = GPreloadSettings.IsDiscoveryCacheWriteEnabled();
+	WriteEnabled = GGatherSettings.IsDiscoveryCacheWriteEnabled();
 	CachedVolumes.Empty();
-	if (!GPreloadSettings.IsDiscoveryCacheReadEnabled() && WriteEnabled == EFeatureEnabled::Never)
+	if (!GGatherSettings.IsDiscoveryCacheReadEnabled() && WriteEnabled == EFeatureEnabled::Never)
 	{
 		return;
 	}
@@ -81,16 +81,16 @@ void FAssetDataDiscoveryCache::LoadAndUpdateCache()
 		}
 	}
 
-	bool bInvalidateEnabled = GPreloadSettings.IsDiscoveryCacheInvalidateEnabled();
-	bool bReadEnabled = GPreloadSettings.IsDiscoveryCacheReadEnabled() && (bPlatformSupported || !bInvalidateEnabled);
+	bool bInvalidateEnabled = GGatherSettings.IsDiscoveryCacheInvalidateEnabled();
+	bool bReadEnabled = GGatherSettings.IsDiscoveryCacheReadEnabled() && (bPlatformSupported || !bInvalidateEnabled);
 	// Precalculate IfPlatformSupported -> Never if we already know the project doesn't support it
 	if (WriteEnabled == EFeatureEnabled::IfPlatformSupported && !bPlatformSupported)
 	{
 		WriteEnabled = EFeatureEnabled::Never;
 	}
 
-	if ((!bReadEnabled && GPreloadSettings.IsDiscoveryCacheReadEnabled()) ||
-		(WriteEnabled == EFeatureEnabled::Never && GPreloadSettings.IsDiscoveryCacheWriteEnabled() != EFeatureEnabled::Never))
+	if ((!bReadEnabled && GGatherSettings.IsDiscoveryCacheReadEnabled()) ||
+		(WriteEnabled == EFeatureEnabled::Never && GGatherSettings.IsDiscoveryCacheWriteEnabled() != EFeatureEnabled::Never))
 	{
 		const TCHAR* MissingOperation = (!bReadEnabled && WriteEnabled == EFeatureEnabled::Never) ? TEXT("read or written") :
 			(!bReadEnabled ? TEXT("read") : TEXT("written"));
diff --git a/Engine/Source/Runtime/AssetRegistry/Private/AssetDataGathererPrivate.h b/Engine/Source/Runtime/AssetRegistry/Private/AssetDataGathererPrivate.h
--- a/Engine/Source/Runtime/AssetRegistry/Private/AssetDataGathererPrivate.h
+++ b/Engine/Source/Runtime/AssetRegistry/Private/AssetDataGathererPrivate.h
@@ -577,8 +577,7 @@ protected:
 class FAssetDataDiscovery : public FRunnable
 {
 public:
-	FAssetDataDiscovery(const TArray<FString>& InLongPackageNamesDenyList,
-		const TArray<FString>& InMountRelativePathsDenyList, bool bInAsyncEnabled);
+	FAssetDataDiscovery();
 	virtual ~FAssetDataDiscovery();
 
 
@@ -842,22 +841,72 @@ private:
 	friend class FScanDir;
 };
 
+/**
+* How the editor, when run with -game or (under more limited circumstances) -server, performs its asset registry scan
+*/
+enum class EEditorGameScanMode : uint8
+{
+	Sync,
+	Async,
+	None
+};
+
 /**
  * Settings about whether to use cache data for the AssetDataGatherer; these settings are shared by
  * FPreloader, FAssetDataGatherer, and FAssetDataDiscovery.
  */
-struct FPreloadSettings
+struct FGatherSettings
 {
 public:
+	static constexpr int32 SingleThreadFilesPerBatch = 3;
+	static constexpr int32 ExpectedMaxBatchSize = 100;
+	static constexpr int32 MinSecondsToElapseBeforeCacheWrite = 30;
+	static constexpr int32 MinAssetReadsBeforeCacheWrite = 100000;
+	static constexpr int32 CacheShardAssetCount = 250000;
+	static constexpr uint32 CacheSerializationMagic = 0xA46B0FDA; // Versioning and integrity checking
+	static constexpr uint64 CurrentVersion = FAssetRegistryVersion::LatestVersion | (uint64(CacheSerializationMagic) << 32);
+	static inline const FTimespan CachePruneAge = FTimespan(7, 0, 0, 0);
+
+	/** How long we wait on the gatherer thread between asking the discovery thread for no results. */
+	static constexpr float PollDiscoveryPeriodSeconds = .005f;
+
+	/**
+	 * The maximum number of threads used in the ParallelFor in FAssetDataDiscovery::TickInternal.
+	 * Higher values cause more time spent waking up workers. Lower values cause more time spent running the DirectoryEnumeration on each thread.
+	 * Current value was chosen based on experiments on an SSD drive with a warm disk cache.
+	 */
+	static inline int32 GARDiscoverThreads = 15;
+	/**
+	 * Minimum batch sized used used in the ParallelFor in FAssetDataDiscovery::TickInternal.
+	 * In theory higher minimum values could reduce the cost of waking up workers, but in practice it does not have a significant effect.
+	 */
+	static inline int32 GARDiscoverMinBatchSize = 1;
+	/**
+	 * The maximum number of threads used in the ParallelFor around ReadAssetFile in FAssetDataGatherer::TickInternal.
+	 * Very high numbers of threads does not reduce wall time of the gather because the threads become IO limited.
+	 * Capping the number of threads allows those threads to be used in other async operations in Engine startup.
+	 */
+	static inline int32 GARGatherThreads = 60;
+
+	/**
+	 * The maximum number of the gatherer's cache files that can be loaded in parallel in LoadCacheFiles.
+	 * The amount of memory used when reading the cache rises proportional to this value.
+	 * Other than the memory cost, it should be set as high as possible to load the cache faster.
+	 */
+	static inline int32 GARGatherCacheParallelism = 16;
+
 	void Initialize();
+	bool IsAsyncEnabled() const;
 	bool IsGatherCacheReadEnabled() const;
 	bool IsGatherCacheWriteEnabled() const;
 	bool IsDiscoveryCacheReadEnabled() const;
 	EFeatureEnabled IsDiscoveryCacheWriteEnabled() const;
 	bool IsDiscoveryCacheInvalidateEnabled() const;
-	bool IsPreloadCache() const;
+	bool IsPreloadGatherCache() const;
+	bool IsPreloadDiscoveryCache() const;
 	bool IsGatherDependsData() const;
 	bool IsForceDependsGathering() const;
+	EEditorGameScanMode GetEditorGameScanMode() const;
 	FString GetLegacyCacheFilename() const;
 	FString GetLegacyNonMonolithicCacheDirectory() const;
 	const FString& GetCacheBaseFilename() const;
@@ -874,12 +923,93 @@ private:
 	bool bDiscoveryCacheReadEnabled = false;
 	EFeatureEnabled DiscoveryCacheWriteEnabled = EFeatureEnabled::Never;
 	bool bDiscoveryCacheInvalidateEnabled = false;
-	bool bPreloadCache = false;
+	bool bPreloadGatherCache = false;
+	bool bAsyncEnabled = false;
 	bool bInitialized = false;
+	EEditorGameScanMode EditorGameScanMode = EEditorGameScanMode::None;
 };
-extern FPreloadSettings GPreloadSettings;
+extern FGatherSettings GGatherSettings;
 
 /** Checks if the given path, for a given file type, matches the naming conventions for said type. */
 bool DoesPathContainInvalidCharacters(const EGatherableFileType FileType, FStringView FilePath);
 
-} // namespace UE::AssetDataGather::Private
+///////////////////////////////////////////////////////
+// Inline implementations
+///////////////////////////////////////////////////////
+
+inline bool FGatherSettings::IsAsyncEnabled() const
+{
+	return bAsyncEnabled;
+}
+
+inline bool FGatherSettings::IsGatherCacheReadEnabled() const
+{
+	return bGatherCacheReadEnabled;
+}
+
+inline bool FGatherSettings::IsGatherCacheWriteEnabled() const
+{
+	return bGatherCacheWriteEnabled;
+}
+
+inline bool FGatherSettings::IsDiscoveryCacheReadEnabled() const
+{
+	return bDiscoveryCacheReadEnabled;
+}
+
+inline EFeatureEnabled FGatherSettings::IsDiscoveryCacheWriteEnabled() const
+{
+	return DiscoveryCacheWriteEnabled;
+}
+
+inline bool FGatherSettings::IsDiscoveryCacheInvalidateEnabled() const
+{
+	return bDiscoveryCacheInvalidateEnabled;
+}
+
+inline bool FGatherSettings::IsPreloadGatherCache() const
+{
+	return bGatherCacheReadEnabled && bPreloadGatherCache;
+}
+
+inline bool FGatherSettings::IsPreloadDiscoveryCache() const
+{
+	return IsDiscoveryCacheReadEnabled() && bPreloadGatherCache;
+}
+
+inline bool FGatherSettings::IsGatherDependsData() const
+{
+	return bGatherDependsData;
+}
+
+inline bool FGatherSettings::IsForceDependsGathering() const
+{
+	return bForceDependsGathering;
+}
+
+inline EEditorGameScanMode FGatherSettings::GetEditorGameScanMode() const
+{
+	return EditorGameScanMode;
+}
+
+inline FString FGatherSettings::GetLegacyCacheFilename() const
+{
+	return CacheBaseFilename + TEXT(".bin");
+}
+
+inline FString FGatherSettings::GetLegacyNonMonolithicCacheDirectory() const
+{
+	return GetAssetRegistryCacheRootFolder() / TEXT("AssetRegistryCache");
+}
+
+inline const FString& FGatherSettings::GetCacheBaseFilename() const
+{
+	return CacheBaseFilename;
+}
+
+inline const FString& FGatherSettings::GetAssetRegistryCacheRootFolder() const
+{
+	return AssetRegistryCacheRootFolder;
+}
+
+} // namespace UE::AssetDataGather::Private
\ No newline at end of file
diff --git a/Engine/Source/Runtime/AssetRegistry/Private/AssetRegistry.cpp b/Engine/Source/Runtime/AssetRegistry/Private/AssetRegistry.cpp
--- a/Engine/Source/Runtime/AssetRegistry/Private/AssetRegistry.cpp
+++ b/Engine/Source/Runtime/AssetRegistry/Private/AssetRegistry.cpp
@@ -4,6 +4,7 @@
 
 #include "Algo/Unique.h"
 #include "AssetDataGatherer.h"
+#include "AssetDataGathererPrivate.h"
 #include "AssetRegistry/ARFilter.h"
 #include "AssetRegistry/AssetDependencyGatherer.h"
 #include "AssetRegistry/AssetRegistryTelemetry.h"
@@ -207,43 +208,6 @@ namespace UE::AssetRegistry::Impl
 namespace UE::AssetRegistry
 {
 
-/**
- * How the editor, when run with -game or (under more limited circumstances) -server, performs its asset registry scan
- */
-enum class EEditorGameScanMode
-{
-	Sync,
-	Async,
-	None
-};
-
-EEditorGameScanMode GetEditorGameScanMode()
-{
-	FString EditorGameScansAR = TEXT("Sync");
-	EEditorGameScanMode Result = EEditorGameScanMode::Sync;
-
-	GConfig->GetString(TEXT("AssetRegistry"), TEXT("EditorGameScansAR"), EditorGameScansAR, GEngineIni);
-	if (EditorGameScansAR.Equals(TEXT("Async"), ESearchCase::IgnoreCase))
-	{
-		Result = EEditorGameScanMode::Async;
-	}
-	else if (EditorGameScansAR.Equals(TEXT("False"), ESearchCase::IgnoreCase))
-	{
-		Result = EEditorGameScanMode::None;
-	}
-	else if (EditorGameScansAR.Equals(TEXT("True"), ESearchCase::IgnoreCase))
-	{
-		Result = EEditorGameScanMode::Sync;
-	}
-	else
-	{
-		ensureMsgf(EditorGameScansAR.Equals(TEXT("Sync"), ESearchCase::IgnoreCase), 
-			TEXT("Valid values for EditorGameScansAR are: true|false|sync|async. Received %s"), *EditorGameScansAR);
-	}
-
-	return Result;
-}
-
 /** 
  * Tracks high-priority waiters on a lock, in an thread- and transactionally-safe manner.
  */
@@ -1931,23 +1895,7 @@ bool FAssetRegistryImpl::TryConstructGathererIfNeeded()
 		return false;
 	}
 
-	TArray<FString> PathsDenyList;
-	TArray<FString> ContentSubPathsDenyList;
-	if (FConfigFile* EngineIni = GConfig->FindConfigFile(GEngineIni))
-	{
-		EngineIni->GetArray(TEXT("AssetRegistry"), TEXT("BlacklistPackagePathScanFilters"), PathsDenyList);
-		EngineIni->GetArray(TEXT("AssetRegistry"), TEXT("BlacklistContentSubPathScanFilters"), ContentSubPathsDenyList);
-	}
-
-	// If EditorGameScanMode is set to allow async then we will use async for any editor build (editor, game, or server)
-	// Otherwise we will use async for commandlets and editor proper
-	bool bAsyncGatherEnabled = (!IsRunningGame() && !IsRunningDedicatedServer())
-#if WITH_EDITOR
-		|| (GetEditorGameScanMode() == EEditorGameScanMode::Async)
-#endif
-		;
-
-	GlobalGatherer = MakeUnique<FAssetDataGatherer>(PathsDenyList, ContentSubPathsDenyList, bAsyncGatherEnabled, *this);
+	GlobalGatherer = MakeUnique<FAssetDataGatherer>(*this);
 	UpdateMaxSecondsPerFrame();
 
 	// Read script packages if all initial plugins have been loaded, otherwise do nothing; we wait for the callback.
@@ -9902,7 +9850,8 @@ bool ShouldSearchAllAssetsAtStart()
 #if WITH_EDITOR
 	else
 	{
-		bSearchAllAssetsAtStart = GetEditorGameScanMode() != EEditorGameScanMode::None;
+		using namespace UE::AssetDataGather::Private;
+		bSearchAllAssetsAtStart = GGatherSettings.GetEditorGameScanMode() != EEditorGameScanMode::None;
 	}
 #endif
 #if WITH_EDITOR || !UE_BUILD_SHIPPING
