[ChaosVD] Fix compile error in CVD Standalone.

The issue was caused by a change made in the editor CL 41203803.

GPreloader is defined if UE_EDITOR is defined to 1, but the new code was using WITH_EDITOR (which can be defined to 1 when UE_EDITOR is not)

diff --git a/Engine/Source/Runtime/AssetRegistry/Private/AssetDataGatherer.cpp b/Engine/Source/Runtime/AssetRegistry/Private/AssetDataGatherer.cpp
--- a/Engine/Source/Runtime/AssetRegistry/Private/AssetDataGatherer.cpp
+++ b/Engine/Source/Runtime/AssetRegistry/Private/AssetDataGatherer.cpp
@@ -3760,7 +3760,7 @@ FAssetDataGatherer::FAssetDataGatherer(UE::AssetRegistry::FAssetRegistryImpl& In
 	// Tick is synchronous until StartAsync is called, even if bAsyncEnabled
 	bSynchronousTick = true; 
 
-#if WITH_EDITOR
+#if UE_EDITOR
 	// If the preloader has already created the FAssetDataDiscovery, take ownership of it now
 	if (TUniquePtr<FAssetDataDiscovery> PreloadedDiscovery = GPreloader.ConsumeDiscovery())
 	{
