[AssetRegistry]

Correct inverted condition for checking EditorGameScanMode from 41203803

#rnx

diff --git a/Engine/Source/Runtime/AssetRegistry/Private/AssetDataGatherer.cpp b/Engine/Source/Runtime/AssetRegistry/Private/AssetDataGatherer.cpp
--- a/Engine/Source/Runtime/AssetRegistry/Private/AssetDataGatherer.cpp
+++ b/Engine/Source/Runtime/AssetRegistry/Private/AssetDataGatherer.cpp
@@ -141,7 +141,7 @@ void LexFromString(EFeatureEnabledReadWrite& OutValue, FStringView Text)
 
 EEditorGameScanMode GetEditorGameScanModeFromConfig()
 {
-#if !WITH_EDITOR
+#if WITH_EDITOR
 	FString EditorGameScansAR = TEXT("Sync");
 	EEditorGameScanMode Result = EEditorGameScanMode::Sync;
 
