[AssetRegistry]

- Renames FFileJournalTelemetry to FFileJournalErrorTelemetry and adds new FFileJournalWrappedTelemetry
- Send telemetry when journaling is disabled in more return paths than previously.
- We now send telemetry when the journal wraps so we can better track inefficient setups using very small journals.
- Adds IPlatformFile::FileJournalGetMaximumSize to support fetching the maximum journal size (if any)
- Stores the maximum journal size as part of the FCachedVolumeInfo
- Provides recommended size when journal wraps that is inline with the recommended size offered when a journal doesn't exist
- Corrects 1GB SizeInBytes recommendation to be 1GB rather than slightly less (953MB or 1,000,000,000 bytes)

#rnx

diff --git a/Engine/Plugins/Experimental/EditorTelemetry/Source/Private/EditorTelemetry.cpp b/Engine/Plugins/Experimental/EditorTelemetry/Source/Private/EditorTelemetry.cpp
--- a/Engine/Plugins/Experimental/EditorTelemetry/Source/Private/EditorTelemetry.cpp
+++ b/Engine/Plugins/Experimental/EditorTelemetry/Source/Private/EditorTelemetry.cpp
@@ -1036,7 +1036,7 @@ void FEditorTelemetry::StartSession()
 				{ TEXT("InitialSearchCompleted"), Data.bInitialSearchCompleted },
 			});
 		});
-		Router.OnTelemetry<FFileJournalTelemetry>([this](const FFileJournalTelemetry& Data){
+		Router.OnTelemetry<FFileJournalErrorTelemetry>([this](const FFileJournalErrorTelemetry& Data){
 			FStudioTelemetry::Get().RecordEvent(TEXT("Editor.AssetRegistry.FileJournalError"),
 			{
 				{ TEXT("SchemaVersion"), 1 },
@@ -1044,6 +1044,14 @@ void FEditorTelemetry::StartSession()
 				{ TEXT("ErrorString"), Data.ErrorString },
 			});
 		});
+		Router.OnTelemetry<FFileJournalWrappedTelemetry>([this](const FFileJournalWrappedTelemetry& Data) {
+			FStudioTelemetry::Get().RecordEvent(TEXT("Editor.AssetRegistry.FileJournalWrapped"),
+				{
+					{ TEXT("SchemaVersion"), 1 },
+					{ TEXT("VolumeName"), Data.VolumeName},
+					{ TEXT("JournalMaximumSize"), Data.JournalMaximumSize},
+				});
+			});
 	}
 
 	{
diff --git a/Engine/Source/Runtime/AssetRegistry/Private/AssetDataGatherer.cpp b/Engine/Source/Runtime/AssetRegistry/Private/AssetDataGatherer.cpp
--- a/Engine/Source/Runtime/AssetRegistry/Private/AssetDataGatherer.cpp
+++ b/Engine/Source/Runtime/AssetRegistry/Private/AssetDataGatherer.cpp
@@ -199,6 +199,15 @@ void FGatherSettings::Initialize()
 		|| !GIsEditor;
 	bGatherCacheReadEnabled = !bNoAssetRegistryCache && !bNoAssetRegistryCacheRead;
 	bGatherCacheWriteEnabled = !bNoAssetRegistryCache && !bNoAssetRegistryCacheWrite;
+
+	if (IsInGameThread())
+	{
+		// Ensure the TelemetryRouter is initialized from the game thread. It may need to load modules upon first access, so we 
+		// fetch the singleton here to allow the module loading to remain an implementation detail. 
+		// Modules can only be loaded from the game thread so we ensure we only attempt this during initialization on the game thread.
+		static_cast<void>(FTelemetryRouter::Get());
+	}
+
 	bool bPlatformSupportsDiscoveryCacheInvalidation = FPlatformFileManager::Get().GetPlatformFile().FileJournalIsAvailable();
 
 	bool bSkipInvalidate = FParse::Param(FCommandLine::Get(), TEXT("AssetRegistryCacheSkipInvalidate"));
@@ -2189,7 +2198,7 @@ void FAssetDataDiscovery::TickInternal(bool bTickAll)
 
 							UE_LOG(LogAssetRegistry, Warning, TEXT("Failed to FileJournalIterateDirectory, failing to cache this directory due to:\n  %s"), *OutError);
 
-							UE::Telemetry::AssetRegistry::FFileJournalTelemetry Telemetry;
+							UE::Telemetry::AssetRegistry::FFileJournalErrorTelemetry Telemetry;
 							Telemetry.Directory = Data.DirLocalAbsPath.ToString();
 							Telemetry.ErrorString = OutError;
 							FTelemetryRouter::Get().ProvideTelemetry(Telemetry);
diff --git a/Engine/Source/Runtime/AssetRegistry/Private/AssetDataGathererDiscoveryCache.cpp b/Engine/Source/Runtime/AssetRegistry/Private/AssetDataGathererDiscoveryCache.cpp
--- a/Engine/Source/Runtime/AssetRegistry/Private/AssetDataGathererDiscoveryCache.cpp
+++ b/Engine/Source/Runtime/AssetRegistry/Private/AssetDataGathererDiscoveryCache.cpp
@@ -74,7 +74,7 @@ void FAssetDataDiscoveryCache::LoadAndUpdateCache()
 		{
 			bPlatformSupported = false;
 
-			UE::Telemetry::AssetRegistry::FFileJournalTelemetry Telemetry;
+			UE::Telemetry::AssetRegistry::FFileJournalErrorTelemetry Telemetry;
 			Telemetry.Directory = FPaths::ConvertRelativePathToFull(ProjectDir);
 			Telemetry.ErrorString = TestError;
 			FTelemetryRouter::Get().ProvideTelemetry(Telemetry);
@@ -97,6 +97,11 @@ void FAssetDataDiscoveryCache::LoadAndUpdateCache()
 		UE_LOG(LogAssetRegistry, Display,
 			TEXT("PlatformFileJournal is not available on volume '%s' of project directory '%s', so AssetDiscovery cache will not be %s. Unavailability reason:\n\t%s"),
 			*TestVolumeName, *FPaths::ConvertRelativePathToFull(ProjectDir), MissingOperation, *TestError);
+
+		UE::Telemetry::AssetRegistry::FFileJournalErrorTelemetry Telemetry;
+		Telemetry.Directory = FPaths::ConvertRelativePathToFull(ProjectDir);
+		Telemetry.ErrorString = TestError;
+		FTelemetryRouter::Get().ProvideTelemetry(Telemetry);
 	}
 
 	if (!bReadEnabled)
@@ -141,6 +146,11 @@ void FAssetDataDiscoveryCache::LoadAndUpdateCache()
 			UE_LOG(LogAssetRegistry, Warning,
 				TEXT("PlatformFileJournal is not available on volume '%s'. AssetRegistry discovery of files on this volume will be uncached. Unavailability reason:\n\t%s"),
 				*VolumeName, *VolumeInfo.LastError);
+
+			UE::Telemetry::AssetRegistry::FFileJournalErrorTelemetry MissingJournalTelemetry;
+			MissingJournalTelemetry.Directory = VolumeName;
+			MissingJournalTelemetry.ErrorString = VolumeInfo.LastError;
+			FTelemetryRouter::Get().ProvideTelemetry(MissingJournalTelemetry);
 		}
 		else
 		{
@@ -153,18 +163,32 @@ void FAssetDataDiscoveryCache::LoadAndUpdateCache()
 				bReadModifiedSucceeded = true;
 				break;
 			case EFileJournalResult::JournalWrapped:
+			{
 				UE_LOG(LogAssetRegistry, Warning,
 					TEXT("PlatformFileJournal journal has wrapped for volume '%s'. AssetRegistry discovery of files on this volume will be uncached. Notes on wrapping:")
 					TEXT("\r\n%s"),
 					*VolumeName, *VolumeInfo.LastError);
+
+				UE::Telemetry::AssetRegistry::FFileJournalWrappedTelemetry WrappedJournalTelemetry;
+				WrappedJournalTelemetry.VolumeName = VolumeName;
+				WrappedJournalTelemetry.JournalMaximumSize = VolumeInfo.JournalMaximumSize;
+				FTelemetryRouter::Get().ProvideTelemetry(WrappedJournalTelemetry);
 				break;
+			}
 			default:
+			{
 				UE_LOG(LogAssetRegistry, Warning,
 					TEXT("PlatformFileJournal is not available for volume '%s'. AssetRegistry discovery of files on this volume will be uncached. Unavailability reason:")
 					TEXT("\n\t%s"),
 					*VolumeName, *VolumeInfo.LastError);
+
+				UE::Telemetry::AssetRegistry::FFileJournalErrorTelemetry MissingJournalTelemetry;
+				MissingJournalTelemetry.Directory = VolumeName;
+				MissingJournalTelemetry.ErrorString = VolumeInfo.LastError;
+				FTelemetryRouter::Get().ProvideTelemetry(MissingJournalTelemetry);
 				break;
 			}
+			}
 		}
 
 		if (!bReadModifiedSucceeded)
@@ -548,6 +572,7 @@ void FCachedVolumeInfo::InitializePlatformData()
 	if (VolumeName.IsEmpty() || VolumeName == GEmptyVolumeName)
 	{
 		bJournalAvailable = false;
+		JournalMaximumSize = 0;
 		JournalIdOnDisk = FileJournalIdInvalid;
 		NextJournalEntryOnDisk = FileJournalEntryHandleInvalid;
 		JournalId = FileJournalIdInvalid;
@@ -558,6 +583,7 @@ void FCachedVolumeInfo::InitializePlatformData()
 		EFileJournalResult Result = FPlatformFileManager::Get().GetPlatformFile().FileJournalGetLatestEntry(
 			*VolumeName, JournalIdOnDisk, NextJournalEntryOnDisk, &LastError);
 		bJournalAvailable = Result == EFileJournalResult::Success;
+		JournalMaximumSize = bJournalAvailable ? FPlatformFileManager::Get().GetPlatformFile().FileJournalGetMaximumSize(*VolumeName) : 0;
 		if (NextJournalEntryToScan == FileJournalEntryHandleInvalid)
 		{
 			JournalId = JournalIdOnDisk;
diff --git a/Engine/Source/Runtime/AssetRegistry/Private/AssetDataGathererDiscoveryCache.h b/Engine/Source/Runtime/AssetRegistry/Private/AssetDataGathererDiscoveryCache.h
--- a/Engine/Source/Runtime/AssetRegistry/Private/AssetDataGathererDiscoveryCache.h
+++ b/Engine/Source/Runtime/AssetRegistry/Private/AssetDataGathererDiscoveryCache.h
@@ -113,6 +113,7 @@ public:
 	FString LastError;
 	FFileJournalId JournalIdOnDisk = FileJournalIdInvalid;
 	FFileJournalEntryHandle NextJournalEntryOnDisk = FileJournalEntryHandleInvalid;
+	uint64 JournalMaximumSize = 0;
 	bool bJournalAvailable = false;
 };
 FArchive& operator<<(FArchive& Ar, FCachedVolumeInfo& Data);
diff --git a/Engine/Source/Runtime/AssetRegistry/Public/AssetRegistry/AssetRegistryTelemetry.h b/Engine/Source/Runtime/AssetRegistry/Public/AssetRegistry/AssetRegistryTelemetry.h
--- a/Engine/Source/Runtime/AssetRegistry/Public/AssetRegistry/AssetRegistryTelemetry.h
+++ b/Engine/Source/Runtime/AssetRegistry/Public/AssetRegistry/AssetRegistryTelemetry.h
@@ -88,7 +88,7 @@ namespace UE::Telemetry::AssetRegistry
     };
 
 	// Information about a failure in our FileJournal area
-	struct FFileJournalTelemetry
+	struct FFileJournalErrorTelemetry
 	{
 		static constexpr inline FGuid TelemetryID = FGuid(0x761ebe19, 0x2f0c427e, 0xaec0f06d, 0xe116ff82);
 
@@ -98,4 +98,13 @@ namespace UE::Telemetry::AssetRegistry
 		// The error string returned from our lower level FileJournal calls
 		FString ErrorString;
 	};
+
+	struct FFileJournalWrappedTelemetry
+	{
+		static constexpr inline FGuid TelemetryID = FGuid(0x542713AA, 0x04724D77, 0x8F7E2319, 0xA0204815);
+
+		FString VolumeName;
+		// Size of journal that wrapped
+		uint64 JournalMaximumSize;
+	};
 }
diff --git a/Engine/Source/Runtime/Core/Private/GenericPlatform/GenericPlatformFile.cpp b/Engine/Source/Runtime/Core/Private/GenericPlatform/GenericPlatformFile.cpp
--- a/Engine/Source/Runtime/Core/Private/GenericPlatform/GenericPlatformFile.cpp
+++ b/Engine/Source/Runtime/Core/Private/GenericPlatform/GenericPlatformFile.cpp
@@ -1080,6 +1080,20 @@ bool IPlatformFile::FileJournalIsAvailable(const TCHAR* VolumeOrPath, ELogVerbos
 	return false;
 }
 
+uint64 IPlatformFile::FileJournalGetMaximumSize(const TCHAR* VolumeOrPath, ELogVerbosity::Type* OutErrorLevel,
+	FString* OutError) const
+{
+	if (OutErrorLevel)
+	{
+		*OutErrorLevel = ELogVerbosity::Display;
+	}
+	if (OutError)
+	{
+		*OutError = UE::PlatformFileJournal::Private::PlatformNotAvailableMessage;
+	}
+	return 0;
+}
+
 EFileJournalResult IPlatformFile::FileJournalGetLatestEntry(const TCHAR* VolumeOrPath,
 	FFileJournalId& OutJournalId, FFileJournalEntryHandle& OutEntryHandle, FString* OutError)
 {
diff --git a/Engine/Source/Runtime/Core/Private/Windows/WindowsPlatformFileJournal.cpp b/Engine/Source/Runtime/Core/Private/Windows/WindowsPlatformFileJournal.cpp
--- a/Engine/Source/Runtime/Core/Private/Windows/WindowsPlatformFileJournal.cpp
+++ b/Engine/Source/Runtime/Core/Private/Windows/WindowsPlatformFileJournal.cpp
@@ -158,6 +158,42 @@ bool FWindowsPlatformFile::FileJournalIsAvailable(const TCHAR* VolumeOrPath, ELo
 	}
 }
 
+uint64 FWindowsPlatformFile::FileJournalGetMaximumSize(const TCHAR* VolumeOrPath, ELogVerbosity::Type* OutErrorLevel,
+	FString* OutError) const
+{
+	using namespace UE::WindowsPlatformFileJournal::Private;
+
+	if (!VolumeOrPath)
+	{
+		return 0;
+	}
+
+	HANDLE JournalHandle;
+	TStringBuilder<16> NormalizedVolumeName;
+	EFileJournalResult Result = CreateJournalWindowsHandle(VolumeOrPath, JournalHandle, NormalizedVolumeName, OutError);
+	if (Result != EFileJournalResult::Success)
+	{
+		if (OutErrorLevel)
+		{
+			*OutErrorLevel = ELogVerbosity::Error;
+		}
+		return 0;
+	}
+	ON_SCOPE_EXIT{ ::CloseHandle(JournalHandle); };
+
+	USN_JOURNAL_DATA_V1 JournalDescriptor;
+	Result = GetJournalDescriptor(JournalHandle, *NormalizedVolumeName, JournalDescriptor, OutError);
+	if (Result != EFileJournalResult::Success)
+	{
+		if (OutErrorLevel)
+		{
+			*OutErrorLevel = ELogVerbosity::Error;
+		}
+		return 0;
+	}
+	
+	return JournalDescriptor.MaximumSize;
+}
 
 EFileJournalResult FWindowsPlatformFile::FileJournalGetLatestEntry(const TCHAR* VolumeName,
 	FFileJournalId& OutJournalId, FFileJournalEntryHandle& OutEntryHandle, FString* OutError)
@@ -521,7 +557,7 @@ EFileJournalResult FWindowsPlatformFile::FileJournalReadModified(const TCHAR* Vo
 						*OutError = FString::Printf(TEXT("NTFS Journal has wrapped for volume '%s'; the number of files written to this volume exceeded the buffer size of the journal and some entries were dropped.")
 							TEXT("\r\nIf this happens frequently increase the buffer size of the journal:")
 							TEXT("\r\nLaunch cmd.exe as admin and run command `fsutil usn createJournal %s m=<SizeInBytes>`")
-							TEXT("\r\nSizeInBytes defaults to 33554432 (32MB). You can see the current value via `fsutil usn queryJournal %s`, 'Maximum Size' field."),
+							TEXT("\r\n<SizeInBytes> defaults to 33554432 (32MB), Recommended <SizeInBytes> is 0x40000000 (1GB). You can see the current value via `fsutil usn queryJournal %s`, 'Maximum Size' field."),
 							VolumeName, VolumeName, VolumeName);
 					}
 					return EFileJournalResult::JournalWrapped;
@@ -794,7 +830,7 @@ EFileJournalResult GetJournalDescriptor(HANDLE JournalWindowsHandle, const TCHAR
 				// Can also create using DeviceIoControl with FSCTL_CREATE_USN_JOURNAL
 				*OutError = FString::Printf(TEXT("NTFS Journal is not active for volume '%s'. ")
 					TEXT("Launch cmd.exe as admin and run command `fsutil usn createJournal %s m=<SizeInBytes>`. ")
-					TEXT("Recommended <SizeInBytes> is 1000000000 (1GB)."),
+					TEXT("Recommended <SizeInBytes> is 0x40000000 (1GB)."),
 					VolumeName, VolumeName);
 			}
 			return EFileJournalResult::JournalNotActive;
@@ -920,6 +956,27 @@ bool FWindowsPlatformFile::FileJournalIsAvailable(const TCHAR* VolumeOrPath, ELo
 	return IPlatformFile::FileJournalIsAvailable(VolumeOrPath, OutErrorLevel, OutError);
 }
 
+uint64 FWindowsPlatformFile::FileJournalGetMaximumSize(const TCHAR* VolumeOrPath, ELogVerbosity::Type* OutErrorLevel,
+	FString* OutError) const
+{
+	using namespace UE::WindowsPlatformFileJournal::Private;
+
+#if !WITH_EDITOR
+	if (OutErrorLevel)
+	{
+		*OutErrorLevel = ELogVerbosity::Display;
+		OutErrorLevel = nullptr;
+	}
+	if (OutError)
+	{
+		*OutError = OnlyAvailableInEditorMessage;
+		OutError = nullptr;
+	}
+#endif
+
+	return IPlatformFile::FileJournalGetMaximumSize(VolumeOrPath, OutErrorLevel, OutError);
+}
+
 EFileJournalResult  FWindowsPlatformFile::FileJournalGetLatestEntry(const TCHAR* VolumeName,
 	FFileJournalId& OutJournalId, FFileJournalEntryHandle& OutEntryHandle, FString* OutError)
 {
diff --git a/Engine/Source/Runtime/Core/Private/Windows/WindowsPlatformFilePrivate.h b/Engine/Source/Runtime/Core/Private/Windows/WindowsPlatformFilePrivate.h
--- a/Engine/Source/Runtime/Core/Private/Windows/WindowsPlatformFilePrivate.h
+++ b/Engine/Source/Runtime/Core/Private/Windows/WindowsPlatformFilePrivate.h
@@ -82,6 +82,8 @@ public:
 
 	virtual bool FileJournalIsAvailable(const TCHAR* VolumeOrPath = nullptr, ELogVerbosity::Type* OutErrorLevel = nullptr,
 		FString* OutError = nullptr) override;
+	virtual uint64 FileJournalGetMaximumSize(const TCHAR* VolumeOrPath, ELogVerbosity::Type* OutErrorLevel = nullptr,
+		FString* OutError = nullptr) const override;
 	virtual EFileJournalResult FileJournalGetLatestEntry(const TCHAR* VolumeName, FFileJournalId& OutJournalId, 
 		FFileJournalEntryHandle& OutEntryHandle, FString* OutError = nullptr) override;
 	virtual bool FileJournalIterateDirectory(const TCHAR* Directory, FDirectoryJournalVisitorFunc Visitor, FString* OutError = nullptr) override;
diff --git a/Engine/Source/Runtime/Core/Public/GenericPlatform/GenericPlatformFile.h b/Engine/Source/Runtime/Core/Public/GenericPlatform/GenericPlatformFile.h
--- a/Engine/Source/Runtime/Core/Public/GenericPlatform/GenericPlatformFile.h
+++ b/Engine/Source/Runtime/Core/Public/GenericPlatform/GenericPlatformFile.h
@@ -642,6 +642,14 @@ public:
 	CORE_API virtual bool FileJournalIsAvailable(const TCHAR* VolumeOrPath = nullptr,
 		ELogVerbosity::Type* OutErrorLevel = nullptr, FString* OutError = nullptr);
 
+	/**
+	* Return the maximum size of the FileJournal for the specified VolumeName, if available, zero otherwise. Optionally returns
+	* a user-displayable string for why the journal is not available and a severity level for the reason. VolumeName 
+	* may be a VolumeName as returned by FileJournalGetVolumeName or any path to a file or directory on the Volume.
+	*/
+	CORE_API virtual uint64 FileJournalGetMaximumSize(const TCHAR* VolumeOrPath = nullptr,
+		ELogVerbosity::Type* OutErrorLevel = nullptr, FString* OutError = nullptr) const;
+
 	/**
 	 * Report the current end of the journal for the given volume, to be used as the StartingJournalEntry
 	 * in FileJournalGetModifiedDirectories. If !FileJournalIsAvaiable for the given volume, sets OutEntryHandle
