Correctly deprecate GetAssetByObjectPath in blueprint and reset the clock so it doesn't get deleted while still in use. This is still called in a lot of BP utilities so I set it to not warn on compile.
Delete some older deprecated functions

diff --git a/Engine/Source/Runtime/AssetRegistry/Private/AssetRegistry.cpp b/Engine/Source/Runtime/AssetRegistry/Private/AssetRegistry.cpp
--- a/Engine/Source/Runtime/AssetRegistry/Private/AssetRegistry.cpp
+++ b/Engine/Source/Runtime/AssetRegistry/Private/AssetRegistry.cpp
@@ -4473,11 +4473,6 @@ void UAssetRegistryImpl::AssetRenamed(const UObject* RenamedAsset, const FString
 	}
 }
 
-void UAssetRegistryImpl::AssetSaved(const UObject& SavedAsset)
-{
-	AssetUpdateTags(const_cast<UObject*>(&SavedAsset), EAssetRegistryTagsCaller::Fast);
-}
-
 void UAssetRegistryImpl::AssetsSaved(TArray<FAssetData>&& Assets)
 {
 #if WITH_EDITOR
@@ -5999,28 +5994,12 @@ void UAssetRegistryImpl::DumpState(const TArray<FString>& Arguments, TArray<FStr
 }
 #endif
 
-const FAssetRegistryState* UAssetRegistryImpl::GetAssetRegistryState() const
-{
-	UE::AssetRegistry::FInterfaceReadScopeLock InterfaceScopeLock(InterfaceLock);
-	PRAGMA_DISABLE_DEPRECATION_WARNINGS
-	return &GuardedData.GetState();
-	PRAGMA_ENABLE_DEPRECATION_WARNINGS
-}
-
 TSet<FName> UAssetRegistryImpl::GetCachedEmptyPackagesCopy() const
 {
 	UE::AssetRegistry::FInterfaceReadScopeLock InterfaceScopeLock(InterfaceLock);
 	return GuardedData.GetCachedEmptyPackages();
 }
 
-const TSet<FName>& UAssetRegistryImpl::GetCachedEmptyPackages() const
-{
-	UE::AssetRegistry::FInterfaceReadScopeLock InterfaceScopeLock(InterfaceLock);
-	PRAGMA_DISABLE_DEPRECATION_WARNINGS
-	return GuardedData.GetCachedEmptyPackages();
-	PRAGMA_ENABLE_DEPRECATION_WARNINGS
-}
-
 bool UAssetRegistryImpl::ContainsTag(FName TagName) const
 {
 	UE::AssetRegistry::FInterfaceReadScopeLock InterfaceScopeLock(InterfaceLock);
diff --git a/Engine/Source/Runtime/AssetRegistry/Private/AssetRegistry.h b/Engine/Source/Runtime/AssetRegistry/Private/AssetRegistry.h
--- a/Engine/Source/Runtime/AssetRegistry/Private/AssetRegistry.h
+++ b/Engine/Source/Runtime/AssetRegistry/Private/AssetRegistry.h
@@ -122,10 +122,7 @@ public:
 #if ASSET_REGISTRY_STATE_DUMPING_ENABLED
 	virtual void DumpState(const TArray<FString>& Arguments, TArray<FString>& OutPages, int32 LinesPerPage = 1) const override;
 #endif
-
-	virtual const FAssetRegistryState* GetAssetRegistryState() const override;
 	virtual TSet<FName> GetCachedEmptyPackagesCopy() const override;
-	virtual const TSet<FName>& GetCachedEmptyPackages() const override;
 	virtual bool ContainsTag(FName TagName) const override;
 	virtual void InitializeSerializationOptions(FAssetRegistrySerializationOptions& Options, const FString& PlatformIniName = FString(), UE::AssetRegistry::ESerializationTarget Target = UE::AssetRegistry::ESerializationTarget::ForGame) const override;
 
@@ -154,8 +151,6 @@ public:
 	virtual void AssetCreated(UObject* NewAsset) override;
 	virtual void AssetDeleted(UObject* DeletedAsset) override;
 	virtual void AssetRenamed(const UObject* RenamedAsset, const FString& OldObjectPath) override;
-	UE_DEPRECATED(5.2, "Use the new AssetsSaved function that takes FAssetData.")
-	virtual void AssetSaved(const UObject& SavedAsset) override;
 	virtual void AssetsSaved(TArray<FAssetData>&& SavedAssets) override;
 	virtual void AssetUpdateTags(UObject* Object, EAssetRegistryTagsCaller Caller) override;
 	UE_DEPRECATED(5.4, "Call AssetUpdateTags with EAssetRegistryTagsCaller::Fast")
diff --git a/Engine/Source/Runtime/AssetRegistry/Public/AssetRegistry/AssetRegistryModule.h b/Engine/Source/Runtime/AssetRegistry/Public/AssetRegistry/AssetRegistryModule.h
--- a/Engine/Source/Runtime/AssetRegistry/Public/AssetRegistry/AssetRegistryModule.h
+++ b/Engine/Source/Runtime/AssetRegistry/Public/AssetRegistry/AssetRegistryModule.h
@@ -73,14 +73,6 @@ public:
 		IAssetRegistry::GetChecked().AssetRenamed(RenamedAsset, OldObjectPath);
 	}
 
-	UE_DEPRECATED(5.2, "Use the new AssetsSaved function that takes FAssetData.")
-	static void AssetSaved(const UObject& SavedAsset)
-	{
-		PRAGMA_DISABLE_DEPRECATION_WARNINGS;
-		IAssetRegistry::GetChecked().AssetSaved(SavedAsset);
-		PRAGMA_ENABLE_DEPRECATION_WARNINGS;
-	}
-
 	static void AssetsSaved(TArray<FAssetData>&& SavedAssets)
 	{
 		IAssetRegistry::GetChecked().AssetsSaved(MoveTemp(SavedAssets));
diff --git a/Engine/Source/Runtime/AssetRegistry/Public/AssetRegistry/IAssetRegistry.h b/Engine/Source/Runtime/AssetRegistry/Public/AssetRegistry/IAssetRegistry.h
--- a/Engine/Source/Runtime/AssetRegistry/Public/AssetRegistry/IAssetRegistry.h
+++ b/Engine/Source/Runtime/AssetRegistry/Public/AssetRegistry/IAssetRegistry.h
@@ -320,16 +320,8 @@ public:
 	virtual bool EnumerateAssets(const FARCompiledFilter& Filter, TFunctionRef<bool(const FAssetData&)> Callback,
 		UE::AssetRegistry::EEnumerateAssetsFlags InEnumerateFlags) const = 0;
 
-	/**
-	 * Gets the asset data for the specified object path
-	 *
-	 * @param ObjectPath the path of the object to be looked up
-	 * @param bIncludeOnlyOnDiskAssets If true, use only DiskGatheredData, do not calculate from UObjects.
-	 *        @see IAssetRegistry class header for bIncludeOnlyOnDiskAssets.
-	 * @return the assets data;Will be invalid if object could not be found
-	 */
-	UE_DEPRECATED(5.1, "Asset path FNames have been deprecated, use Soft Object Path instead.")
-	UFUNCTION(BlueprintCallable, BlueprintPure=false, Category="AssetRegistry")
+	UE_DEPRECATED(5.6, "Asset path FNames have been deprecated, use Soft Object Path instead.")
+	UFUNCTION(BlueprintCallable, BlueprintPure=false, Category="AssetRegistry", meta=(DeprecatedFunction="Note", DeprecationMessage="Use other version that takes a SoftObjectPath"))
 	virtual FAssetData GetAssetByObjectPath( const FName ObjectPath, bool bIncludeOnlyOnDiskAssets = false ) const = 0;
 
 	/**
@@ -817,10 +809,6 @@ public:
 	/** Informs the asset registry that an in-memory asset has been renamed */
 	virtual void AssetRenamed(const UObject* RenamedAsset, const FString& OldObjectPath) = 0;
 
-	/** Informs the asset registry that an in-memory asset has been saved */
-	UE_DEPRECATED(5.2, "Use the new AssetsSaved function that takes FAssetData.")
-	virtual void AssetSaved(const UObject& SavedAsset) = 0;
-
 	/**
 	 * Called during SavePackage to update the AssetRegistry's copy of the AssetDatas in the package to match the
 	 * newly saved values.
@@ -998,9 +986,6 @@ public:
 		bool bRefreshExisting = false, const TSet<FName>& RequiredPackages = TSet<FName>(),
 		const TSet<FName>& RemovePackages = TSet<FName>()) const = 0;
 
-	UE_DEPRECATED(5.0, "Receiving a pointer is not threadsafe. Use other functions on IAssetRegistry to access the same data, or contact Epic Core team to add the threadsafe functions you require..")
-	virtual const FAssetRegistryState* GetAssetRegistryState() const = 0;
-
 #if ASSET_REGISTRY_STATE_DUMPING_ENABLED
 	/**
 	 * Writes out the state in textual form. Use arguments to control which segments to emit.
@@ -1018,8 +1003,6 @@ public:
 
 	/** Returns the set of empty package names fast iteration */
 	virtual TSet<FName> GetCachedEmptyPackagesCopy() const = 0;
-	UE_DEPRECATED(5.0, "Receiving a reference is not threadsafe. Use GetCachedEmptyPackagesCopy instead.")
-	virtual const TSet<FName>& GetCachedEmptyPackages() const = 0;
 
 	/** Return whether the given TagName occurs in the tags of any asset in the AssetRegistry */
 	virtual bool ContainsTag(FName TagName) const = 0;
