AssetRegistry: Gatherer: Fix SetPropertiesAndWait's use of a dangling FStringView pointer when the QueryPath absolute path is reallocated by LoadExistenceData.
#rnx

diff --git a/Engine/Source/Runtime/AssetRegistry/Private/AssetDataGatherer.cpp b/Engine/Source/Runtime/AssetRegistry/Private/AssetDataGatherer.cpp
--- a/Engine/Source/Runtime/AssetRegistry/Private/AssetDataGatherer.cpp
+++ b/Engine/Source/Runtime/AssetRegistry/Private/AssetDataGatherer.cpp
@@ -2522,9 +2522,9 @@ void FAssetDataDiscovery::SetPropertiesAndWait(TArrayView<FPathExistence> QueryP
 				Properties.IgnoreDenyList = true;
 			}
 			FScanDirAndParentData ControlDirData;
-			FStringView RelPathFromControlDir;
+			FStringView RelPathFromControlDirView;
 			bool bMadeChanges;
-			MountDir->TrySetDirectoryProperties(QueryPath, Properties, &ControlDirData, &RelPathFromControlDir, &bMadeChanges);
+			MountDir->TrySetDirectoryProperties(QueryPath, Properties, &ControlDirData, &RelPathFromControlDirView, &bMadeChanges);
 			if (bMadeChanges)
 			{
 				SetIsIdle(false);
@@ -2577,6 +2577,9 @@ void FAssetDataDiscovery::SetPropertiesAndWait(TArrayView<FPathExistence> QueryP
 				continue;
 			}
 
+			// Save a copy of RelPathFromControlDirView; it points into QueryPath.AbsolutePath which might change
+			FString RelPathFromControlDir(RelPathFromControlDirView);
+
 			// After this point we are no longer in the common case of having no work to do, so now we can pay the IO
 			// cost of loading the querypath's existence data, if we haven't already encountered that not-in-common-case
 			// and fetched the data inside of TrySetDirectoryProperties.
