AssetRegistry: Gatherer: Remove warning when FScanDir is constructed while bIsIdle. We now sometimes need to trigger scandir construction and set non-idle afterwards, in SetPropertiesAndWait.
#rnx

diff --git a/Engine/Source/Runtime/AssetRegistry/Private/AssetDataGatherer.cpp b/Engine/Source/Runtime/AssetRegistry/Private/AssetDataGatherer.cpp
--- a/Engine/Source/Runtime/AssetRegistry/Private/AssetDataGatherer.cpp
+++ b/Engine/Source/Runtime/AssetRegistry/Private/AssetDataGatherer.cpp
@@ -462,7 +462,6 @@ FScanDir::FScanDir(FMountDir& InMountDir, FScanDir* InParent, FStringView InRelP
 {
 	FAssetDataDiscovery& Discovery = InMountDir.GetDiscovery();
 	Discovery.NumDirectoriesToScan.Increment();
-	UE_CLOG(Discovery.bIsIdle, LogAssetRegistry, Warning, TEXT("AssetDataGatherer: FScanDir is constructed while bIsIdle=true."));
 }
 
 FScanDir::~FScanDir()
