AssetRegistry: Stop inheriting from F[TransactionallySafe]RWLock to allow making the mutex types final

diff --git a/Engine/Source/Runtime/AssetRegistry/Private/AssetRegistry.cpp b/Engine/Source/Runtime/AssetRegistry/Private/AssetRegistry.cpp
--- a/Engine/Source/Runtime/AssetRegistry/Private/AssetRegistry.cpp
+++ b/Engine/Source/Runtime/AssetRegistry/Private/AssetRegistry.cpp
@@ -298,9 +298,9 @@ public:
 		}
 	}
 
-	UE::AssetRegistry::Private::FRWLockWithPriority& Lock;
-	TOptional<UE::TRWScopeLock<FTransactionallySafeRWLock>> GuardWrapper;
-	UE::AssetRegistry::Private::ELockPriority Priority;
+	Private::FRWLockWithPriority& Lock;
+	TOptional<TRWScopeLock<Private::FRWLockWithPriority>> GuardWrapper;
+	Private::ELockPriority Priority;
 	FRWScopeLockType LockType;
 };
 
diff --git a/Engine/Source/Runtime/AssetRegistry/Private/AssetRegistryImpl.h b/Engine/Source/Runtime/AssetRegistry/Private/AssetRegistryImpl.h
--- a/Engine/Source/Runtime/AssetRegistry/Private/AssetRegistryImpl.h
+++ b/Engine/Source/Runtime/AssetRegistry/Private/AssetRegistryImpl.h
@@ -44,9 +44,11 @@ namespace UE::AssetRegistry::Impl { struct FInitializeContext; }
 namespace UE::AssetRegistry::Impl { struct FScanPathContext; }
 namespace UE::AssetRegistry::Impl { struct FTickContext; }
 namespace UE::AssetRegistry::Impl { struct FClassInheritanceContext; }
+namespace UE::AssetRegistry::Private { class FRWLockWithPriority; }
+namespace UE::AssetRegistry::Private { using FInterfaceRWLock = FRWLockWithPriority; }
 namespace UE::AssetRegistry { template <typename TScopeLockType> class TRWScopeLockWithPriority; }
-namespace UE::AssetRegistry { using FInterfaceReadScopeLock = TRWScopeLockWithPriority<UE::TReadScopeLock<FTransactionallySafeRWLock>>; }
-namespace UE::AssetRegistry { using FInterfaceWriteScopeLock = TRWScopeLockWithPriority<UE::TWriteScopeLock<FTransactionallySafeRWLock>>; }
+namespace UE::AssetRegistry { using FInterfaceReadScopeLock = TRWScopeLockWithPriority<UE::TReadScopeLock<Private::FRWLockWithPriority>>; }
+namespace UE::AssetRegistry { using FInterfaceWriteScopeLock = TRWScopeLockWithPriority<UE::TWriteScopeLock<Private::FRWLockWithPriority>>; }
 namespace UE::AssetRegistry { using FInterfaceRWScopeLock = class FRWScopeLockWithPriority; }
 
 #if WITH_EDITOR
@@ -95,9 +97,29 @@ namespace Private
 	 *  to work correctly it must be used with the associated TPriorityScopeLock types 
 	 *  (TScopeLockWithPriority and FRWScopeLockWithPriority)
 	 */
-	class FRWLockWithPriority : public FTransactionallySafeRWLock
+	class FRWLockWithPriority
 	{
 	public:
+		FORCEINLINE void WriteLock()
+		{
+			Mutex.WriteLock();
+		}
+
+		FORCEINLINE void WriteUnlock()
+		{
+			Mutex.WriteUnlock();
+		}
+
+		FORCEINLINE void ReadLock()
+		{
+			Mutex.ReadLock();
+		}
+
+		FORCEINLINE void ReadUnlock()
+		{
+			Mutex.ReadUnlock();
+		}
+
 		bool HasWaiters() const 
 		{
 			bool Result;
@@ -112,6 +134,7 @@ namespace Private
 		friend FInterfaceReadScopeLock;
 		friend FInterfaceWriteScopeLock;
 		friend FInterfaceRWScopeLock;
+		FTransactionallySafeRWLock Mutex;
 		std::atomic<int32> HighPriorityWaitersCount;
 	};
 
@@ -120,8 +143,6 @@ namespace Private
 		PriorityLow,
 		PriorityHigh
 	};
-
-	using FInterfaceRWLock = FRWLockWithPriority;
 } // namespace Private within namespace UE::AssetRegistry
 
 namespace Impl
