Only add Packages to the ignore filter if the package is fully loaded.
When searching in memory assets the UPackage can be partially loaded and only partially find all their subobjects. This was found when investigating verse hot reload.

diff --git a/Engine/Source/Runtime/AssetRegistry/Private/AssetRegistry.cpp b/Engine/Source/Runtime/AssetRegistry/Private/AssetRegistry.cpp
--- a/Engine/Source/Runtime/AssetRegistry/Private/AssetRegistry.cpp
+++ b/Engine/Source/Runtime/AssetRegistry/Private/AssetRegistry.cpp
@@ -2667,7 +2667,10 @@ void EnumerateMemoryAssetsHelper(const FARCompiledFilter& InFilter, TSet<FName>&
 		// We mark the package found for this passing asset, so that any followup search for assets on disk will not
 		// add a duplicate of this Asset. We do this here for convenience; it would be more correct to call it only for assets that
 		// pass the callers remaining filters inside of Callback
-		OutPackageNamesWithAssets.Add(FilterData.Package->GetFName());
+		if (FilterData.Package->IsFullyLoaded())
+		{
+			OutPackageNamesWithAssets.Add(FilterData.Package->GetFName());
+		}
 
 		// Could perhaps save some FName -> String conversions by creating this a bit earlier using the UObject constructor
 		// to get package name and path.
