AssetRegistry: Memory: Avoid wasted memory in runtime game due to deferredevents being stored but never processed and never freed. Implement a request function that adds a CoreTicker lambda to call the AssetRegistry's tick function. This is only necessary in non-editor; the editor ticks the AssetRegistry every frame. Also skip some unnecessary work in Broadcast if there are no listeners.
#rnx

diff --git a/Engine/Source/Runtime/AssetRegistry/Private/AssetRegistry.cpp b/Engine/Source/Runtime/AssetRegistry/Private/AssetRegistry.cpp
--- a/Engine/Source/Runtime/AssetRegistry/Private/AssetRegistry.cpp
+++ b/Engine/Source/Runtime/AssetRegistry/Private/AssetRegistry.cpp
@@ -1786,6 +1786,11 @@ void UAssetRegistryImpl::OnPreExit()
 void UAssetRegistryImpl::FinishDestroy()
 {
 	LLM_SCOPE(ELLMTag::AssetRegistry);
+
+	{
+		FScopeLock DeferredEventsLock(&DeferredEventsCriticalSection);
+		ClearRequestTick();
+	}
 	{
 		UE::AssetRegistry::FInterfaceWriteScopeLock InterfaceScopeLock(InterfaceLock);
 
@@ -4687,10 +4692,12 @@ UE::AssetRegistry::Impl::EGatherStatus UAssetRegistryImpl::TickOnBackgroundThrea
 			TickContext.bHandleDeferred = true;
 			Status = GuardedData.TickGatherer(TickContext);
 
+			if (!EventContext.IsEmpty())
 			{
 				FScopeLock DeferredEventsLock(&DeferredEventsCriticalSection);
 				DeferredEvents.Append(MoveTemp(EventContext));
 				EventContext.Clear();
+				RequestTick();
 			}
 		}
 		else 
@@ -4735,6 +4742,17 @@ void UAssetRegistryImpl::Tick(float DeltaTime)
 	bool bInterruptedOrShouldProcessDeferredEvents = false;
 	float LocalMaxSecondsPerFrame = UE::AssetRegistry::Impl::MaxSecondsPerFrame;
 
+	// Ticks in !WITH_EDITOR are done on request. Mark that we have satisfied the request; any further requests that come in from
+	// other threads while we are ticking will cause an additional queued request. RequestTick is not expected to be called from
+	// this thread, because all callers of it come only when writing to DeferredEvents, which we will not do in this thread because
+	// we can process the events immediately on this thread.
+#if !WITH_EDITOR
+	{
+		FScopeLock DeferredEventsLock(&DeferredEventsCriticalSection);
+		ClearRequestTick();
+	}
+#endif
+
 	do
 	{
 		bInterruptedOrShouldProcessDeferredEvents = false;
@@ -4847,6 +4865,42 @@ void UAssetRegistryImpl::Tick(float DeltaTime)
 		(TickStartTime < 0 || (FPlatformTime::Seconds() - TickStartTime) <= LocalMaxSecondsPerFrame));
 }
 
+void UAssetRegistryImpl::RequestTick()
+{
+	// Called from within DeferredEventsCriticalSection
+	// RequestTick is not used if WITH_EDITOR; the AssetRegistry is ticked every frame from
+	// UEditorEngine::Tick or a WITH_EDITOR block in UGameEngine::Tick.
+#if !WITH_EDITOR
+	if (TickRequestHandle.IsValid())
+	{
+		// Already Set
+		return;
+	}
+	TickRequestHandle = FTSTicker::GetCoreTicker().AddTicker(FTickerDelegate::CreateLambda([this](float)
+		{
+			// DeltaTime: non-negative to indicate we have a time limit, as small as possible to do just the minimum.
+			constexpr float DeltaTime = 0.0f;
+			this->Tick(DeltaTime);
+			return false; // Do not continue ticking
+		}));
+#endif
+}
+
+void UAssetRegistryImpl::ClearRequestTick()
+{
+	// Called from within DeferredEventsCriticalSection
+	// RequestTick/ClearRequestTick are not used if WITH_EDITOR; the AssetRegistry is ticked every frame from
+	// UEditorEngine::Tick or a WITH_EDITOR block in UGameEngine::Tick.
+#if !WITH_EDITOR
+	if (!TickRequestHandle.IsValid())
+	{
+		return;
+	}
+	FTSTicker::GetCoreTicker().RemoveTicker(TickRequestHandle);
+	TickRequestHandle.Reset();
+#endif
+}
+
 namespace UE::AssetRegistry
 {
 
@@ -9312,6 +9366,7 @@ void UAssetRegistryImpl::Broadcast(UE::AssetRegistry::Impl::FEventContext& Event
 		FScopeLock DeferredEventsLock(&DeferredEventsCriticalSection);
 		check(&EventContext != &DeferredEvents);
 		DeferredEvents.Append(MoveTemp(EventContext));
+		RequestTick();
 		EventContext.Clear();
 		return;
 	}
@@ -9325,41 +9380,45 @@ void UAssetRegistryImpl::Broadcast(UE::AssetRegistry::Impl::FEventContext& Event
 
 	if (EventContext.PathEvents.Num())
 	{
-		// Batch add/remove events 
-		TArray<FStringView> Params;
-		// Ensure loop batch condition is always false first iteration
-		bool bCurrentBatchIsAdd = EventContext.PathEvents[0].Get<1>() == FEventContext::EEvent::Added;
-		for (const TPair<FString, FEventContext::EEvent>& PathEvent : EventContext.PathEvents)
+		if (PathsAddedEvent.IsBound() || PathsRemovedEvent.IsBound()
+			|| PathAddedEvent.IsBound() || PathRemovedEvent.IsBound())
 		{
-			const FString& Path = PathEvent.Get<0>();
-			bool bEventIsAdd = PathEvent.Get<1>() == FEventContext::EEvent::Added;
-			if (bEventIsAdd != bCurrentBatchIsAdd)
+			// Batch add/remove events 
+			TArray<FStringView> Params;
+			// Ensure loop batch condition is always false first iteration
+			bool bCurrentBatchIsAdd = EventContext.PathEvents[0].Get<1>() == FEventContext::EEvent::Added;
+			for (const TPair<FString, FEventContext::EEvent>& PathEvent : EventContext.PathEvents)
+			{
+				const FString& Path = PathEvent.Get<0>();
+				bool bEventIsAdd = PathEvent.Get<1>() == FEventContext::EEvent::Added;
+				if (bEventIsAdd != bCurrentBatchIsAdd)
+				{
+					(bCurrentBatchIsAdd ? PathsAddedEvent : PathsRemovedEvent).Broadcast(MakeArrayView(Params));
+					Params.Reset();
+					bCurrentBatchIsAdd = bEventIsAdd;
+				}
+				Params.Add(FStringView(Path));
+			}
+			if (Params.Num() != 0)
 			{
 				(bCurrentBatchIsAdd ? PathsAddedEvent : PathsRemovedEvent).Broadcast(MakeArrayView(Params));
-				Params.Reset();
-				bCurrentBatchIsAdd = bEventIsAdd;
 			}
-			Params.Add(FStringView(Path));
-		}
-		if (Params.Num() != 0)
-		{
-			(bCurrentBatchIsAdd ? PathsAddedEvent : PathsRemovedEvent).Broadcast(MakeArrayView(Params));
-		}
 
-		// Legacy single events 
-		if (PathAddedEvent.IsBound() || PathRemovedEvent.IsBound())
-		{
-			for (const TPair<FString, FEventContext::EEvent>& PathEvent : EventContext.PathEvents)
+			// Legacy single events 
+			if (PathAddedEvent.IsBound() || PathRemovedEvent.IsBound())
 			{
-				const FString& Path = PathEvent.Get<0>();
-				switch (PathEvent.Get<1>())
+				for (const TPair<FString, FEventContext::EEvent>& PathEvent : EventContext.PathEvents)
 				{
-				case FEventContext::EEvent::Added:
-					PathAddedEvent.Broadcast(Path);
-					break;
-				case FEventContext::EEvent::Removed:
-					PathRemovedEvent.Broadcast(Path);
-					break;
+					const FString& Path = PathEvent.Get<0>();
+					switch (PathEvent.Get<1>())
+					{
+					case FEventContext::EEvent::Added:
+						PathAddedEvent.Broadcast(Path);
+						break;
+					case FEventContext::EEvent::Removed:
+						PathRemovedEvent.Broadcast(Path);
+						break;
+					}
 				}
 			}
 		}
@@ -9371,82 +9430,102 @@ void UAssetRegistryImpl::Broadcast(UE::AssetRegistry::Impl::FEventContext& Event
 		// Batch events so that if adds/updates are interspersed with removes, relative ordering of the add/remove is maintained 
 		constexpr uint32 EventTypeCount = static_cast<uint32>(FEventContext::EEvent::MAX);
 		static_assert(EventTypeCount == 4, "Loop needs to be rewritten to correctly order new event types");
-		TArray<FAssetData> EventBatches[EventTypeCount];
-		FEventContext::EEvent LastEvent = EventContext.AssetEvents[0].Get<1>();
-		auto FlushBatchedEvents = [&EventBatches, &Events=BatchedAssetEvents]() 
+		bool bHasListeners = AssetAddedEvent.IsBound() || AssetRemovedEvent.IsBound()
+			|| AssetUpdatedEvent.IsBound() || AssetUpdatedOnDiskEvent.IsBound();
+		if (!bHasListeners)
 		{
-			for (int32 i=0; i < UE_ARRAY_COUNT(EventBatches); ++i)
+			for (int32 i = 0; i < EventTypeCount; ++i)
 			{
-				if (EventBatches[i].Num())
+				if (BatchedAssetEvents[i].IsBound())
 				{
-					Events[i].Broadcast(EventBatches[i]);
-					EventBatches[i].Reset();
+					bHasListeners = true;
+					break;
 				}
 			}
-		};
+		}
 
-		for (const TPair<FAssetData, FEventContext::EEvent>& AssetEvent : EventContext.AssetEvents)
+		if (bHasListeners)
 		{
-			const FAssetData& AssetData = AssetEvent.Get<0>();
-			FEventContext::EEvent Event = AssetEvent.Get<1>();
+			TArray<FAssetData> EventBatches[EventTypeCount];
+			FEventContext::EEvent LastEvent = EventContext.AssetEvents[0].Get<1>();
+			auto FlushBatchedEvents = [&EventBatches, &Events = BatchedAssetEvents]()
+				{
+					for (int32 i = 0; i < UE_ARRAY_COUNT(EventBatches); ++i)
+					{
+						if (EventBatches[i].Num())
+						{
+							Events[i].Broadcast(EventBatches[i]);
+							EventBatches[i].Reset();
+						}
+					}
+				};
 
-			// Flush events when switching between removed and non-removed events 
-			if ((Event == FEventContext::EEvent::Removed) != (LastEvent == FEventContext::EEvent::Removed))
+			for (const TPair<FAssetData, FEventContext::EEvent>& AssetEvent : EventContext.AssetEvents)
 			{
-				FlushBatchedEvents();
+				const FAssetData& AssetData = AssetEvent.Get<0>();
+				FEventContext::EEvent Event = AssetEvent.Get<1>();
+
+				// Flush events when switching between removed and non-removed events 
+				if ((Event == FEventContext::EEvent::Removed) != (LastEvent == FEventContext::EEvent::Removed))
+				{
+					FlushBatchedEvents();
+				}
+				EventBatches[static_cast<int32>(Event)].Add(AssetData);
+				LastEvent = Event;
 			}
-			EventBatches[static_cast<int32>(Event)].Add(AssetData);
-			LastEvent = Event;
-		}
-		// Flush last batch of events 
-		FlushBatchedEvents();
-		
-		// Single events
-		for (const TPair<FAssetData, FEventContext::EEvent>& AssetEvent : EventContext.AssetEvents)
-		{
-			const FAssetData& AssetData = AssetEvent.Get<0>();
-			switch (AssetEvent.Get<1>())
+			// Flush last batch of events 
+			FlushBatchedEvents();
+
+			// Single events
+			for (const TPair<FAssetData, FEventContext::EEvent>& AssetEvent : EventContext.AssetEvents)
 			{
-			case FEventContext::EEvent::Added:
-				AssetAddedEvent.Broadcast(AssetData);
-				break;
-			case FEventContext::EEvent::Removed:
-				AssetRemovedEvent.Broadcast(AssetData);
-				break;
-			case FEventContext::EEvent::Updated:
-				AssetUpdatedEvent.Broadcast(AssetData);
-				break;
-			case FEventContext::EEvent::UpdatedOnDisk:
-				AssetUpdatedOnDiskEvent.Broadcast(AssetData);
-				break;
-			default:
-				checkNoEntry();
-				break;
+				const FAssetData& AssetData = AssetEvent.Get<0>();
+				switch (AssetEvent.Get<1>())
+				{
+				case FEventContext::EEvent::Added:
+					AssetAddedEvent.Broadcast(AssetData);
+					break;
+				case FEventContext::EEvent::Removed:
+					AssetRemovedEvent.Broadcast(AssetData);
+					break;
+				case FEventContext::EEvent::Updated:
+					AssetUpdatedEvent.Broadcast(AssetData);
+					break;
+				case FEventContext::EEvent::UpdatedOnDisk:
+					AssetUpdatedOnDiskEvent.Broadcast(AssetData);
+					break;
+				default:
+					checkNoEntry();
+					break;
+				}
 			}
 		}
 		EventContext.AssetEvents.Empty();
 	}
 	if (EventContext.VerseEvents.Num())
 	{
-		for (const TPair<FName, FEventContext::EEvent>& VerseEvent : EventContext.VerseEvents)
+		if (VerseAddedEvent.IsBound() || VerseRemovedEvent.IsBound())
 		{
-			const FName& VerseFilepath = VerseEvent.Get<0>();
-			switch (VerseEvent.Get<1>())
+			for (const TPair<FName, FEventContext::EEvent>& VerseEvent : EventContext.VerseEvents)
 			{
-			case FEventContext::EEvent::Added:
-				VerseAddedEvent.Broadcast(VerseFilepath);
-				break;
-			case FEventContext::EEvent::Removed:
-				VerseRemovedEvent.Broadcast(VerseFilepath);
-				break;
-			// (jcotton) We are not yet broadcasting Verse updating events as the only use case for VerseEvent broadcasts currently is to trigger a Verse-build
-			// and triggering a build on every change would be far too expensive.
-			case FEventContext::EEvent::Updated:
-				[[fallthrough]];
-			case FEventContext::EEvent::UpdatedOnDisk:
-				[[fallthrough]];
-			default:
-				break;
+				const FName& VerseFilepath = VerseEvent.Get<0>();
+				switch (VerseEvent.Get<1>())
+				{
+				case FEventContext::EEvent::Added:
+					VerseAddedEvent.Broadcast(VerseFilepath);
+					break;
+				case FEventContext::EEvent::Removed:
+					VerseRemovedEvent.Broadcast(VerseFilepath);
+					break;
+				// (jcotton) We are not yet broadcasting Verse updating events as the only use case for VerseEvent broadcasts currently is to trigger a Verse-build
+				// and triggering a build on every change would be far too expensive.
+				case FEventContext::EEvent::Updated:
+					[[fallthrough]];
+				case FEventContext::EEvent::UpdatedOnDisk:
+					[[fallthrough]];
+				default:
+					break;
+				}
 			}
 		}
 		EventContext.VerseEvents.Empty();
@@ -9489,6 +9568,7 @@ void UAssetRegistryImpl::Broadcast(UE::AssetRegistry::Impl::FEventContext& Event
 			// and broadcast called on that separate EventContext outside of the lock.
 			check(&EventContext != &DeferredEvents);
 			DeferredEvents.Append(MoveTemp(EventContext));
+			RequestTick();
 			EventContext.Clear();
 			check(!EventContext.bFileLoadedEventBroadcast && !EventContext.bKnownGathersCompleteEventBroadcast); // was cleared by Append and by Clear
 			check(!bLocalFileLoadedEventBroadcast || DeferredEvents.bFileLoadedEventBroadcast); // was set by Append
diff --git a/Engine/Source/Runtime/AssetRegistry/Private/AssetRegistry.h b/Engine/Source/Runtime/AssetRegistry/Private/AssetRegistry.h
--- a/Engine/Source/Runtime/AssetRegistry/Private/AssetRegistry.h
+++ b/Engine/Source/Runtime/AssetRegistry/Private/AssetRegistry.h
@@ -3,6 +3,7 @@
 #pragma once
 
 #include "AssetRegistryImpl.h"
+#include "Containers/Ticker.h"
 #include "UObject/Object.h"
 
 #include "AssetRegistry.generated.h"
@@ -330,6 +331,11 @@ private:
 #endif
 	}
 
+	/** In engine modes that do not tick regularly, request a tick, to process deferred events. */
+	void RequestTick();
+	/** In engine modes that do not tick regularly, mark the tick was called and clear the request. */
+	void ClearRequestTick();
+
 private:
 
 	UE::AssetRegistry::FAssetRegistryImpl GuardedData;
@@ -418,10 +424,16 @@ private:
 	/** The delegate to execute when each batch of searching/gathering is complete */
 	FKnownGathersCompleteEvent KnownGathersCompleteEvent;
 
-	/** Storage for events that will be broadcast later on the game thread. Only safe
+	/**
+	 * Storage for events that will be broadcast later on the game thread. Only safe
 	 *  to access under the DeferredEventsCriticalSection. 
 	 */
 	UE::AssetRegistry::Impl::FEventContext DeferredEvents;
+	/**
+	 * Handle for a currently active tick request. Runtime game only, not used in GIsEditor. Used to process DeferredEvents.
+	 * Read/Write only within DeferredEventsCriticalSection.
+	 */
+	FTSTicker::FDelegateHandle TickRequestHandle;
 	FCriticalSection DeferredEventsCriticalSection;
 
 	friend class UE::AssetRegistry::FAssetRegistryImpl;
