AssetManager: UpdateManagementDatabase: Move ExternalActor-handling code out of ShouldSetManager  into the AssetRegistry's read of dependencies, to save time by running it only once per asset and only when uncached. The moved code is the code that modifies the propertytype of ExternalActor's reference on its outer from UsedInGame to EditorOnly.
#rnx

diff --git a/Engine/Source/Runtime/AssetRegistry/Private/PackageDependencyData.cpp b/Engine/Source/Runtime/AssetRegistry/Private/PackageDependencyData.cpp
--- a/Engine/Source/Runtime/AssetRegistry/Private/PackageDependencyData.cpp
+++ b/Engine/Source/Runtime/AssetRegistry/Private/PackageDependencyData.cpp
@@ -1,7 +1,11 @@
 // Copyright Epic Games, Inc. All Rights Reserved.
 
 #include "PackageDependencyData.h"
+
 #include "Algo/Unique.h"
+#include "Misc/StringBuilder.h"
+#include "String/Find.h"
+#include "UObject/NameTypes.h"
 
 FName FPackageDependencyData::GetImportPackageName(TConstArrayView<FObjectImport> ImportMap, int32 ImportIndex)
 {
@@ -50,13 +54,45 @@ void FPackageDependencyData::LoadDependenciesFromPackageHeader(FName SourcePacka
 {
 	using namespace UE::AssetRegistry;
 
+	// ExternalActor imports of their outer can be treated as editor-only imports because ExternalActors are EditorOnly
+	// packages. References to ExternalActors can occur from collector assets that find all ExternalActor packages in
+	// a plugin. When propagating should-be-cooked rules or chunk assignment rules in the AssetManager, we don't want
+	// to pull in/set the chunk of the world containing the plugin's actor. SavePackage records the import of the
+	// externalactor's world package so we would ordinarily mark it as UsedInGame in the loop below. But we override
+	// that for AssetRegistry dependencies so that the AssetManager propagation does not pull in the map package when
+	// an actor is referenced from a collector. We implement that here by naming convention.
+	// See also FAssetRegistryGenerator::ComputePackageDifferences.
+	FStringView ExternalActorFolder(FPackagePath::GetExternalActorsFolderName());
+	TStringBuilder<FName::StringBufferSize> SourceStr(InPlace, SourcePackageName);
+	FStringView ExternalActorWorldRelPath;
+	if (int32 FolderIndex = UE::String::FindFirst(SourceStr, ExternalActorFolder, ESearchCase::IgnoreCase);
+		FolderIndex != INDEX_NONE)
+	{
+		ExternalActorWorldRelPath = SourceStr.ToView().RightChop(FolderIndex + ExternalActorFolder.Len());
+	}
+	auto IsWorldOfExternalActor = [ExternalActorWorldRelPath](FName DependencyPackageName)
+		{
+			TStringBuilder<256> TargetStr(InPlace, DependencyPackageName);
+			FStringView TargetMountPoint = FPathViews::GetMountPointNameFromPath(TargetStr);
+			FStringView TargetRelativePath = FStringView(TargetStr).RightChop(TargetMountPoint.Len() + 1);
+			return (INDEX_NONE !=
+				UE::String::FindFirst(ExternalActorWorldRelPath, TargetRelativePath, ESearchCase::IgnoreCase));
+		};
+
 	PackageDependencies.Reset(ImportMap.Num() + SoftPackageReferenceList.Num());
 	check(ImportMap.Num() == ImportUsedInGame.Num());
 	for (int32 ImportIdx = 0; ImportIdx < ImportMap.Num(); ++ImportIdx)
 	{
 		FName DependencyPackageName = GetImportPackageName(ImportMap, ImportIdx);
 		EDependencyProperty DependencyProperty = EDependencyProperty::Hard;
-		DependencyProperty |= ImportUsedInGame[ImportIdx] ? EDependencyProperty::Game : EDependencyProperty::None;
+		bool bUsedInGame = ImportUsedInGame[ImportIdx];
+
+		if (!ExternalActorWorldRelPath.IsEmpty())
+		{
+			bUsedInGame &= !IsWorldOfExternalActor(DependencyPackageName);
+		}
+
+		DependencyProperty |= bUsedInGame ? EDependencyProperty::Game : EDependencyProperty::None;
 		PackageDependencies.Add({ DependencyPackageName, DependencyProperty });
 	}
 
@@ -72,7 +108,14 @@ void FPackageDependencyData::LoadDependenciesFromPackageHeader(FName SourcePacka
 		FName DependencyPackageName = SoftPackageReferenceList[SoftPackageIdx];
 		FAssetIdentifier AssetId(DependencyPackageName);
 		EDependencyProperty DependencyProperty = EDependencyProperty::None;
-		DependencyProperty |= (SoftPackageUsedInGame[SoftPackageIdx] ? EDependencyProperty::Game : EDependencyProperty::None);
+
+		bool bUsedInGame = SoftPackageUsedInGame[SoftPackageIdx];
+		if (!ExternalActorWorldRelPath.IsEmpty())
+		{
+			bUsedInGame &= !IsWorldOfExternalActor(DependencyPackageName);
+		}
+
+		DependencyProperty |= (bUsedInGame ? EDependencyProperty::Game : EDependencyProperty::None);
 
 		// Don't need to remove duplicates here because SavePackage only writes unique elements into SoftPackageReferenceList
 		PackageDependencies.Add({ DependencyPackageName, DependencyProperty });
@@ -80,11 +123,21 @@ void FPackageDependencyData::LoadDependenciesFromPackageHeader(FName SourcePacka
 
 	for (const TPair<FName, EExtraDependencyFlags>& Pair : ExtraPackageDependencies)
 	{
+		FName DependencyPackageName = Pair.Key;
 		EDependencyProperty DependencyProperty = EDependencyProperty::None;
 		DependencyProperty |= EnumHasAnyFlags(Pair.Value, EExtraDependencyFlags::Build)
 			? EDependencyProperty::Build : EDependencyProperty::None;
+
+		if (!ExternalActorWorldRelPath.IsEmpty())
+		{
+			if (IsWorldOfExternalActor(DependencyPackageName))
+			{
+				continue;
+			}
+		}
+
 		// Don't need to remove duplicates here because SavePackage only writes unique elements into PackageBuildDependencies
-		PackageDependencies.Add({ Pair.Key, DependencyProperty });
+		PackageDependencies.Add({ DependencyPackageName, DependencyProperty });
 	}
 
 	SearchableNameDependencies.Reset(SearchableNames.Num());
diff --git a/Engine/Source/Runtime/CoreUObject/Public/AssetRegistry/AssetData.h b/Engine/Source/Runtime/CoreUObject/Public/AssetRegistry/AssetData.h
--- a/Engine/Source/Runtime/CoreUObject/Public/AssetRegistry/AssetData.h
+++ b/Engine/Source/Runtime/CoreUObject/Public/AssetRegistry/AssetData.h
@@ -101,6 +101,7 @@ struct FAssetRegistryVersion
 		AssetPackageDataHasPackageLocation,	// Added PackageLocation to AssetPackageData.
 		MarshalledTextAsUTF8String,			// Replaced 2 byte wide string with UTF8 String
 		PackageSavedHash,					// Replaced FAssetPackageData::PackageGuid with PackageSavedHash
+		ExternalActorToWorldIsEditorOnly,   // FPackageDependencyData::LoadDependenciesFromPackageHeader changed how it calculates PackageDependencies
 
 		// -----<new versions can be added above this line>-------------------------------------------------
 		VersionPlusOne,
diff --git a/Engine/Source/Runtime/Engine/Private/AssetManager.cpp b/Engine/Source/Runtime/Engine/Private/AssetManager.cpp
--- a/Engine/Source/Runtime/Engine/Private/AssetManager.cpp
+++ b/Engine/Source/Runtime/Engine/Private/AssetManager.cpp
@@ -3982,30 +3982,6 @@ EAssetSetManagerResult::Type UAssetManager::ShouldSetManager(const FAssetIdentif
 		return EAssetSetManagerResult::DoNotSet;
 	}
 
-	// EXTERNALACTOR_TODO: Replace this workaround for ExternalActors with a modification to the ExternalActor Packages' 
-	// dependencies. External actors have an import dependency (hard, build, game) on their Map package because
-	// the map package is their outer. At cook time they are saved in umaps (WorldPartition cells, generated packages). Prevent
-	// their dependency on the WorldPartition generator package to set them as a manager of the generator package.
-	// Workaround: Detect external actors by naming convention and suppress their reference to the map package.
-	// Long-Term Fix: Make the external actors dependency on their map a non-game one so they are not considered while evaluating the manager asset.
-	// See also FAssetRegistryGenerator::ComputePackageDifferences
-	TStringBuilder<256> SourcePackageString;
-	Source.PackageName.ToString(SourcePackageString);
-	int32 ExternalActorIdx = UE::String::FindFirst(SourcePackageString, ULevel::GetExternalActorsFolderName(), ESearchCase::IgnoreCase);
-	if (ExternalActorIdx != INDEX_NONE)
-	{
-		FStringView TargetMountPoint = FPathViews::GetMountPointNameFromPath(TargetPackageString);
-		FStringView TargetRelativePath = FStringView(TargetPackageString).RightChop(TargetMountPoint.Len() + 1);
-
-		FStringView ChoppedExternalActor = FStringView(SourcePackageString).RightChop(ExternalActorIdx + FStringView(ULevel::GetExternalActorsFolderName()).Len());
-		bool bIsTargetWorld = UE::String::FindFirst(ChoppedExternalActor, TargetRelativePath, ESearchCase::IgnoreCase) != INDEX_NONE;
-		if (bIsTargetWorld)
-		{
-			// If the Target is the UWorld and its source an External actor, then do not set. ExternalActors do not influence their level's chunk.
-			return EAssetSetManagerResult::DoNotSet;
-		}
-	}
-
 	if (Flags & EAssetSetManagerFlags::TargetHasExistingManager)
 	{
 		// If target has a higher priority manager, never recurse and only set manager if direct
