AssetManager+AssetRegistry: UpdateManagementDatabase: Performance improvements.
* Change the arguments to SetManageReferences and ShouldSetManager to be a single context argument each, so we can make implementation-specific changes to the arguments in the future without needing further deprecation.
* Move the detection of whether a package is a script package from the frequently called ShouldSetManager function into the called-once-and-cached function in the AssetRegistry's parsing of depends nodes.
* Refactor containers in SetManageReferences to use a single TMap for all the data the function holds about each FDependsNode; persist this map across multiple calls to SetManageReferences by storing it in the new FSetManageReferencesContext argument.
#rnx

diff --git a/Engine/Source/Runtime/AssetRegistry/Private/AssetRegistry.cpp b/Engine/Source/Runtime/AssetRegistry/Private/AssetRegistry.cpp
--- a/Engine/Source/Runtime/AssetRegistry/Private/AssetRegistry.cpp
+++ b/Engine/Source/Runtime/AssetRegistry/Private/AssetRegistry.cpp
@@ -8851,15 +8851,117 @@ static FAutoConsoleVariableRef CVarAssetRegistryManagementPathsPackageDebugName(
 	TEXT("AssetRegistry.ManagementPathsPackageDebugName"),
 	GAssetRegistryManagementPathsPackageDebugName,
 	TEXT("If set, when manage references are set, the chain of references that caused this package to become managed will be printed to the log"));
+#endif
 
-void PrintAssetRegistryManagementPathsPackageDebugInfo(FDependsNode* Node, const TMap<FDependsNode*, FDependsNode*>& EditorOnlyManagementPaths)
+namespace UE::AssetRegistry
+{
+
+/**
+ * Data held for each Asset (aka node in our graph search) across the multiple graph searches conducted within
+ * SetManageReferences, and across multiple calls to SetManageReferences
+ * 
+ * Each call to SetManageReferences is called a round. The rounds occur in descending priority order of managers;
+ * each round does a graph search for the nodes reported as managed by managers in that round. All the managers in
+ * a given round have the same priority.
+ *
+ */
+struct FSetManageReferencesNodeData
+{
+	// Values preserved across multiple rounds
+	/**
+	 * bManagedInEarlierRound is set for every asset found managed in a round, and later rounds can use that information
+	 * in ShouldSetManager calls to e.g. ignore the management of a node by a lower-priority manager.
+	 */
+	bool bManagedInEarlierRound = false;
+
+	// Values reset for each round.
+	/**
+	 * Whether the asset was modified by the round (and therefore needs to have transient properties restored at the
+	 * end of the round.
+	 */
+	bool bModifiedByRound = false;
+	/**
+	 * Whether the asset was found to be managed by the round (and therefore needs to be added to
+	 * bManagedInEarlierRound).
+	 */
+	bool bManagedInThisRound = false;
+	/**
+	 * The direct managers of the asset in the current round, as described by the caller of the round in
+	 * FSetManageReferencesContext.ManagerMap.
+	 */
+	TArray<FDependsNode*> DirectManagersThisRound;
+	/**
+	 * The source node (aka asset) that caused this node to be referenced in the current round, used to provide
+	 * instigator chains for what caused an asset to be added to the cook or to a chunk.
+	 */
+	FDependsNode* DebugInstigator = nullptr;
+
+	// Values reset for the graph search from each manager with a round.
+	/**
+	 * Whether the asset was modified by the manager's search (and therefore needs to have transient properties
+	 * restored at the end of the search).
+	 */
+	bool bModifiedByCurrentManager = false;
+	/** Whether the asset was found to be managed by the manager. */
+	bool bVisitedByCurrentManager = false;
+	/**
+	 * The AssetRegistry dependency properties that should be assigned to the targetnode for the Management edge we
+	 * create for the node from the current manager.
+	 */
+	EDependencyProperty CurrentManagerProperties = EDependencyProperty::None;
+};
+
+// TODO: Replace the three existing copies of TFastPointerMap/TFastPointerSet with a shared implementation,
+// or move the behavior into DefaultKeyFuncs.
+template<typename KeyType, typename ValueType, bool bInAllowDuplicateKeys>
+struct TFastPointerMapKeyFuncs : public TDefaultMapKeyFuncs<KeyType, ValueType, bInAllowDuplicateKeys>
+{
+	using typename TDefaultMapKeyFuncs<KeyType, ValueType, bInAllowDuplicateKeys>::KeyInitType;
+	static FORCEINLINE uint32 GetKeyHash(KeyInitType Key)
+	{
+#if PLATFORM_64BITS
+		static_assert(sizeof(UPTRINT) == sizeof(uint64), "Expected pointer size to be 64 bits");
+		// Ignoring the lower 4 bits since they are likely zero anyway.
+		const uint64 ImportantBits = reinterpret_cast<uint64>(Key) >> 4;
+		return GetTypeHash(ImportantBits);
+#else
+		static_assert(sizeof(UPTRINT) == sizeof(uint32), "Expected pointer size to be 32 bits");
+		return static_cast<uint32>(reinterpret_cast<uint32>(Key));
+#endif
+	}
+};
+
+template<typename KeyType, typename ValueType, typename SetAllocator = FDefaultSetAllocator>
+class TFastPointerMap : public TMap<KeyType, ValueType, SetAllocator,
+	TFastPointerMapKeyFuncs<KeyType, ValueType, false /* bInAllowDuplicateKeys */>>
+{};
+
+using FSetManageReferencesNodeDataMap = TFastPointerMap<FDependsNode*, TUniquePtr<FSetManageReferencesNodeData>>;
+struct FSetManageReferencesScratch
+{
+	FSetManageReferencesNodeDataMap NodeData;
+};
+
+#if WITH_EDITOR
+void PrintAssetRegistryManagementPathsPackageDebugInfo(FDependsNode* Node, UE::AssetRegistry::FSetManageReferencesNodeDataMap* NodeData)
 {
 	if (Node)
 	{
 		UE_LOG(LogAssetRegistry, Display, TEXT("SetManageReferences is printing out the reference chain that caused '%s' to be managed"), *GAssetRegistryManagementPathsPackageDebugName);
 		TSet<FDependsNode*> AllVisitedNodes;
-		while (FDependsNode* ReferencingNode = EditorOnlyManagementPaths.FindRef(Node))
+		for (;;)
 		{
+			TUniquePtr<FSetManageReferencesNodeData>* CurrentNodeData = NodeData->Find(Node);
+			if (!CurrentNodeData)
+			{
+				break;
+			}
+			FDependsNode* ReferencingNode = (*CurrentNodeData)->DebugInstigator;
+			if (!ReferencingNode)
+			{
+				break;
+			}
+
 			UE_LOG(LogAssetRegistry, Display, TEXT("  %s"), *ReferencingNode->GetIdentifier().ToString());
 			if (AllVisitedNodes.Contains(ReferencingNode))
 			{
@@ -8878,190 +8980,300 @@ void PrintAssetRegistryManagementPathsPackageDebugInfo(FDependsNode* Node, const
 }
 #endif // WITH_EDITOR
 
+} // namespace UE::AssetRegistry
+
 void UAssetRegistryImpl::SetManageReferences(const TMultiMap<FAssetIdentifier, FAssetIdentifier>& ManagerMap,
 	bool bClearExisting, UE::AssetRegistry::EDependencyCategory RecurseType, TSet<FDependsNode*>& ExistingManagedNodes,
 	ShouldSetManagerPredicate ShouldSetManager)
+{
+	if (!ShouldSetManager)
+	{
+		ShouldSetManager = [](const FAssetIdentifier& Manager, const FAssetIdentifier& Source, const FAssetIdentifier& Target,
+			UE::AssetRegistry::EDependencyCategory Category, UE::AssetRegistry::EDependencyProperty Properties,
+			EAssetSetManagerFlags::Type Flags)
+			{
+				return EAssetSetManagerResult::SetButDoNotRecurse;
+			};
+	}
+	auto ShouldSetManagerByContext = [&ShouldSetManager](UE::AssetRegistry::FShouldSetManagerContext& ShouldSetContext)
+		{
+			return ShouldSetManager(ShouldSetContext.Manager, ShouldSetContext.Source, ShouldSetContext.Target,
+				ShouldSetContext.EdgeARCategory, ShouldSetContext.EdgeARProperties, ShouldSetContext.EdgeFlags);
+		};
+
+	UE::AssetRegistry::FSetManageReferencesContext Context;
+	Context.ManagerMap = &ManagerMap;
+	Context.bClearExisting = bClearExisting;
+	Context.RecurseType = RecurseType;
+	Context.ShouldSetManager = ShouldSetManagerByContext;
+	SetManageReferences(Context);
+}
+
+void UAssetRegistryImpl::SetManageReferences(UE::AssetRegistry::FSetManageReferencesContext& Context)
 {
 	// For performance reasons we call the ShouldSetManager callback when inside the lock. Licensee UAssetManagers
 	// are responsible for not calling AssetRegistry functions from ShouldSetManager as that would create a deadlock
 	LLM_SCOPE(ELLMTag::AssetRegistry);
 	UE::AssetRegistry::FInterfaceWriteScopeLock InterfaceScopeLock(InterfaceLock);
-	GuardedData.SetManageReferences(ManagerMap, bClearExisting, RecurseType, ExistingManagedNodes, ShouldSetManager);
+	GuardedData.SetManageReferences(Context);
 }
 
 namespace UE::AssetRegistry
 {
 
-void FAssetRegistryImpl::SetManageReferences(const TMultiMap<FAssetIdentifier, FAssetIdentifier>& ManagerMap,
-	bool bClearExisting, UE::AssetRegistry::EDependencyCategory RecurseType, TSet<FDependsNode*>& ExistingManagedNodes,
-	IAssetRegistry::ShouldSetManagerPredicate ShouldSetManager)
+static EAssetSetManagerAssetFlags ConstructAssetManagerAssetFlags(const FDependsNode& DependsNode)
 {
-	// Set default predicate if needed
-	if (!ShouldSetManager)
+	return DependsNode.IsScriptPath() ? EAssetSetManagerAssetFlags::ScriptPackage
+		: EAssetSetManagerAssetFlags::None;
+};
+
+void FAssetRegistryImpl::SetManageReferences(FSetManageReferencesContext& Context)
+{
+	if (!Context.ManagerMap)
+	{
+		UE_LOG(LogAssetRegistry, Error, TEXT("SetManageReferences: Context.ManagerMap must not be null."));
+		FDebug::DumpStackTraceToLog(ELogVerbosity::Warning);
+		return;
+	}
+	if (!Context.ShouldSetManager.IsSet())
 	{
-		ShouldSetManager = [](const FAssetIdentifier& Manager, const FAssetIdentifier& Source, const FAssetIdentifier& Target, UE::AssetRegistry::EDependencyCategory Category, UE::AssetRegistry::EDependencyProperty Properties, EAssetSetManagerFlags::Type Flags)
+		UE_LOG(LogAssetRegistry, Error, TEXT("SetManageReferences: Context.ShouldSetManager must be a bound pointer."));
+		FDebug::DumpStackTraceToLog(ELogVerbosity::Warning);
+		return;
+	}
+
+	FSetManageReferencesScratch* Scratch = nullptr;
+	FSetManageReferencesNodeDataMap* NodeData = nullptr;
+	// TODO: Invalidate Scratch and give an error if dependsnodes are written in between calls; Scratch has pointers
+	// and cached data into the AssetRegistry's list of DependsNodes.
+
+	TArray<FDependsNode*> CurrentRoundModifiedNodes;
+
+	if (!Context.Scratch)
+	{
+		Context.Scratch = MakePimpl<FSetManageReferencesScratch>();
+		Context.Scratch->NodeData.Reserve(State.CachedDependsNodes.Num());
+	}
+	Scratch = Context.Scratch.Get();
+	NodeData = &Scratch->NodeData;
+
+	// We use a TMap of TUniquePtr rather than direct c++ member data so that we can have persistent pointers to
+	// NodeDatas throughout the graph search, rather than possibly have our data pointers be invalidated by a
+	// reallocation when a new node is added. This has a cputime cost cost - an extra memory access per read.
+	// It also has a complexity cost, because instead of calling FindOrAdd to create if non-existing, we also have
+	// to check whether it was added and allocate the TUniquePtr if so. This FindOrAddNodeData handles that
+	// complexity of creation.
+	auto FindOrAddNodeData = [NodeData](FDependsNode* Node) -> FSetManageReferencesNodeData&
 		{
-			return EAssetSetManagerResult::SetButDoNotRecurse;
+			TUniquePtr<FSetManageReferencesNodeData>& Existing = NodeData->FindOrAdd(Node);
+			if (!Existing)
+			{
+				Existing.Reset(new FSetManageReferencesNodeData());
+			}
+			return *Existing;
 		};
-	}
 
-	if (bClearExisting)
+	if (Context.bClearExisting)
 	{
-		// Find all nodes with incoming manage dependencies
-		for (const TPair<FAssetIdentifier, FDependsNode*>& Pair : State.CachedDependsNodes)
+		// Clear information from Scratch about Earlier rounds
+		for (TPair<FDependsNode*, TUniquePtr<FSetManageReferencesNodeData>>& NodePair : *NodeData)
 		{
-			Pair.Value->IterateOverDependencies([&ExistingManagedNodes](FDependsNode* TestNode, EDependencyCategory Category, EDependencyProperty Property, bool bUnique)
-				{
-					ExistingManagedNodes.Add(TestNode);
-				}, EDependencyCategory::Manage);
+			NodePair.Value->bManagedInEarlierRound = false;
 		}
 
-		// Clear them
+		// Clear all outgoing manage dependencies, and update the referencers of any nodes with incoming
+		// manage dependencies
+		TSet<FDependsNode*> NodesToRefreshReferencers;
 		for (const TPair<FAssetIdentifier, FDependsNode*>& Pair : State.CachedDependsNodes)
 		{
-			Pair.Value->ClearDependencies(EDependencyCategory::Manage);
+			FDependsNode* SourceNode = Pair.Value;
+			SourceNode->IterateOverDependencies([&NodesToRefreshReferencers]
+			(FDependsNode* TargetNode, EDependencyCategory Category, EDependencyProperty Property, bool bUnique)
+				{
+					NodesToRefreshReferencers.Add(TargetNode);
+				}, EDependencyCategory::Manage);
+			SourceNode->ClearDependencies(EDependencyCategory::Manage);
 		}
-		for (FDependsNode* NodeToClear : ExistingManagedNodes)
+
+		for (FDependsNode* NodeToClear : NodesToRefreshReferencers)
 		{
 			NodeToClear->SetIsReferencersSorted(false);
 			NodeToClear->RefreshReferencers();
+
+			CurrentRoundModifiedNodes.Add(NodeToClear);
+			FSetManageReferencesNodeData& ModifiedData = FindOrAddNodeData(NodeToClear);
+			ModifiedData.bModifiedByRound = true;
 		}
-		ExistingManagedNodes.Empty();
 	}
 
-	TMap<FDependsNode*, TArray<FDependsNode *>> ExplicitMap; // Reverse of ManagerMap, specifies what relationships to add to each node
-
-	for (const TPair<FAssetIdentifier, FAssetIdentifier>& Pair : ManagerMap)
+	TArray<FDependsNode*> CurrentRoundExplicitlyManagedNodes;
+	TArray<FDependsNode*> CurrentManagerModifiedNodes;
+	TArray<FDependsNode*> CurrentManagerVisitQueue;
+	for (const TPair<FAssetIdentifier, FAssetIdentifier>& Pair : *Context.ManagerMap)
 	{
 		FDependsNode* ManagedNode = State.FindDependsNode(Pair.Value);
-
 		if (!ManagedNode)
 		{
 			UE_LOG(LogAssetRegistry, Warning, TEXT("Cannot set %s to manage asset %s because %s does not exist!"), 
-				*Pair.Key.ToString(),
-				*Pair.Value.ToString(),
-				*Pair.Value.ToString());
+				*Pair.Key.ToString(), *Pair.Value.ToString(), *Pair.Value.ToString());
 			continue;
 		}
-
-		TArray<FDependsNode*>& ManagerList = ExplicitMap.FindOrAdd(ManagedNode);
-
 		FDependsNode* ManagerNode = State.CreateOrFindDependsNode(Pair.Key);
 
-		ManagerList.Add(ManagerNode);
+		CurrentRoundExplicitlyManagedNodes.Add(ManagedNode);
+		FindOrAddNodeData(ManagedNode).DirectManagersThisRound.Add(ManagerNode);
 	}
 
-	TSet<FDependsNode*> Visited;
-	TMap<FDependsNode*, EDependencyProperty> NodesToManage;
-	TArray<FDependsNode*> NodesToHardReference;
-	TArray<FDependsNode*> NodesToRecurse;
-
-#if WITH_EDITOR
-	// Map of every depends node to the node whose reference caused it to become managed by an asset. Used to look up why an asset was chosen to be the manager.
-	TMap<FDependsNode*, FDependsNode*> EditorOnlyManagementPaths;
-#endif
+	Algo::Sort(CurrentRoundExplicitlyManagedNodes);
+	CurrentRoundExplicitlyManagedNodes.SetNum(Algo::Unique(CurrentRoundExplicitlyManagedNodes));
 
-	TSet<FDependsNode*> NewManageNodes;
-	// For each explicitly set asset
-	for (const TPair<FDependsNode*, TArray<FDependsNode *>>& ExplicitPair : ExplicitMap)
+	// For each directly mangaged node, do a graph search for each of its managers.
+	for (FDependsNode* BaseManagedNode : CurrentRoundExplicitlyManagedNodes)
 	{
-		FDependsNode* BaseManagedNode = ExplicitPair.Key;
-		const TArray<FDependsNode*>& ManagerNodes = ExplicitPair.Value;
-
-		for (FDependsNode* ManagerNode : ManagerNodes)
+		FSetManageReferencesNodeData& BaseManagedNodeData = FindOrAddNodeData(BaseManagedNode);
+		for (FDependsNode* ManagerNode : BaseManagedNodeData.DirectManagersThisRound)
 		{
-			Visited.Reset();
-			NodesToManage.Reset();
-			NodesToRecurse.Reset();
+			check(CurrentManagerModifiedNodes.IsEmpty());
+			check(CurrentManagerVisitQueue.IsEmpty());
 
-			FDependsNode* SourceNode = ManagerNode;
-
-			auto IterateFunction = [&ManagerNode, &SourceNode, &ShouldSetManager, &NodesToManage, &NodesToHardReference, &NodesToRecurse, &Visited, &ExplicitMap, &ExistingManagedNodes
-#if WITH_EDITOR
-										, &EditorOnlyManagementPaths
-#endif
-										](FDependsNode* ReferencingNode, FDependsNode* TargetNode, EDependencyCategory DependencyType, EDependencyProperty DependencyProperties)
+			auto IterateFunction = [&ManagerNode, &Context, NodeData, &CurrentManagerModifiedNodes,
+				&CurrentManagerVisitQueue, &FindOrAddNodeData]
+			(FDependsNode* SourceNode, FDependsNode* TargetNode, EDependencyCategory DependencyType,
+				EDependencyProperty DependencyProperties)
 			{
 				// Only recurse if we haven't already visited, and this node passes recursion test
-				if (!Visited.Contains(TargetNode))
+				FSetManageReferencesNodeData& TargetNodeData = FindOrAddNodeData(TargetNode);
+				if (TargetNodeData.bVisitedByCurrentManager)
 				{
-					EAssetSetManagerFlags::Type Flags = (EAssetSetManagerFlags::Type)((SourceNode == ManagerNode ? EAssetSetManagerFlags::IsDirectSet : 0)
-						| (ExistingManagedNodes.Contains(TargetNode) ? EAssetSetManagerFlags::TargetHasExistingManager : 0)
-						| (ExplicitMap.Find(TargetNode) && SourceNode != ManagerNode ? EAssetSetManagerFlags::TargetHasDirectManager : 0));
+					return;
+				}
 
-					EAssetSetManagerResult::Type Result = ShouldSetManager(ManagerNode->GetIdentifier(), SourceNode->GetIdentifier(), TargetNode->GetIdentifier(), DependencyType, DependencyProperties, Flags);
+				EAssetSetManagerFlags::Type EdgeFlags = (EAssetSetManagerFlags::Type)(
+					(SourceNode == ManagerNode ? EAssetSetManagerFlags::IsDirectSet : 0)
+					| (TargetNodeData.bManagedInEarlierRound ? EAssetSetManagerFlags::TargetHasExistingManager : 0)
+					| (TargetNodeData.DirectManagersThisRound.Num() > 0 && SourceNode != ManagerNode
+						? EAssetSetManagerFlags::TargetHasDirectManager : 0));
+				FShouldSetManagerContext ShouldSetContext
+				{
+					.Manager = ManagerNode->GetIdentifier(),
+					.Source = SourceNode->GetIdentifier(),
+					.Target = TargetNode->GetIdentifier(),
+					.ManagerAssetFlags = ConstructAssetManagerAssetFlags(*ManagerNode),
+					.SourceAssetFlags = ConstructAssetManagerAssetFlags(*SourceNode),
+					.TargetAssetFlags = ConstructAssetManagerAssetFlags(*TargetNode),
+					.EdgeFlags = EdgeFlags,
+					.EdgeARCategory = DependencyType,
+					.EdgeARProperties = DependencyProperties
+				};
 
-					if (Result == EAssetSetManagerResult::DoNotSet)
-					{
-						return;
-					}
+				EAssetSetManagerResult::Type Result = (*Context.ShouldSetManager)(ShouldSetContext);
+				if (Result == EAssetSetManagerResult::DoNotSet)
+				{
+					return;
+				}
 
-					EDependencyProperty ManageProperties = (Flags & EAssetSetManagerFlags::IsDirectSet) ? EDependencyProperty::Direct : EDependencyProperty::None;
-					NodesToManage.Add(TargetNode, ManageProperties);
+				TargetNodeData.bVisitedByCurrentManager = true;
+				TargetNodeData.bModifiedByCurrentManager = true;
+				CurrentManagerModifiedNodes.Add(TargetNode);
 
-#if WITH_EDITOR
-					if (!GAssetRegistryManagementPathsPackageDebugName.IsEmpty())
-					{
-						EditorOnlyManagementPaths.Add(TargetNode, ReferencingNode ? ReferencingNode : ManagerNode);
-					}
-#endif
+				EDependencyProperty ManageProperties = 
+				TargetNodeData.CurrentManagerProperties = (EdgeFlags & EAssetSetManagerFlags::IsDirectSet)
+					? EDependencyProperty::Direct : EDependencyProperty::None;
+				TargetNodeData.DebugInstigator = SourceNode;
 
-					if (Result == EAssetSetManagerResult::SetAndRecurse)
-					{
-						NodesToRecurse.Push(TargetNode);
-					}
+				if (Result == EAssetSetManagerResult::SetAndRecurse)
+				{
+					CurrentManagerVisitQueue.Push(TargetNode);
 				}
 			};
 
 			// Check initial node
-			IterateFunction(nullptr, BaseManagedNode, EDependencyCategory::Manage, EDependencyProperty::Direct);
+			IterateFunction(ManagerNode, BaseManagedNode, EDependencyCategory::Manage, EDependencyProperty::Direct);
 
 			// Do all recursion first, but only if we have a recurse type
-			if (RecurseType != EDependencyCategory::None)
+			if (Context.RecurseType != EDependencyCategory::None)
 			{
-				while (NodesToRecurse.Num())
+				while (CurrentManagerVisitQueue.Num())
 				{
-					// Pull off end of array, order doesn't matter
-					SourceNode = NodesToRecurse.Pop();
-										
-					Visited.Add(SourceNode);
+					// Pull off end of array to avoid the cost of shifts; order of visitation doesn't matter
+					FDependsNode* SourceNode = CurrentManagerVisitQueue.Pop();
 
-					SourceNode->IterateOverDependencies([&IterateFunction, SourceNode](FDependsNode* TargetNode, EDependencyCategory DependencyCategory, EDependencyProperty DependencyProperties, bool bDuplicate)
+					SourceNode->IterateOverDependencies([&IterateFunction, SourceNode]
+					(FDependsNode* TargetNode, EDependencyCategory DependencyCategory,
+						EDependencyProperty DependencyProperties, bool bDuplicate)
 						{
-							// Skip editor-only, non-build dependencies. Propagate only through used-in-game or build dependencies.
+							// Skip dependencies that are EditorOnly and non-build.
+							// Propagate only through UsedInGame or build dependencies.
 							if (EnumHasAnyFlags(DependencyProperties, EDependencyProperty::Game | EDependencyProperty::Build))
 							{
 								IterateFunction(SourceNode, TargetNode, DependencyCategory, DependencyProperties);
 							}
-						}, RecurseType);
+						}, Context.RecurseType);
 				}
 			}
 
 			ManagerNode->SetIsDependencyListSorted(UE::AssetRegistry::EDependencyCategory::Manage, false);
-			for (TPair<FDependsNode*, EDependencyProperty>& ManagePair : NodesToManage)
+			FSetManageReferencesNodeData& ManagerData = FindOrAddNodeData(ManagerNode);
+			ManagerData.bModifiedByRound = true;
+			CurrentRoundModifiedNodes.Add(ManagerNode);
+
+			for (FDependsNode* ModifiedNode : CurrentManagerModifiedNodes)
 			{
-				ManagePair.Key->SetIsReferencersSorted(false);
-				ManagePair.Key->AddReferencer(ManagerNode);
-				ManagerNode->AddDependency(ManagePair.Key, EDependencyCategory::Manage, ManagePair.Value);
-				NewManageNodes.Add(ManagePair.Key);
+				FSetManageReferencesNodeData& ModifiedData = FindOrAddNodeData(ModifiedNode);
+				if (!ModifiedData.bModifiedByCurrentManager)
+				{
+					// A duplicate of a NodeData we already handled earlier in the list
+					continue;
+				}
+
+				ModifiedNode->SetIsReferencersSorted(false);
+				ModifiedNode->AddReferencer(ManagerNode);
+				ManagerNode->AddDependency(ModifiedNode, EDependencyCategory::Manage,
+					ModifiedData.CurrentManagerProperties);
+
+				ModifiedData.bManagedInThisRound = true;
+				ModifiedData.bModifiedByRound = true;
+				CurrentRoundModifiedNodes.Add(ModifiedNode);
+
+				ModifiedData.bModifiedByCurrentManager = false;
+				ModifiedData.bVisitedByCurrentManager = false;
+				ModifiedData.CurrentManagerProperties = EDependencyProperty::None;
 			}
+			CurrentManagerModifiedNodes.Reset();
+			CurrentManagerVisitQueue.Reset();
 		}
 	}
 
-	for (FDependsNode* ManageNode : NewManageNodes)
-	{
-		ExistingManagedNodes.Add(ManageNode);
-	}
-	// Restore all nodes to manage dependencies sorted and references sorted, so we can efficiently read them in future operations
-	State.SetDependencyNodeSorting(ShouldSortDependencies(), ShouldSortReferencers());
-
 #if WITH_EDITOR
 	if (!GAssetRegistryManagementPathsPackageDebugName.IsEmpty())
 	{
 		FDependsNode* PackageDebugInfoNode = State.FindDependsNode(FName(*GAssetRegistryManagementPathsPackageDebugName));
-		PrintAssetRegistryManagementPathsPackageDebugInfo(PackageDebugInfoNode, EditorOnlyManagementPaths);
+		PrintAssetRegistryManagementPathsPackageDebugInfo(PackageDebugInfoNode, NodeData);
 	}
 #endif
+
+	bool bShouldSortDependencies = ShouldSortDependencies();
+	bool bShouldSortReferencers = ShouldSortReferencers();
+	for (FDependsNode* ModifiedNode : CurrentRoundModifiedNodes)
+	{
+		FSetManageReferencesNodeData& ModifiedData = FindOrAddNodeData(ModifiedNode);
+		if (!ModifiedData.bModifiedByRound)
+		{
+			// A duplicate of a NodeData we already handled earlier in the list
+			continue;
+		}
+
+		// Restore all modified nodes to manage dependencies sorted and references sorted, so we can efficiently read
+		// them in future operations
+		ModifiedData.bManagedInEarlierRound |= ModifiedData.bManagedInThisRound;
+		ModifiedNode->SetIsDependencyListSorted(EDependencyCategory::Manage, bShouldSortDependencies);
+		ModifiedNode->SetIsReferencersSorted(bShouldSortReferencers);
+
+		ModifiedData.bModifiedByRound = false;
+		ModifiedData.bManagedInThisRound = false;
+		ModifiedData.DirectManagersThisRound.Empty();
+		ModifiedData.DebugInstigator = nullptr;
+	}
 }
 
 }
diff --git a/Engine/Source/Runtime/AssetRegistry/Private/AssetRegistry.h b/Engine/Source/Runtime/AssetRegistry/Private/AssetRegistry.h
--- a/Engine/Source/Runtime/AssetRegistry/Private/AssetRegistry.h
+++ b/Engine/Source/Runtime/AssetRegistry/Private/AssetRegistry.h
@@ -225,7 +225,10 @@ public:
 	UE::AssetRegistry::Impl::EGatherStatus TickOnBackgroundThread();
 
 protected:
-	virtual void SetManageReferences(const TMultiMap<FAssetIdentifier, FAssetIdentifier>& ManagerMap, bool bClearExisting, UE::AssetRegistry::EDependencyCategory RecurseType, TSet<FDependsNode*>& ExistingManagedNodes, ShouldSetManagerPredicate ShouldSetManager = nullptr) override;
+	virtual void SetManageReferences(const TMultiMap<FAssetIdentifier, FAssetIdentifier>& ManagerMap,
+		bool bClearExisting, UE::AssetRegistry::EDependencyCategory RecurseType,
+		TSet<FDependsNode*>& ExistingManagedNode, ShouldSetManagerPredicate ShouldSetManager = nullptr) override;
+	virtual void SetManageReferences(UE::AssetRegistry::FSetManageReferencesContext& Context) override;
 	virtual bool SetPrimaryAssetIdForObjectPath(const FSoftObjectPath& ObjectPath, FPrimaryAssetId PrimaryAssetId) override;
 
 private:
diff --git a/Engine/Source/Runtime/AssetRegistry/Private/AssetRegistryImpl.h b/Engine/Source/Runtime/AssetRegistry/Private/AssetRegistryImpl.h
--- a/Engine/Source/Runtime/AssetRegistry/Private/AssetRegistryImpl.h
+++ b/Engine/Source/Runtime/AssetRegistry/Private/AssetRegistryImpl.h
@@ -295,9 +295,7 @@ public:
 	void GetAllocatedSize(bool bLogDetailed, SIZE_T& StateSize, SIZE_T& StaticSize, SIZE_T& SearchSize) const;
 	bool IsLoadingAssets() const;
 	bool IsGathering() const;
-	void SetManageReferences(const TMultiMap<FAssetIdentifier, FAssetIdentifier>& ManagerMap,
-		bool bClearExisting, UE::AssetRegistry::EDependencyCategory RecurseType,
-		TSet<FDependsNode*>& ExistingManagedNodes, IAssetRegistry::ShouldSetManagerPredicate ShouldSetManager);
+	void SetManageReferences(FSetManageReferencesContext& Context);
 	bool SetPrimaryAssetIdForObjectPath(Impl::FEventContext& EventContext, const FSoftObjectPath& ObjectPath, FPrimaryAssetId PrimaryAssetId);
 	bool ResolveRedirect(const FString& InPackageName, FString& OutPackageName) const;
 #if WITH_EDITOR
diff --git a/Engine/Source/Runtime/AssetRegistry/Private/DependsNode.cpp b/Engine/Source/Runtime/AssetRegistry/Private/DependsNode.cpp
--- a/Engine/Source/Runtime/AssetRegistry/Private/DependsNode.cpp
+++ b/Engine/Source/Runtime/AssetRegistry/Private/DependsNode.cpp
@@ -4,10 +4,11 @@
 
 #include "AssetRegistry/AssetRegistryState.h"
 #include "AssetRegistryPrivate.h"
+#include "Misc/StringBuilder.h"
 
 void FDependsNode::PrintNode() const
 {
-	UE_LOG(LogAssetRegistry, Log, TEXT("*** Printing DependsNode: %s ***"), *Identifier.ToString());
+	UE_LOG(LogAssetRegistry, Log, TEXT("*** Printing DependsNode: %s ***"), *GetIdentifier().ToString());
 	UE_LOG(LogAssetRegistry, Log, TEXT("*** Dependencies:"));
 	PrintDependencies();
 	UE_LOG(LogAssetRegistry, Log, TEXT("*** Referencers:"));
@@ -28,6 +29,20 @@ void FDependsNode::PrintReferencers() const
 	PrintReferencersRecursive(TEXT(""), VisitedNodes);
 }
 
+void FDependsNode::SetIdentifier(const FAssetIdentifier& InIdentifier)
+{
+	Identifier = InIdentifier;
+	if (!InIdentifier.PackageName.IsNone())
+	{
+		TStringBuilder<FName::StringBufferSize> PackageNameStr(InPlace, InIdentifier.PackageName);
+		bScriptPath = PackageNameStr.ToView().StartsWith(TEXT("/Script/"), ESearchCase::CaseSensitive) ? 1 : 0;
+	}
+	else
+	{
+		bScriptPath = 0;
+	}
+}
+
 /**
  * Helper function for IterateDependencyList, which itself is a helper function for FDependsNode functions.
  * This function is given an index into one of the lists. All of the data for the list and how to interpret it is also
@@ -713,11 +728,11 @@ void FDependsNode::PrintDependenciesRecursive(const FString& Indent, TSet<const
 	}
 	else if ( VisitedNodes.Contains(this) )
 	{
-		UE_LOG(LogAssetRegistry, Log, TEXT("%s[CircularReferenceTo]%s"), *Indent, *Identifier.ToString());
+		UE_LOG(LogAssetRegistry, Log, TEXT("%s[CircularReferenceTo]%s"), *Indent, *GetIdentifier().ToString());
 	}
 	else
 	{
-		UE_LOG(LogAssetRegistry, Log, TEXT("%s%s"), *Indent, *Identifier.ToString());
+		UE_LOG(LogAssetRegistry, Log, TEXT("%s%s"), *Indent, *GetIdentifier().ToString());
 		VisitedNodes.Add(this);
 
 		IterateOverDependencies([&Indent, &VisitedNodes](FDependsNode* InDependency,
@@ -743,11 +758,11 @@ void FDependsNode::PrintReferencersRecursive(const FString& Indent, TSet<const F
 	}
 	else if ( VisitedNodes.Contains(this) )
 	{
-		UE_LOG(LogAssetRegistry, Log, TEXT("%s[CircularReferenceTo]%s"), *Indent, *Identifier.ToString());
+		UE_LOG(LogAssetRegistry, Log, TEXT("%s[CircularReferenceTo]%s"), *Indent, *GetIdentifier().ToString());
 	}
 	else
 	{
-		UE_LOG(LogAssetRegistry, Log, TEXT("%s%s"), *Indent, *Identifier.ToString());
+		UE_LOG(LogAssetRegistry, Log, TEXT("%s%s"), *Indent, *GetIdentifier().ToString());
 		VisitedNodes.Add(this);
 
 		for (FDependsNode* Node : Referencers)
@@ -832,6 +847,7 @@ void FDependsNode::SerializeLoad(FArchive& Ar,
 	FLoadScratch& Scratch)
 {
 	Ar << Identifier;
+	SetIdentifier(Identifier);
 
 	auto ReadDependencies = [&Ar, &GetNodeFromSerializeIndex, &Scratch](TArray<FDependsNode*>& OutDependencies,
 		TBitArray<>* OutFlagBits,
@@ -916,6 +932,7 @@ void FDependsNode::SerializeLoad_BeforeFlags(FArchive& Ar,
 	uint32 SoftManageBits)
 {
 	Ar << Identifier;
+	SetIdentifier(Identifier);
 
 	int32 NumHard, NumSoft, NumName, NumSoftManage, NumHardManage, NumReferencers;
 	Ar << NumHard;
diff --git a/Engine/Source/Runtime/AssetRegistry/Private/DependsNode.h b/Engine/Source/Runtime/AssetRegistry/Private/DependsNode.h
--- a/Engine/Source/Runtime/AssetRegistry/Private/DependsNode.h
+++ b/Engine/Source/Runtime/AssetRegistry/Private/DependsNode.h
@@ -23,8 +23,8 @@ public:
 	static constexpr uint32 ManageFlagSetWidth = TPropertyCombinationSet<ManageFlagWidth>::StorageBitCount;
 
 public:
-	FDependsNode() { Construct(); }
-	FDependsNode(const FAssetIdentifier& InIdentifier) : Identifier(InIdentifier) { Construct(); }
+	FDependsNode();
+	FDependsNode(const FAssetIdentifier& InIdentifier);
 
 	/** Prints the dependencies and referencers for this node to the log */
 	void PrintNode() const;
@@ -50,14 +50,13 @@ public:
 	void GetReferencers(TArray<FAssetDependency>& OutReferencers,
 		UE::AssetRegistry::EDependencyCategory Category = UE::AssetRegistry::EDependencyCategory::All,
 		const UE::AssetRegistry::FDependencyQuery& Flags = UE::AssetRegistry::FDependencyQuery()) const;
-	/** Gets the name of the package that this node represents */
-	FName GetPackageName() const { return Identifier.PackageName; }
-	/** Sets the name of the package that this node represents */
-	void SetPackageName(FName InName) { Identifier = FAssetIdentifier(InName); }
+	/** Helper function to return GetIdentifier.PackageName. */
+	FName GetPackageName() const;
 	/** Returns the entire identifier */
-	const FAssetIdentifier& GetIdentifier() const { return Identifier; }
+	const FAssetIdentifier& GetIdentifier() const;
+	bool IsScriptPath() const;
 	/** Sets the entire identifier */
-	void SetIdentifier(const FAssetIdentifier& InIdentifier) { Identifier = InIdentifier; }
+	void SetIdentifier(const FAssetIdentifier& InIdentifier);
 	/** Add a dependency to this node */
 	void AddDependency(FDependsNode* InDependency,
 		UE::AssetRegistry::EDependencyCategory InDependencyType, UE::AssetRegistry::EDependencyProperty InProperties);
@@ -187,7 +186,7 @@ private:
 	 */
 	void PrintReferencersRecursive(const FString& Indent, TSet<const FDependsNode*>& VisitedNodes) const;
 
-	void Construct();
+	void ConstructFlags();
 
 	/** The name of the package/object this node represents */
 	FAssetIdentifier Identifier;
@@ -197,11 +196,14 @@ private:
 	FDependsNodeList Referencers;
 	TBitArray<> PackageFlags;
 	TBitArray<> ManageFlags;
+
+	// Transient flags that are not serialized
 	uint32 PackageIsSorted : 1;
 	uint32 SearchableNameIsSorted : 1;
 	uint32 ManageIsSorted : 1;
 	uint32 ReferencersIsSorted : 1;
 	uint32 DependenciesInitialized : 1;
+	uint32 bScriptPath : 1;
 };
 
 
@@ -209,6 +211,31 @@ private:
 // Inline implementations
 ///////////////////////////////////////////////////////
 
+inline FDependsNode::FDependsNode()
+{
+	ConstructFlags();
+}
+
+inline FDependsNode::FDependsNode(const FAssetIdentifier& InIdentifier)
+{
+	ConstructFlags();
+	SetIdentifier(InIdentifier);
+}
+
+inline FName FDependsNode::GetPackageName() const
+{
+	return Identifier.PackageName;
+}
+
+inline const FAssetIdentifier& FDependsNode::GetIdentifier() const
+{
+	return Identifier;
+}
+
+inline bool FDependsNode::IsScriptPath() const
+{
+	return bScriptPath != 0;
+}
 
 inline SIZE_T FDependsNode::GetAllocatedSize(void) const
 {
@@ -270,11 +297,12 @@ inline UE::AssetRegistry::EDependencyProperty FDependsNode::ByteToManageProperti
 		);
 }
 
-inline void FDependsNode::Construct()
+inline void FDependsNode::ConstructFlags()
 {
 	PackageIsSorted = 1;
 	SearchableNameIsSorted = 1;
 	ManageIsSorted = 1;
 	ReferencersIsSorted = 1;
 	DependenciesInitialized = 0;
+	bScriptPath = 0;
 }
diff --git a/Engine/Source/Runtime/AssetRegistry/Public/AssetRegistry/IAssetRegistry.h b/Engine/Source/Runtime/AssetRegistry/Public/AssetRegistry/IAssetRegistry.h
--- a/Engine/Source/Runtime/AssetRegistry/Public/AssetRegistry/IAssetRegistry.h
+++ b/Engine/Source/Runtime/AssetRegistry/Public/AssetRegistry/IAssetRegistry.h
@@ -9,6 +9,8 @@
 #include "Containers/StringFwd.h"
 #include "Misc/AssetRegistryInterface.h"
 #include "Misc/Optional.h"
+#include "Templates/Function.h"
+#include "Templates/PimplPtr.h"
 #include "UObject/Interface.h"
 #include "IAssetRegistry.generated.h"
 
@@ -22,6 +24,7 @@ class FCbFieldView;
 class FCbWriter;
 class FDependsNode;
 enum class EAssetRegistryTagsCaller : uint8;
+namespace UE::AssetRegistry { struct FSetManageReferencesScratch; }
 struct FARFilter;
 struct FARCompiledFilter;
 struct FAssetData;
@@ -66,6 +69,63 @@ ENUM_CLASS_FLAGS(EScanFlags);
 
 ASSETREGISTRY_API FString LexToString(EScanFlags Flags);
 
+/** Properties about each node in the graph search conducted from SetManageReferences. */
+enum class EAssetSetManagerAssetFlags : uint32
+{
+	None = 0x0,
+	ScriptPackage = 0x1,
+};
+ENUM_CLASS_FLAGS(EAssetSetManagerAssetFlags);
+
+/** Information passed into ShouldSetManager calls from the graph search performed in SetManageReferences. */
+struct FShouldSetManagerContext
+{
+	/** The top-level manager node for which the search is conducted. */
+	const FAssetIdentifier& Manager;
+	/** The source node of the edge in the graph search that is being asked about in the ShouldSetManager call. */
+	const FAssetIdentifier& Source;
+	/** The target node of the edge in the graph search that is being asked about in the ShouldSetManager call. */
+	const FAssetIdentifier& Target;
+	/** Properties of the manager node. */
+	EAssetSetManagerAssetFlags ManagerAssetFlags;
+	/** Properties of the source node. */
+	EAssetSetManagerAssetFlags SourceAssetFlags;
+	/** Properties of the source node. */
+	EAssetSetManagerAssetFlags TargetAssetFlags;
+	/* Properties of the edge, and of the target's status in the current phase of SetManageReferences. */
+	EAssetSetManagerFlags::Type EdgeFlags;
+	/** AssetRegistry category of the edge: which kind of AssetRegistry dependency the edge comes from. */
+	EDependencyCategory EdgeARCategory;
+	/** AssetRegistry Properties of the edge: which properties within its category the edge has. */
+	EDependencyProperty EdgeARProperties;
+};
+
+/**
+ * Information passed into SetManageReferences from ther caller. Includes arguments for the call as well as
+ * scratch space data that is shared across multiple SetManageReferences calls conducted by the caller for different
+ * priorities of managers.
+ */
+struct FSetManageReferencesContext
+{
+	/** Map from Managing asset to directly-managed assets for the current call to SetManageReferences.*/
+	const TMultiMap<FAssetIdentifier, FAssetIdentifier>* ManagerMap = nullptr;
+	/** If true, will clear any existing manage dependencies. Should be set only for the first call. */
+	bool bClearExisting = false;
+	/** Dependency types to follow during the graph search out from the starting nodes in ManagerMap. */
+	EDependencyCategory RecurseType = EDependencyCategory::None;
+	/**
+	 * Predicate that is called on recursive dependencies if bound, returns whether the graph search should visit a
+	 * target node from a source node and whether that visit should be recursive.
+	 */
+	TOptional<TFunctionRef<EAssetSetManagerResult::Type(FShouldSetManagerContext& Context)>> ShouldSetManager;
+
+	/**
+	 * Scratch variable that will be constructed during the call, and should not be modified by the caller but should
+	* be preserved across a series of calls to SetManageReferences for managers of different priorities
+	*/
+	TPimplPtr<FSetManageReferencesScratch> Scratch;
+};
+
 } // namespace UE::AssetRegistry
 
 USTRUCT(BlueprintType)
@@ -1087,19 +1147,18 @@ protected:
 	// Functions specifically for calling from the asset manager
 	friend class UAssetManager;
 
+	UE_DEPRECATED(5.6, "Use version that takes FSetManageReferencesContext.")
+	virtual void SetManageReferences(const TMultiMap<FAssetIdentifier, FAssetIdentifier>& ManagerMap,
+		bool bClearExisting, UE::AssetRegistry::EDependencyCategory RecurseType,
+		TSet<FDependsNode*>& ExistingManagedNodes, ShouldSetManagerPredicate ShouldSetManager = nullptr) = 0;
+
 	/**
 	 * Specifies a list of manager mappings, optionally recursing to dependencies. These mappings can then be queried
 	 * later to see which assets "manage" other assets. This function is only meant to be called by the AssetManager,
 	 * calls from anywhere else will conflict and lose data.
-	 *
-	 * @param ManagerMap Map of Managing asset to Managed asset. This will construct Manager references and clear existing
-	 * @param bClearExisting If true, will clear any existing manage dependencies
-	 * @param RecurseType Dependency types to recurse into, from the value of the manager map
-	 * @param RecursePredicate Predicate that is called on recursion if bound, return true if it should recurse into that node
+	 * @see UE::AssetRegistry::FSetManageReferencesContext for arguments.
 	 */
-	virtual void SetManageReferences(const TMultiMap<FAssetIdentifier, FAssetIdentifier>& ManagerMap,
-		bool bClearExisting, UE::AssetRegistry::EDependencyCategory RecurseType,
-		TSet<FDependsNode*>& ExistingManagedNodes, ShouldSetManagerPredicate ShouldSetManager = nullptr) = 0;
+	virtual void SetManageReferences(UE::AssetRegistry::FSetManageReferencesContext& Context) = 0;
 
 	/**
 	 * Sets the PrimaryAssetId for a specific asset. This should only be called by the AssetManager, and is needed
diff --git a/Engine/Source/Runtime/Engine/Classes/Engine/AssetManager.h b/Engine/Source/Runtime/Engine/Classes/Engine/AssetManager.h
--- a/Engine/Source/Runtime/Engine/Classes/Engine/AssetManager.h
+++ b/Engine/Source/Runtime/Engine/Classes/Engine/AssetManager.h
@@ -16,17 +16,17 @@
 
 DECLARE_LOG_CATEGORY_EXTERN(LogAssetManager, Log, All);
 
-/** Defined in C++ file */
 class IAssetRegistry;
 class ITargetPlatform;
 class FAssetManagerEditorModule;
 class FTimerManager;
 class UAssetManagerSettings;
+namespace UE::AssetRegistry { struct FShouldSetManagerContext; }
+namespace UE::Cook { class ICookInfo; }
+struct FCompiledAssetManagerSearchRules;
 struct FPrimaryAssetTypeData;
 struct FPrimaryAssetData;
 struct FPrimaryAssetRulesCustomOverride;
-namespace UE::Cook { class ICookInfo; }
-struct FCompiledAssetManagerSearchRules;
 
 /** Delegate called when acquiring resources/chunks for assets, first parameter will be true if all resources were acquired, false if any failed.
 	Second parameter will contain a list of chunks not yet downloaded */
@@ -1013,8 +1013,13 @@ protected:
 
 #if WITH_EDITOR
 	/** Function used during creating Management references to decide when to recurse and set references */
+	ENGINE_API virtual EAssetSetManagerResult::Type ShouldSetManager(UE::AssetRegistry::FShouldSetManagerContext& Context) const;
+	UE_DEPRECATED(5.6, "Override the version that takes an FShouldSetManagerContext instead.")
 	ENGINE_API virtual EAssetSetManagerResult::Type ShouldSetManager(const FAssetIdentifier& Manager, const FAssetIdentifier& Source, const FAssetIdentifier& Target,
-		UE::AssetRegistry::EDependencyCategory Category, UE::AssetRegistry::EDependencyProperty Properties, EAssetSetManagerFlags::Type Flags) const;
+		UE::AssetRegistry::EDependencyCategory Category, UE::AssetRegistry::EDependencyProperty Properties, EAssetSetManagerFlags::Type Flags) const
+	{
+		return EAssetSetManagerResult::DoNotSet;
+	}
 
 	/** Called when asset registry is done loading off disk, will finish any deferred loads */
 	ENGINE_API virtual void OnAssetRegistryFilesLoaded();
diff --git a/Engine/Source/Runtime/Engine/Private/AssetManager.cpp b/Engine/Source/Runtime/Engine/Private/AssetManager.cpp
--- a/Engine/Source/Runtime/Engine/Private/AssetManager.cpp
+++ b/Engine/Source/Runtime/Engine/Private/AssetManager.cpp
@@ -9,6 +9,7 @@
 #include "AssetRegistry/AssetRegistryHelpers.h"
 #include "AssetRegistry/AssetRegistryModule.h"
 #include "AssetRegistry/AssetRegistryState.h"
+#include "AssetRegistry/IAssetRegistry.h"
 #include "Engine/AssetManagerSettings.h"
 #include "Engine/AssetManagerTypes.h"
 #include "Engine/BlueprintCore.h"
@@ -22,6 +23,7 @@
 #include "IPlatformFilePak.h"
 #include "Logging/StructuredLog.h"
 #include "Misc/CoreDelegates.h"
+#include "Misc/EnumClassFlags.h"
 #include "Misc/FileHelper.h"
 #include "Misc/DelayedAutoRegister.h"
 #include "Misc/ScopedSlowTask.h"
@@ -3969,30 +3971,26 @@ void UAssetManager::RefreshPrimaryAssetDirectory(bool bForceRefresh)
 
 #if WITH_EDITOR
 
-EAssetSetManagerResult::Type UAssetManager::ShouldSetManager(const FAssetIdentifier& Manager, const FAssetIdentifier& Source, const FAssetIdentifier& Target,
-	UE::AssetRegistry::EDependencyCategory Category, UE::AssetRegistry::EDependencyProperty Properties, EAssetSetManagerFlags::Type Flags) const
+EAssetSetManagerResult::Type UAssetManager::ShouldSetManager(UE::AssetRegistry::FShouldSetManagerContext& Context) const
 {
-	FPrimaryAssetId ManagerPrimaryAssetId = Manager.GetPrimaryAssetId();
-	FPrimaryAssetId TargetPrimaryAssetId = Target.GetPrimaryAssetId();
+	FPrimaryAssetId ManagerPrimaryAssetId = Context.Manager.GetPrimaryAssetId();
+	FPrimaryAssetId TargetPrimaryAssetId = Context.Target.GetPrimaryAssetId();
 	if (TargetPrimaryAssetId.IsValid())
 	{
 		// Don't recurse Primary Asset Id references
 		return EAssetSetManagerResult::SetButDoNotRecurse;
 	}
 
-	TStringBuilder<256> TargetPackageString;
-	Target.PackageName.ToString(TargetPackageString);
-
 	// Ignore script references
-	if (FStringView(TargetPackageString).StartsWith(TEXT("/Script/"), ESearchCase::CaseSensitive))
+	if (EnumHasAnyFlags(Context.TargetAssetFlags, UE::AssetRegistry::EAssetSetManagerAssetFlags::ScriptPackage))
 	{
 		return EAssetSetManagerResult::DoNotSet;
 	}
 
-	if (Flags & EAssetSetManagerFlags::TargetHasExistingManager)
+	if (Context.EdgeFlags & EAssetSetManagerFlags::TargetHasExistingManager)
 	{
 		// If target has a higher priority manager, never recurse and only set manager if direct
-		if (Flags & EAssetSetManagerFlags::IsDirectSet)
+		if (Context.EdgeFlags & EAssetSetManagerFlags::IsDirectSet)
 		{
 			return EAssetSetManagerResult::SetButDoNotRecurse;
 		}
@@ -4001,11 +3999,11 @@ EAssetSetManagerResult::Type UAssetManager::ShouldSetManager(const FAssetIdentif
 			return EAssetSetManagerResult::DoNotSet;
 		}
 	}
-	else if (Flags & EAssetSetManagerFlags::TargetHasDirectManager)
+	else if (Context.EdgeFlags & EAssetSetManagerFlags::TargetHasDirectManager)
 	{
 		// If target has another direct manager being set in this run, never recurse and set manager if we think this is an "owner" reference and not a back reference
 
-		bool bIsOwnershipReference = Flags & EAssetSetManagerFlags::IsDirectSet;
+		bool bIsOwnershipReference = Context.EdgeFlags & EAssetSetManagerFlags::IsDirectSet;
 
 		if (ManagerPrimaryAssetId.PrimaryAssetType == MapType)
 		{
@@ -4162,31 +4160,42 @@ void UAssetManager::UpdateManagementDatabase(bool bForceRefresh)
 	const bool bAllowInPIE = true;
 	SlowTask.MakeDialog(bShowCancelButton, bAllowInPIE);
 
-	auto SetManagerPredicate = [this, &PackagesToUpdateChunksFor](const FAssetIdentifier& Manager, const FAssetIdentifier& Source, const FAssetIdentifier& Target,
-		UE::AssetRegistry::EDependencyCategory Category, UE::AssetRegistry::EDependencyProperty Properties, EAssetSetManagerFlags::Type Flags)
-	{
-		EAssetSetManagerResult::Type Result = this->ShouldSetManager(Manager, Source, Target, Category, Properties, Flags);
-		if (Result != EAssetSetManagerResult::DoNotSet && Target.IsPackage())
+	auto SetManagerPredicate = [this, &PackagesToUpdateChunksFor](UE::AssetRegistry::FShouldSetManagerContext& Context)
 		{
-			PackagesToUpdateChunksFor.Add(Target.PackageName);
-		}
-		return Result;
-	};
+			EAssetSetManagerResult::Type Result = this->ShouldSetManager(Context);
+			if (Result != EAssetSetManagerResult::DoNotSet && Context.Target.IsPackage())
+			{
+				PackagesToUpdateChunksFor.Add(Context.Target.PackageName);
+			}
+			return Result;
+		};
+	auto SetManagerPredicateDirectSet = [](UE::AssetRegistry::FShouldSetManagerContext& Context)
+		{
+			return EAssetSetManagerResult::SetButDoNotRecurse;
+		};
 
-	TSet<FDependsNode*> ExistingManagedNodes;
+	UE::AssetRegistry::FSetManageReferencesContext SetManageContext;
+	
 	for (int32 PriorityIndex = 0; PriorityIndex < PriorityArray.Num(); PriorityIndex++)
 	{
-		TMultiMap<FAssetIdentifier, FAssetIdentifier>* ManagerMap = PriorityManagementMap.Find(PriorityArray[PriorityIndex]);
-
+		SetManageContext.ManagerMap = PriorityManagementMap.Find(PriorityArray[PriorityIndex]);
+		SetManageContext.bClearExisting = PriorityIndex == 0;
+		SetManageContext.RecurseType = UE::AssetRegistry::EDependencyCategory::Package;
+		SetManageContext.ShouldSetManager = SetManagerPredicate;
 		SlowTask.EnterProgressFrame(1);
 
-		AssetRegistry.SetManageReferences(*ManagerMap, PriorityIndex == 0, UE::AssetRegistry::EDependencyCategory::Package, ExistingManagedNodes, SetManagerPredicate);
+		AssetRegistry.SetManageReferences(SetManageContext);
 	}
 
+	SetManageContext.bClearExisting = false;
+	SetManageContext.RecurseType = UE::AssetRegistry::EDependencyCategory::None;
+	SetManageContext.ShouldSetManager = SetManagerPredicateDirectSet;
+
 	// Do non recursive set last
 	if (NoReferenceManagementMap.Num() > 0)
 	{
-		AssetRegistry.SetManageReferences(NoReferenceManagementMap, false, UE::AssetRegistry::EDependencyCategory::None, ExistingManagedNodes);
+		SetManageContext.ManagerMap = &NoReferenceManagementMap;
+		AssetRegistry.SetManageReferences(SetManageContext);
 	}
 
 	TMultiMap<FAssetIdentifier, FAssetIdentifier> PrimaryAssetIdManagementMap;
@@ -4244,7 +4253,8 @@ void UAssetManager::UpdateManagementDatabase(bool bForceRefresh)
 
 	if (PrimaryAssetIdManagementMap.Num() > 0)
 	{
-		AssetRegistry.SetManageReferences(PrimaryAssetIdManagementMap, false, UE::AssetRegistry::EDependencyCategory::None, ExistingManagedNodes);
+		SetManageContext.ManagerMap = &PrimaryAssetIdManagementMap;
+		AssetRegistry.SetManageReferences(SetManageContext);
 	}
 
 	UProjectPackagingSettings* ProjectPackagingSettings = GetMutableDefault<UProjectPackagingSettings>();
