[AssetRegistry]

Demote Warning to Display for logs that the journal has wrapped. We already prompt users who have a misconfigured journal size so this wrapping message is now just informational about slower runs that might happen through natural use of the journal.

#rnx

diff --git a/Engine/Source/Runtime/AssetRegistry/Private/AssetDataGathererDiscoveryCache.cpp b/Engine/Source/Runtime/AssetRegistry/Private/AssetDataGathererDiscoveryCache.cpp
--- a/Engine/Source/Runtime/AssetRegistry/Private/AssetDataGathererDiscoveryCache.cpp
+++ b/Engine/Source/Runtime/AssetRegistry/Private/AssetDataGathererDiscoveryCache.cpp
@@ -164,7 +164,7 @@ void FAssetDataDiscoveryCache::LoadAndUpdateCache()
 				break;
 			case EFileJournalResult::JournalWrapped:
 			{
-				UE_LOG(LogAssetRegistry, Warning,
+				UE_LOG(LogAssetRegistry, Display,
 					TEXT("PlatformFileJournal journal has wrapped for volume '%s'. AssetRegistry discovery of files on this volume will be uncached. Notes on wrapping:")
 					TEXT("\r\n%s"),
 					*VolumeName, *VolumeInfo.LastError);
