Resolve redirections in CachedBPInheritanceMap at the end of the asset scan.
This fixes an issue with blueprints class digest using the closest native base class. CachedBPInheritanceMap is built using asset tags of blueprints. When the base class of a blueprint is renamed and the blueprint is not resaved, the asset tag still points to the old base class name. The map will be updated with the correct base class name only when the blueprint is loaded. The consequence is the first cook creates dependencies using the correct base class name because the package is loaded. But during the second cook, when evaluating dependencies to decide if the package should be cooked, the package is not loaded so the old name of the base class will be used and trigger a recook when it shouldn't.

#rnx
#incrementalcook

diff --git a/Engine/Source/Runtime/AssetRegistry/Private/AssetRegistry.cpp b/Engine/Source/Runtime/AssetRegistry/Private/AssetRegistry.cpp
--- a/Engine/Source/Runtime/AssetRegistry/Private/AssetRegistry.cpp
+++ b/Engine/Source/Runtime/AssetRegistry/Private/AssetRegistry.cpp
@@ -1184,6 +1184,20 @@ void FAssetRegistryImpl::RebuildAssetDependencyGathererMapIfNeeded()
 	bRegisteredDependencyGathererClassesDirty = false;
 }
 
+void FAssetRegistryImpl::ResolveRedirectionsInCachedBPInheritanceMap()
+{
+	for (TPair<FTopLevelAssetPath, FTopLevelAssetPath>& Pair : CachedBPInheritanceMap)
+	{
+		FSoftObjectPath BaseClassPath(Pair.Value);
+		FSoftObjectPath TargetBaseClassPath = GRedirectCollector.GetAssetPathRedirection(BaseClassPath);
+		if (TargetBaseClassPath.IsAsset())
+		{
+			Pair.Value = TargetBaseClassPath.GetAssetPath();
+			TempCachedInheritanceBuffer.bDirty = true;
+		}
+	}
+}
+
 #endif
 
 }
@@ -5423,6 +5437,10 @@ void FAssetRegistryImpl::OnInitialSearchCompleted(Impl::FEventContext& EventCont
 
 	bInitialSearchCompleted.store(true, std::memory_order_relaxed);
 	UpdateMaxSecondsPerFrame();
+
+#if WITH_EDITOR
+	ResolveRedirectionsInCachedBPInheritanceMap();
+#endif
 }
 
 void FAssetRegistryImpl::OnAdditionalMountSearchCompleted(Impl::FEventContext& EventContext)
diff --git a/Engine/Source/Runtime/AssetRegistry/Private/AssetRegistryImpl.h b/Engine/Source/Runtime/AssetRegistry/Private/AssetRegistryImpl.h
--- a/Engine/Source/Runtime/AssetRegistry/Private/AssetRegistryImpl.h
+++ b/Engine/Source/Runtime/AssetRegistry/Private/AssetRegistryImpl.h
@@ -537,6 +537,10 @@ private:
 
 	/** If new gatherers were registered, this call will refresh the mapping */
 	void RebuildAssetDependencyGathererMapIfNeeded();
+
+	/** When the scan is completed, resolve the redirections in CachedBPInheritanceMap. */
+	void ResolveRedirectionsInCachedBPInheritanceMap();
+
 #endif // WITH_EDITOR
 	void GetSubClasses_Recursive(Impl::FClassInheritanceContext& InheritanceContext, FTopLevelAssetPath InClassName,
 		TSet<FTopLevelAssetPath>& SubClassNames, TSet<FTopLevelAssetPath>& ProcessedClassNames, const TSet<FTopLevelAssetPath>& ExcludedClassNames) const;
