Preallocate TArrays to known upper bounds.
#rnx

diff --git a/Engine/Source/Runtime/AssetRegistry/Private/AssetRegistry.cpp b/Engine/Source/Runtime/AssetRegistry/Private/AssetRegistry.cpp
--- a/Engine/Source/Runtime/AssetRegistry/Private/AssetRegistry.cpp
+++ b/Engine/Source/Runtime/AssetRegistry/Private/AssetRegistry.cpp
@@ -958,6 +958,7 @@ void FAssetRegistryImpl::LoadPremadeAssetRegistry(Impl::FEventContext& EventCont
 				CachePathsFromState(EventContext, State);
 				if (bEmitAssetEvents)
 				{
+					EventContext.AssetEvents.Reserve(State.GetNumAssets());
 					State.EnumerateAllAssets([&EventContext](const FAssetData& AssetData)
 					{
 						EventContext.AssetEvents.Emplace(AssetData, UE::AssetRegistry::Impl::FEventContext::EEvent::Added);
@@ -971,6 +972,7 @@ void FAssetRegistryImpl::LoadPremadeAssetRegistry(Impl::FEventContext& EventCont
 				CachePathsFromState(EventContext, State);
 				if (bEmitAssetEvents)
 				{
+					EventContext.AssetEvents.Reserve(State.GetNumAssets());
 					State.EnumerateAllAssets([&EventContext](const FAssetData& AssetData)
 					{
 						EventContext.AssetEvents.Emplace(AssetData, UE::AssetRegistry::Impl::FEventContext::EEvent::Added);
@@ -9674,6 +9676,21 @@ void UAssetRegistryImpl::Broadcast(UE::AssetRegistry::Impl::FEventContext& Event
 		if (bHasListeners)
 		{
 			TArray<FAssetData> EventBatches[EventTypeCount];
+
+			// Pre-allocate memory for event batches
+			uint32 EventBatchCounters[EventTypeCount] = { 0 };
+			for (const TPair<FAssetData, FEventContext::EEvent>& AssetEvent : EventContext.AssetEvents)
+			{
+				FEventContext::EEvent Event = AssetEvent.Get<1>();
+				EventBatchCounters[static_cast<int32>(Event)]++;
+			}
+
+			for (int32 i = 0; i < UE_ARRAY_COUNT(EventBatches); ++i)
+			{
+				EventBatches[i].Reserve(EventBatchCounters[i]);
+			}
+
+			// Dispatch events
 			FEventContext::EEvent LastEvent = EventContext.AssetEvents[0].Get<1>();
 			auto FlushBatchedEvents = [&EventBatches, &Events = BatchedAssetEvents]()
 				{
