[AssetRegistry]

Force discovery cache usage off on build machines until UE-274858 is investigated

#rnx

diff --git a/Engine/Source/Runtime/AssetRegistry/Private/AssetDataGatherer.cpp b/Engine/Source/Runtime/AssetRegistry/Private/AssetDataGatherer.cpp
--- a/Engine/Source/Runtime/AssetRegistry/Private/AssetDataGatherer.cpp
+++ b/Engine/Source/Runtime/AssetRegistry/Private/AssetDataGatherer.cpp
@@ -189,6 +189,12 @@ void FGatherSettings::Initialize()
 #endif
 	bool bNoAssetRegistryCache = FParse::Param(FCommandLine::Get(), TEXT("NoAssetRegistryCache"));
 	bool bNoAssetRegistryDiscoveryCache = bNoAssetRegistryCache || FParse::Param(FCommandLine::Get(), TEXT("NoAssetRegistryDiscoveryCache"));
+	if (GIsBuildMachine)
+	{
+		// Force discovery cache usage off until UE-274858 is investigated
+		bNoAssetRegistryDiscoveryCache = true;
+	}
+	
 	bool bNoAssetRegistryCacheRead = FParse::Param(FCommandLine::Get(), TEXT("NoAssetRegistryCacheRead"));
 	uint32 MultiprocessId = UE::GetMultiprocessId();
 	bool bMultiprocess = MultiprocessId > 0 || FParse::Param(FCommandLine::Get(), TEXT("multiprocess"));
