[Cooker]

-If AssetRegistry cache writes are disabled, ensure cook workers do not read from the discovery cache. Disabling cache writes disables the discovery cache write as well meaning that if any workers read from the discovery cache they will read in results that don't match what is actually on disk.
- Removes hack forcing discovery cache usage off for all buildmachines now that the above issue is corrected.

#rnx

diff --git a/Engine/Source/Editor/UnrealEd/Private/Cooker/CookDirector.cpp b/Engine/Source/Editor/UnrealEd/Private/Cooker/CookDirector.cpp
--- a/Engine/Source/Editor/UnrealEd/Private/Cooker/CookDirector.cpp
+++ b/Engine/Source/Editor/UnrealEd/Private/Cooker/CookDirector.cpp
@@ -1620,6 +1620,12 @@ FString FCookDirector::GetWorkerCommandLine(FWorkerId WorkerId, int32 ProfileId)
 				continue;
 			}
 		}
+		else if (Token == TEXT("-NoAssetRegistryCacheWrite"))
+		{
+			// If cache writing has been disabled, we must also disable reading of the discovery cache on workers.
+			// Otherwise workers could be looking at old discovery results that no longer match what's on disk.
+			Tokens.Add(TEXT("-NoAssetRegistryDiscoveryCache"));
+		}
 		Tokens.Add(MoveTemp(Token));
 	}
 
diff --git a/Engine/Source/Runtime/AssetRegistry/Private/AssetDataGatherer.cpp b/Engine/Source/Runtime/AssetRegistry/Private/AssetDataGatherer.cpp
--- a/Engine/Source/Runtime/AssetRegistry/Private/AssetDataGatherer.cpp
+++ b/Engine/Source/Runtime/AssetRegistry/Private/AssetDataGatherer.cpp
@@ -189,12 +189,6 @@ void FGatherSettings::Initialize()
 #endif
 	bool bNoAssetRegistryCache = FParse::Param(FCommandLine::Get(), TEXT("NoAssetRegistryCache"));
 	bool bNoAssetRegistryDiscoveryCache = bNoAssetRegistryCache || FParse::Param(FCommandLine::Get(), TEXT("NoAssetRegistryDiscoveryCache"));
-	if (GIsBuildMachine)
-	{
-		// Force discovery cache usage off until UE-274858 is investigated
-		bNoAssetRegistryDiscoveryCache = true;
-	}
-	
 	bool bNoAssetRegistryCacheRead = FParse::Param(FCommandLine::Get(), TEXT("NoAssetRegistryCacheRead"));
 	uint32 MultiprocessId = UE::GetMultiprocessId();
 	bool bMultiprocess = MultiprocessId > 0 || FParse::Param(FCommandLine::Get(), TEXT("multiprocess"));
