Added parallel option for AssetRegistry::EnumerateAllPackages and used it in UCookOnTheFlyServer::GetNeverCookPackageNames.

#rnx

diff --git a/Engine/Source/Editor/UnrealEd/Private/CookOnTheFlyServer.cpp b/Engine/Source/Editor/UnrealEd/Private/CookOnTheFlyServer.cpp
--- a/Engine/Source/Editor/UnrealEd/Private/CookOnTheFlyServer.cpp
+++ b/Engine/Source/Editor/UnrealEd/Private/CookOnTheFlyServer.cpp
@@ -12058,16 +12058,17 @@ TArray<FName> UCookOnTheFlyServer::GetNeverCookPackageNames(TArrayView<const FSt
 			}
 		}
 
-		FString PackageNameStr;
+		FRWLock Lock;
 		AssetRegistry->EnumerateAllPackages(
-			[&NeverCookPackages, &NeverCookDirectoryTree, &PackageNameStr](FName PackageName, const FAssetPackageData& PackageData)
+			[&NeverCookPackages, &NeverCookDirectoryTree, &Lock](FName PackageName, const FAssetPackageData& PackageData)
 			{
-				PackageName.ToString(PackageNameStr);
-				if (NeverCookDirectoryTree.ContainsPathOrParent(PackageNameStr))
+				FNameBuilder PackageNameBuilder(PackageName);
+				if (NeverCookDirectoryTree.ContainsPathOrParent(PackageNameBuilder))
 				{
+					FWriteScopeLock ScopeLock(Lock);
 					NeverCookPackages.Add(PackageName);
 				}
-			});
+			}, UE::AssetRegistry::EEnumeratePackagesFlags::Parallel);
 	}
 	else
 	{
diff --git a/Engine/Source/Runtime/AssetRegistry/Private/AssetRegistry.cpp b/Engine/Source/Runtime/AssetRegistry/Private/AssetRegistry.cpp
--- a/Engine/Source/Runtime/AssetRegistry/Private/AssetRegistry.cpp
+++ b/Engine/Source/Runtime/AssetRegistry/Private/AssetRegistry.cpp
@@ -3259,16 +3259,32 @@ TArray<TOptional<FAssetPackageData>> UAssetRegistryImpl::GetAssetPackageDatasCop
 	return OutAssetPackagesData;
 }
 
-void UAssetRegistryImpl::EnumerateAllPackages(TFunctionRef<void(FName PackageName, const FAssetPackageData& PackageData)> Callback) const
+void UAssetRegistryImpl::EnumerateAllPackages(TFunctionRef<void(FName PackageName, const FAssetPackageData& PackageData)> Callback, UE::AssetRegistry::EEnumeratePackagesFlags InEnumerateFlags) const
 {
 	UE::AssetRegistry::FInterfaceReadScopeLock InterfaceScopeLock(InterfaceLock);
-	for (const TPair<FName, const FAssetPackageData*>& Pair : GuardedData.GetState().GetAssetPackageDataMap())
+	if(EnumHasAnyFlags(InEnumerateFlags, UE::AssetRegistry::EEnumeratePackagesFlags::Parallel))
 	{
-		Callback(Pair.Key, *Pair.Value);
+		const TMap<FName, const FAssetPackageData*>& DataMap = GuardedData.GetState().GetAssetPackageDataMap();
+		int32 MaxId = DataMap.GetMaxIndex();
+		ParallelFor(MaxId, [&DataMap, &Callback](int32 Index)
+		{
+			FSetElementId Id = FSetElementId::FromInteger(Index);
+			if (DataMap.IsValidId(Id))
+			{
+				const TPair<FName, const FAssetPackageData*>& Pair = DataMap.Get(Id);
+				Callback(Pair.Key, *Pair.Value);
+			}
+		});
+	}
+	else
+	{
+		for (const TPair<FName, const FAssetPackageData*>& Pair : GuardedData.GetState().GetAssetPackageDataMap())
+		{
+			Callback(Pair.Key, *Pair.Value);
+		}
 	}
 }
 
-
 bool UAssetRegistryImpl::DoesPackageExistOnDisk(FName PackageName, FString* OutCorrectCasePackageName, FString* OutExtension) const
 {
 	auto CalculateExtension = [](const FString& PackageNameStr, TConstArrayView<FAssetData> Assets) -> FString
diff --git a/Engine/Source/Runtime/AssetRegistry/Private/AssetRegistry.h b/Engine/Source/Runtime/AssetRegistry/Private/AssetRegistry.h
--- a/Engine/Source/Runtime/AssetRegistry/Private/AssetRegistry.h
+++ b/Engine/Source/Runtime/AssetRegistry/Private/AssetRegistry.h
@@ -68,7 +68,7 @@ public:
 	virtual bool GetReferencers(FName PackageName, TArray<FName>& OutReferencers, UE::AssetRegistry::EDependencyCategory Category = UE::AssetRegistry::EDependencyCategory::Package, const UE::AssetRegistry::FDependencyQuery& Flags = UE::AssetRegistry::FDependencyQuery()) const override; //-V1101
 	virtual TOptional<FAssetPackageData> GetAssetPackageDataCopy(FName PackageName) const override;
 	virtual TArray<TOptional<FAssetPackageData>> GetAssetPackageDatasCopy(TArrayView<FName> PackageNames) const override;
-	virtual void EnumerateAllPackages(TFunctionRef<void(FName PackageName, const FAssetPackageData& PackageData)> Callback) const override;
+	virtual void EnumerateAllPackages(TFunctionRef<void(FName PackageName, const FAssetPackageData& PackageData)> Callback, UE::AssetRegistry::EEnumeratePackagesFlags InEnumerateFlags) const override;
 	virtual bool DoesPackageExistOnDisk(FName PackageName, FString* OutCorrectCasePackageName = nullptr, FString* OutExtension = nullptr) const override;
 	virtual FSoftObjectPath GetRedirectedObjectPath(const FSoftObjectPath& ObjectPath) override;
 	virtual bool GetAncestorClassNames(FTopLevelAssetPath ClassName, TArray<FTopLevelAssetPath>& OutAncestorClassNames) const override;
diff --git a/Engine/Source/Runtime/AssetRegistry/Public/AssetRegistry/IAssetRegistry.h b/Engine/Source/Runtime/AssetRegistry/Public/AssetRegistry/IAssetRegistry.h
--- a/Engine/Source/Runtime/AssetRegistry/Public/AssetRegistry/IAssetRegistry.h
+++ b/Engine/Source/Runtime/AssetRegistry/Public/AssetRegistry/IAssetRegistry.h
@@ -611,8 +611,13 @@ public:
 	/**
 	 * Enumerate all PackageDatas in the AssetRegistry. The callback is called from within the AssetRegistry's lock, so it must not call
 	 * arbitrary code that could call back into the AssetRegistry; doing so would deadlock.
+	 * 
+	 * @param Callback function to call for each package data enumerated
+	 * @param InEnumerateFlags flags to control enumeration
+	 *        @see EEnumeratePackagesFlags
 	 */
-	virtual void EnumerateAllPackages(TFunctionRef<void(FName PackageName, const FAssetPackageData& PackageData)> Callback) const = 0;
+	virtual void EnumerateAllPackages(TFunctionRef<void(FName PackageName, const FAssetPackageData& PackageData)> Callback,
+		UE::AssetRegistry::EEnumeratePackagesFlags InEnumerateFlags = UE::AssetRegistry::EEnumeratePackagesFlags::None) const = 0;
 
 	/*
 	* Checks if the specified package exists on disk. If the package exists and the optional parameters are provided, 
diff --git a/Engine/Source/Runtime/CoreUObject/Public/Misc/AssetRegistryInterface.h b/Engine/Source/Runtime/CoreUObject/Public/Misc/AssetRegistryInterface.h
--- a/Engine/Source/Runtime/CoreUObject/Public/Misc/AssetRegistryInterface.h
+++ b/Engine/Source/Runtime/CoreUObject/Public/Misc/AssetRegistryInterface.h
@@ -172,6 +172,13 @@ namespace UE::AssetRegistry
 		AllowUnfilteredArAssets = (1 << 2),	// skip the filtering of UE::AssetRegistry::FFiltering
 	};
 	ENUM_CLASS_FLAGS(EEnumerateAssetsFlags);
+	
+	enum class EEnumeratePackagesFlags : uint32
+	{
+		None = 0,							// No flags
+		Parallel = (1 << 0),	            // Iterate with a ParallelFor
+	};
+	ENUM_CLASS_FLAGS(EEnumeratePackagesFlags);
 
 	/**
 	 * A struct that is equivalent to EDependencyQuery, but is more useful for performance in filtering operations.
