AssetRegistry: TickGatherer: Refactor  to reduce complexity and duplication, and to reduce the amount of data that needs to be read each Tick after gathering is complete.
Moved lambdas into member function FTickContext.
Reorded the calculation of status for clarity.
Changed the non-gathering case to read the minimal amount of data possible and early exit.
#rnx
#incrementalcook

diff --git a/Engine/Source/Runtime/AssetRegistry/Private/AssetDataGatherer.cpp b/Engine/Source/Runtime/AssetRegistry/Private/AssetDataGatherer.cpp
--- a/Engine/Source/Runtime/AssetRegistry/Private/AssetDataGatherer.cpp
+++ b/Engine/Source/Runtime/AssetRegistry/Private/AssetDataGatherer.cpp
@@ -4746,7 +4746,8 @@ void FAssetDataGatherer::AddToCache(FName PackageName, FDiskCachedAssetData* Dis
 	}
 }
 
-void FAssetDataGatherer::GetAndTrimSearchResults(FResults& InOutResults, FResultContext& OutContext)
+void FAssetDataGatherer::GetAndTrimSearchResults(UE::AssetDataGather::FResults& InOutResults,
+	UE::AssetDataGather::FResultContext& OutContext)
 {
 	auto MoveAppendRangeToRingBuffer = [](auto& InOutRingBuffer, auto& InArray)
 	{
@@ -4798,7 +4799,7 @@ FAssetGatherDiagnostics FAssetDataGatherer::GetDiagnostics()
 	return Diag;
 }
 
-void FAssetDataGatherer::GetPackageResults(FResults& InOutResults)
+void FAssetDataGatherer::GetPackageResults(UE::AssetDataGather::FResults& InOutResults)
 {
 	FGathererScopeLock ResultsScopeLock(&ResultsLock);
 
diff --git a/Engine/Source/Runtime/AssetRegistry/Private/AssetDataGatherer.h b/Engine/Source/Runtime/AssetRegistry/Private/AssetDataGatherer.h
--- a/Engine/Source/Runtime/AssetRegistry/Private/AssetDataGatherer.h
+++ b/Engine/Source/Runtime/AssetRegistry/Private/AssetDataGatherer.h
@@ -72,6 +72,57 @@ struct FAssetGatherDiagnostics
 	int32 NumUncachedAssetFiles;
 };
 
+namespace UE::AssetDataGather
+{
+
+/** Structure to accumulate the results of the gather. Appended to from calls to FAssetDataGatherer::GetAndTrimSearchResults. */
+struct FResults
+{
+	TMultiMap<FName, TUniquePtr<FAssetData>> Assets;
+	TMultiMap<FName, TUniquePtr<FAssetData>> AssetsForGameThread;
+	TRingBuffer<FString> Paths;
+	TMultiMap<FName, FPackageDependencyData> Dependencies;
+	TMultiMap<FName, FPackageDependencyData> DependenciesForGameThread;
+	TRingBuffer<FString> CookedPackageNamesWithoutAssetData;
+	TRingBuffer<FName> VerseFiles;
+	TArray<FString> BlockedFiles;
+
+	SIZE_T GetAllocatedSize() const
+	{
+		return Assets.GetAllocatedSize() + AssetsForGameThread.GetAllocatedSize() + Paths.GetAllocatedSize() + Dependencies.GetAllocatedSize() +
+			DependenciesForGameThread.GetAllocatedSize() + CookedPackageNamesWithoutAssetData.GetAllocatedSize() + VerseFiles.GetAllocatedSize() +
+			BlockedFiles.GetAllocatedSize();
+	}
+	void Shrink()
+	{
+		Assets.Shrink();
+		AssetsForGameThread.Shrink();
+		Paths.Trim();
+		Dependencies.Shrink();
+		DependenciesForGameThread.Shrink();
+		CookedPackageNamesWithoutAssetData.Trim();
+		VerseFiles.Trim();
+		BlockedFiles.Shrink();
+	}
+};
+
+/**
+ * Structure to receive transient data about the current state of the gather. Repopulated during every call to
+ * FAssetDataGatherer::GetAndTrimSearchResults.
+ */
+struct FResultContext
+{
+	bool bIsSearching = false;
+	bool bAbleToProgress = false;
+	TArray<double> SearchTimes;
+	int32 NumFilesToSearch = 0;
+	int32 NumPathsToSearch = 0;
+	bool bIsDiscoveringFiles = false;
+};
+
+} // namespace UE::AssetDataGather
+
+
 /**
  * Async task for gathering asset data from from the file list in FAssetRegistry
  */
@@ -107,51 +158,13 @@ public:
 	void EnsureCompletion();
 
 
-	// Receiving Results (possibly while tick is running)
-	struct FResults
-	{
-		TMultiMap<FName, TUniquePtr<FAssetData>> Assets;
-		TMultiMap<FName, TUniquePtr<FAssetData>> AssetsForGameThread;
-		TRingBuffer<FString> Paths;
-		TMultiMap<FName, FPackageDependencyData> Dependencies;
-		TMultiMap<FName, FPackageDependencyData> DependenciesForGameThread;
-		TRingBuffer<FString> CookedPackageNamesWithoutAssetData;
-		TRingBuffer<FName> VerseFiles;
-		TArray<FString> BlockedFiles;
-
-		SIZE_T GetAllocatedSize() const
-		{
-			return Assets.GetAllocatedSize() + AssetsForGameThread.GetAllocatedSize() + Paths.GetAllocatedSize() + Dependencies.GetAllocatedSize() +
-				DependenciesForGameThread.GetAllocatedSize() + 	CookedPackageNamesWithoutAssetData.GetAllocatedSize() + VerseFiles.GetAllocatedSize() +
-				BlockedFiles.GetAllocatedSize();
-		}
-		void Shrink()
-		{
-			Assets.Shrink();
-			AssetsForGameThread.Shrink();
-			Paths.Trim();
-			Dependencies.Shrink();
-			DependenciesForGameThread.Shrink();
-			CookedPackageNamesWithoutAssetData.Trim();
-			VerseFiles.Trim();
-			BlockedFiles.Shrink();
-		}
-	};
-	struct FResultContext
-	{
-		bool bIsSearching = false;
-		bool bAbleToProgress = false;
-		TArray<double> SearchTimes;
-		int32 NumFilesToSearch = 0;
-		int32 NumPathsToSearch = 0;
-		bool bIsDiscoveringFiles = false;
-	};
 	/** Gets search results from the data gatherer. */
-	void GetAndTrimSearchResults(FResults& InOutResults, FResultContext& OutContext);
+	void GetAndTrimSearchResults(UE::AssetDataGather::FResults& InOutResults,
+		UE::AssetDataGather::FResultContext& OutContext);
 	/** Get diagnostics for telemetry or logging. */
 	FAssetGatherDiagnostics GetDiagnostics();
 	/** Gets just the Assets, AssetsForGameThread, Dependencies, and DependenciesForGameThread from the data gatherer. */
-	void GetPackageResults(FResults& InOutResults);
+	void GetPackageResults(UE::AssetDataGather::FResults& InOutResults);
 	/**
 	 * Wait for all monitored assets under the given path to be added to search results.
 	 * Returns immediately if the given path is not monitored.
diff --git a/Engine/Source/Runtime/AssetRegistry/Private/AssetRegistry.cpp b/Engine/Source/Runtime/AssetRegistry/Private/AssetRegistry.cpp
--- a/Engine/Source/Runtime/AssetRegistry/Private/AssetRegistry.cpp
+++ b/Engine/Source/Runtime/AssetRegistry/Private/AssetRegistry.cpp
@@ -1965,7 +1965,7 @@ void FAssetRegistryImpl::SearchAllAssets(Impl::FEventContext& EventContext,
 	if (bSynchronousSearch)
 	{
 		Gatherer.WaitForIdle();
-		Impl::FTickContext TickContext(EventContext, InheritanceContext);
+		Impl::FTickContext TickContext(*this, EventContext, InheritanceContext);
 		TickContext.bHandleDeferred = true;
 		TickContext.bHandleCompletion = false; // Our caller will call WaitForCompletion which will handle this
 		Impl::EGatherStatus UnusedStatus = TickGatherer(TickContext);
@@ -2042,7 +2042,7 @@ void UAssetRegistryImpl::WaitForCompletion()
 
 			GuardedData.WaitForGathererIdleIfSynchronous();
 
-			UE::AssetRegistry::Impl::FTickContext TickContext(EventContext, InheritanceContext);
+			UE::AssetRegistry::Impl::FTickContext TickContext(GuardedData, EventContext, InheritanceContext);
 			TickContext.bHandleCompletion = true;
 			TickContext.bHandleDeferred = true;
 			Status = GuardedData.TickGatherer(TickContext);
@@ -4620,7 +4620,7 @@ UE::AssetRegistry::Impl::EGatherStatus UAssetRegistryImpl::TickOnBackgroundThrea
 					return false;
 				};
 
-			UE::AssetRegistry::Impl::FTickContext TickContext(EventContext, InheritanceContext);
+			UE::AssetRegistry::Impl::FTickContext TickContext(GuardedData, EventContext, InheritanceContext);
 			TickContext.InterruptionContext.SetLimitedTickTime(FPlatformTime::Seconds(),
 				UE::AssetRegistry::Impl::MaxSecondsPerTickBackgroundThread);
 			TickContext.InterruptionContext.SetEarlyExitCallback(EarlyExitHelper);
@@ -4750,7 +4750,7 @@ void UAssetRegistryImpl::Tick(float DeltaTime)
 			if (EventContext.IsEmpty())
 			{
 				// Tick the Gatherer
-				UE::AssetRegistry::Impl::FTickContext TickContext(EventContext, InheritanceContext);
+				UE::AssetRegistry::Impl::FTickContext TickContext(GuardedData, EventContext, InheritanceContext);
 				LocalMaxSecondsPerFrame = GuardedData.MaxSecondsPerFrame;
 				TickContext.InterruptionContext.SetLimitedTickTime(TickStartTime, LocalMaxSecondsPerFrame);
 				TickContext.bHandleCompletion = true;
@@ -4899,199 +4899,59 @@ void FAssetRegistryImpl::UpdateMaxSecondsPerFrame()
 
 Impl::EGatherStatus FAssetRegistryImpl::TickGatherer(Impl::FTickContext& TickContext)
 {
-	TRACE_CPUPROFILER_EVENT_SCOPE_STR_CONDITIONAL("FAssetRegistryImpl::TickGatherer", IsGathering());
-
 	using namespace UE::AssetRegistry::Impl;
-	FEventContext& EventContext = TickContext.EventContext;
-	FInterruptionContext& InOutInterruptionContext = TickContext.InterruptionContext;
-	TOptional<TSet<FString>> MountPointsForVerifyAfterGather;
-	bool bLocalIsInGameThread = IsInGameThread();
 
-	EGatherStatus OutStatus = EGatherStatus::Complete;
 	if (!GlobalGatherer.IsValid())
 	{
-		return OutStatus;
+		return EGatherStatus::Complete;
 	}
-	double TimingStartTime = -1.;
-	auto LazyStartTimer = [&TimingStartTime]()
-	{
-		if (TimingStartTime <= 0.)
-		{
-			TimingStartTime = FPlatformTime::Seconds();
-		}
-	};
-	auto RecordTimer = [this, &TimingStartTime]()
+
+	// Gather results from the background search
+	GlobalGatherer->GetAndTrimSearchResults(BackgroundResults, TickContext.ResultContext);
+	TickContext.SetNumGatherFromDiskPendingDirty(); // Invalidate cache if used in an earlier TickGatherer call
+	if (!IsGathering() && !TickContext.ResultContext.bIsSearching && TickContext.GetNumPending() == 0)
 	{
-		if (TimingStartTime > 0.)
-		{
-			StoreGatherResultsTimeSeconds += static_cast<float>(FPlatformTime::Seconds() - TimingStartTime);
-			TimingStartTime = -1.;
-		}
-	};
+		// This is the common case of AssetRegistry ticks after the gather completed. We were already complete before
+		// the tick, and found nothing new to gather. Early exit now to avoid doing spending time checking each of the
+		// pieces of gathered data.
+		return EGatherStatus::Complete;
+	}
+
+	TRACE_CPUPROFILER_EVENT_SCOPE_STR_CONDITIONAL("FAssetRegistryImpl::TickGatherer", IsGathering());
+	// Initialize per-tick TickContext data and alias some of its fields to reduce duplication
+	FEventContext& EventContext = TickContext.EventContext;
+	FInterruptionContext& InOutInterruptionContext = TickContext.InterruptionContext;
+	TickContext.MountPointsForVerifyAfterGather.Reset();
+	TickContext.bHadAssetsToProcess = BackgroundResults.Assets.Num() > 0 || BackgroundResults.Dependencies.Num() > 0;
+	TickContext.bIsInGameThread = IsInGameThread();
+	TickContext.TimingStartTime = -1.;
 	ON_SCOPE_EXIT
 	{
-		RecordTimer();
+		TickContext.RecordTimer();
 	};
 
-	// Gather results from the background search
-	FAssetDataGatherer::FResultContext ResultContext;
-	GlobalGatherer->GetAndTrimSearchResults(BackgroundResults, ResultContext);
+	// We will be modifying the arrays that contribute to NumGatherFromDiskPending below, and we will need the updated
+	// value after we finish them or when we early exit, so mark it dirty now.
+	TickContext.SetNumGatherFromDiskPendingDirty();
 
 	// Report the search times
-	for (double SearchTime : ResultContext.SearchTimes)
+	for (double SearchTime : TickContext.ResultContext.SearchTimes)
 	{
 		UE_LOG(LogAssetRegistry, Verbose, TEXT("### Background search completed in %0.4f seconds"), SearchTime);
 	}
 
-	const bool bHadAssetsToProcess = BackgroundResults.Assets.Num() > 0 || BackgroundResults.Dependencies.Num() > 0;
-	auto GetNumGatherFromDiskPending = [&ResultContext, this]()
-	{
-		return ResultContext.NumFilesToSearch + ResultContext.NumPathsToSearch + BackgroundResults.Paths.Num() + BackgroundResults.Assets.Num() + BackgroundResults.AssetsForGameThread.Num()
-			+ BackgroundResults.Dependencies.Num() + BackgroundResults.DependenciesForGameThread.Num() + BackgroundResults.CookedPackageNamesWithoutAssetData.Num()
-			+ DeferredAssets.Num() + DeferredAssetsForGameThread.Num() + DeferredDependencies.Num() + DeferredDependenciesForGameThread.Num();
-	};
-	
-	auto GetTotalNumPackagesNeedingDependencyCalculation = [this]()
-#if WITH_EDITOR
-	{ return PackagesNeedingDependencyCalculation.Num() + PackagesNeedingDependencyCalculationOnGameThread.Num(); };
-#else
-	{ return 0; };
-#endif
-
-	int32 InitialNumPending = GetNumGatherFromDiskPending() + GetTotalNumPackagesNeedingDependencyCalculation();
-	int32 NumPending = 0;
-	auto CalculateStatus =
-		[&NumPending, &InitialNumPending, &ResultContext, &EventContext, &GetTotalNumPackagesNeedingDependencyCalculation,
-		&OutStatus, &InOutInterruptionContext, this]
-		(int32 NumGatherPending)
-	{
-		// Compute total pending, plus highest pending for this run so we can show a good progress bar
-		NumPending = NumGatherPending + (GetTotalNumPackagesNeedingDependencyCalculation() ? 1 : 0);
-		HighestPending = FMath::Max(this->HighestPending, NumPending);
-
-		if (!InOutInterruptionContext.WasInterrupted() && !ResultContext.bIsSearching && NumPending == 0)
-		{
-			OutStatus = EGatherStatus::Complete;
-		}
-		else if (!InOutInterruptionContext.WasInterrupted() && !ResultContext.bAbleToProgress)
-		{
-			OutStatus = EGatherStatus::UnableToProgress;
-		}
-		else
-		{
-			EGatherStatus NewStatus = ResultContext.bAbleToProgress ? EGatherStatus::TickActiveGatherActive
-				: EGatherStatus::TickActiveGatherIdle;
-			if (InOutInterruptionContext.WasInterrupted())
-			{
-				// When interrupted we don't know the current status, so just keep the previous status, unless
-				// the previous status is a temporary status, in which case just switch it over to TickActive
-				switch (GatherStatus)
-				{
-				case EGatherStatus::WaitingForEvents:
-				case EGatherStatus::UnableToProgress:
-					OutStatus = NewStatus;
-					break;
-				default:
-					OutStatus = this->GatherStatus;
-					break;
-				}
-			}
-			else
-			{
-				OutStatus = NewStatus;
-			}
-		}
-		if (OutStatus == Impl::EGatherStatus::TickActiveGatherIdle)
-		{
-			// if there's no additional work the gatherer thread can perform, change the status from TickActiveGatherIdle
-			// to TickGameThreadActiveGatherIdle.
-			if (DeferredAssets.Num() == 0
-#if WITH_EDITOR
-				&& PackagesNeedingDependencyCalculation.Num() == 0
-#endif
-				&& BackgroundResults.Assets.Num() == 0
-				&& BackgroundResults.Dependencies.Num() == 0
-				&& BackgroundResults.CookedPackageNamesWithoutAssetData.Num() == 0
-				&& BackgroundResults.Paths.Num() == 0
-				&& DeferredDependencies.Num() == 0)
-			{
-				OutStatus = Impl::EGatherStatus::TickGameThreadActiveGatherIdle;
-			}
-		}
-	};
-	auto UpdateStatus = [bHadAssetsToProcess, &NumPending, &ResultContext, &EventContext, &OutStatus, this]()
-		{
-			// Notify the status change, only when something changed, or when sending the final result before going idle
-			if (ResultContext.bIsSearching || bHadAssetsToProcess ||
-				(OutStatus == EGatherStatus::Complete && this->GatherStatus != EGatherStatus::Complete))
-			{
-				EventContext.ProgressUpdateData.Emplace(
-					HighestPending,					// NumTotalAssets
-					HighestPending - NumPending,	// NumAssetsProcessedByAssetRegistry
-					NumPending / 2,					// NumAssetsPendingDataLoad, divided by 2 because assets are double counted due to dependencies
-					ResultContext.bIsDiscoveringFiles // bIsDiscoveringAssetFiles
-				);
-
-				// Clear highest pending if we are complete
-				HighestPending = (OutStatus != EGatherStatus::Complete) * HighestPending;
-			}
-			this->GatherStatus = OutStatus;
-		};
-
 	// Add discovered paths
 	if (BackgroundResults.Paths.Num())
 	{
-		LazyStartTimer();
+		TickContext.LazyStartTimer();
 		PathDataGathered(EventContext, BackgroundResults.Paths, InOutInterruptionContext,
-			MountPointsForVerifyAfterGather);
+			TickContext.MountPointsForVerifyAfterGather);
 	}
 	if (InOutInterruptionContext.ShouldExitEarly())
 	{
-		CalculateStatus(GetNumGatherFromDiskPending());
-		UpdateStatus();
-		return OutStatus;
+		return TickContext.UpdateIntermediateStatus();
 	}
 
-	auto RunAssetSearchDataGathered = [this, &EventContext, &TickContext, &LazyStartTimer,
-		&InOutInterruptionContext, &MountPointsForVerifyAfterGather]
-	(TMultiMap<FName, TUniquePtr<FAssetData>>& InAssetResults,
-		TMultiMap<FName, TUniquePtr<FAssetData>>& OutDeferredAssetResults)
-		{
-			// Process the asset results
-			if (InAssetResults.Num())
-			{
-				LazyStartTimer();
-				// Mark the first amortize time
-				if (TickContext.AssetsFoundCallback.IsSet())
-				{
-					TMultiMap<FName, FAssetData*> NonOwningContainer;
-					for (auto Iter = InAssetResults.CreateIterator(); Iter; ++Iter)
-					{
-						NonOwningContainer.Add(Iter.Key(), Iter.Value().Get());
-					}
-					TickContext.AssetsFoundCallback.GetValue()(NonOwningContainer);
-				}
-
-				AssetSearchDataGathered(EventContext, InAssetResults, OutDeferredAssetResults,
-					InOutInterruptionContext, MountPointsForVerifyAfterGather);
-			}
-		};
-	auto RunDependencyDataGathered = [this, &bLocalIsInGameThread, &LazyStartTimer, &InOutInterruptionContext,
-		&MountPointsForVerifyAfterGather]
-	(TMultiMap<FName, FPackageDependencyData>& DependenciesToProcess,
-		TMultiMap<FName, FPackageDependencyData>& OutDeferredDependencies,
-		TSet<FName>* OutPackagesNeedingDependencyCalculation)
-		{
-			// Add dependencies
-			if (DependenciesToProcess.Num())
-			{
-				LazyStartTimer();
-
-				DependencyDataGathered(DependenciesToProcess, OutDeferredDependencies,
-					OutPackagesNeedingDependencyCalculation, InOutInterruptionContext, MountPointsForVerifyAfterGather);
-			}
-		};
-
 	bool bRetryAssetGathering = true;
 	int32 OriginalDeferredAssetsCount = 0;
 	int32 NumRetries = 0;
@@ -5100,22 +4960,18 @@ Impl::EGatherStatus FAssetRegistryImpl::TickGatherer(Impl::FTickContext& TickCon
 		bRetryAssetGathering = false;
 
 		// Process the normal results and defer anything that isn't ready
-		RunAssetSearchDataGathered(BackgroundResults.Assets, DeferredAssets);
+		TickContext.RunAssetSearchDataGathered(BackgroundResults.Assets, DeferredAssets);
 		if (InOutInterruptionContext.ShouldExitEarly())
 		{
-			CalculateStatus(GetNumGatherFromDiskPending());
-			UpdateStatus();
-			return OutStatus;
+			return TickContext.UpdateIntermediateStatus();
 		}
 
-		if (bLocalIsInGameThread)
+		if (TickContext.bIsInGameThread)
 		{
-			RunAssetSearchDataGathered(BackgroundResults.AssetsForGameThread, DeferredAssetsForGameThread);
+			TickContext.RunAssetSearchDataGathered(BackgroundResults.AssetsForGameThread, DeferredAssetsForGameThread);
 			if (InOutInterruptionContext.ShouldExitEarly())
 			{
-				CalculateStatus(GetNumGatherFromDiskPending());
-				UpdateStatus();
-				return OutStatus;
+				return TickContext.UpdateIntermediateStatus();
 			}
 		}
 
@@ -5123,36 +4979,34 @@ Impl::EGatherStatus FAssetRegistryImpl::TickGatherer(Impl::FTickContext& TickCon
 #if WITH_EDITOR
 		PackagesNeedingDependencyCalculationPointer = &PackagesNeedingDependencyCalculation;
 #endif
-		RunDependencyDataGathered(BackgroundResults.Dependencies, DeferredDependencies, PackagesNeedingDependencyCalculationPointer);
+		TickContext.RunDependencyDataGathered(BackgroundResults.Dependencies, DeferredDependencies,
+			PackagesNeedingDependencyCalculationPointer);
 		if (InOutInterruptionContext.ShouldExitEarly())
 		{
-			CalculateStatus(GetNumGatherFromDiskPending());
-			UpdateStatus();
-			return OutStatus;
+			return TickContext.UpdateIntermediateStatus();
 		}
 
-		if (bLocalIsInGameThread)
+		if (TickContext.bIsInGameThread)
 		{
 #if WITH_EDITOR
 			PackagesNeedingDependencyCalculationPointer = &PackagesNeedingDependencyCalculationOnGameThread;
 #endif
-			RunDependencyDataGathered(BackgroundResults.DependenciesForGameThread, DeferredDependenciesForGameThread, PackagesNeedingDependencyCalculationPointer);
+			TickContext.RunDependencyDataGathered(BackgroundResults.DependenciesForGameThread,
+				DeferredDependenciesForGameThread, PackagesNeedingDependencyCalculationPointer);
 			if (InOutInterruptionContext.ShouldExitEarly())
 			{
-				CalculateStatus(GetNumGatherFromDiskPending());
-				UpdateStatus();
-				return OutStatus;
+				return TickContext.UpdateIntermediateStatus();
 			}
 		}
 
 		// Retry deferred assets if we've finished all the other assets; we need to do this in the current tick
 		// so we avoid spuriously reporting status == UnableToProgress
 		if (BackgroundResults.Assets.IsEmpty()
-			&& (!bLocalIsInGameThread || BackgroundResults.AssetsForGameThread.IsEmpty())
+			&& (!TickContext.bIsInGameThread || BackgroundResults.AssetsForGameThread.IsEmpty())
 			&& TickContext.bHandleDeferred)
 		{
 			if (!DeferredAssets.IsEmpty() || !DeferredDependencies.IsEmpty() ||
-				(bLocalIsInGameThread && (!DeferredAssetsForGameThread.IsEmpty() || !DeferredDependenciesForGameThread.IsEmpty())))
+				(TickContext.bIsInGameThread && (!DeferredAssetsForGameThread.IsEmpty() || !DeferredDependenciesForGameThread.IsEmpty())))
 			{
 				if (bProcessedAnyAssetsAfterRetryDeferred)
 				{
@@ -5192,20 +5046,18 @@ Impl::EGatherStatus FAssetRegistryImpl::TickGatherer(Impl::FTickContext& TickCon
 	// Load cooked packages that do not have asset data
 	if (BackgroundResults.CookedPackageNamesWithoutAssetData.Num())
 	{
-		LazyStartTimer();
+		TickContext.LazyStartTimer();
 		CookedPackageNamesWithoutAssetDataGathered(EventContext, BackgroundResults.CookedPackageNamesWithoutAssetData, InOutInterruptionContext);
 		if (InOutInterruptionContext.ShouldExitEarly())
 		{
-			CalculateStatus(GetNumGatherFromDiskPending());
-			UpdateStatus();
-			return OutStatus;
+			return TickContext.UpdateIntermediateStatus();
 		}
 	}
 
 	// Add Verse files
 	if (BackgroundResults.VerseFiles.Num())
 	{
-		LazyStartTimer();
+		TickContext.LazyStartTimer();
 		if (TickContext.VerseFilesFoundCallback.IsSet())
 		{
 			TickContext.VerseFilesFoundCallback.GetValue()(BackgroundResults.VerseFiles);
@@ -5214,9 +5066,7 @@ Impl::EGatherStatus FAssetRegistryImpl::TickGatherer(Impl::FTickContext& TickCon
 		VerseFilesGathered(EventContext, BackgroundResults.VerseFiles, InOutInterruptionContext);
 		if (InOutInterruptionContext.ShouldExitEarly())
 		{
-			CalculateStatus(GetNumGatherFromDiskPending());
-			UpdateStatus();
-			return OutStatus;
+			return TickContext.UpdateIntermediateStatus();
 		}
 	}
 
@@ -5227,105 +5077,299 @@ Impl::EGatherStatus FAssetRegistryImpl::TickGatherer(Impl::FTickContext& TickCon
 		BackgroundResults.BlockedFiles.Reset();
 	}
 
-	int32 NumGatherFromDiskPending = GetNumGatherFromDiskPending();
-#if WITH_EDITOR
-	// Load Calculated Dependencies when the gather from disk is complete; the full gather is not complete until after this is done
-	bool bDiskGatherComplete = !ResultContext.bIsSearching && NumGatherFromDiskPending == 0;
+	bool bDiskGatherComplete = !TickContext.ResultContext.bIsSearching
+		&& TickContext.GetNumGatherFromDiskPending() == 0;
+	if (!bDiskGatherComplete)
+	{
+		// We're still gathering from disk or deferred data in the arrays above, we can't run the on-disk-gather
+		// complete code below so there is nothing further to do for now.
+		return TickContext.UpdateIntermediateStatus();
+	}
 
-	// We can't do do this work until we've finished startup because modules might add new entries to RegisteredDependencyGathererClasses
-	// as they are loaded.
-	if (bDiskGatherComplete && IsEngineStartupModuleLoadingComplete()
-		&& (PackagesNeedingDependencyCalculation.Num() || PackagesNeedingDependencyCalculationOnGameThread.Num()))
+	// Execute calculations in the gather that have to happen after the gather from disk is completed; the gather is
+	// not marked complete until these are done.
+
+	if (!IsEngineStartupModuleLoadingComplete() || !bPreloadingComplete)
 	{
-		LazyStartTimer();
-		// Only assets whose classes have a RegisteredDependencyGathererClasses entry actually need to run through LoadCalculatedDependencies
-		// at all. Furthermore, in that situation, we must always perform the gather on the game thread. This function ensures that happens and
-		// clears out any spurious entries.
+		// We can't complete the gather until we've finished startup, because more modules might mount up to that
+		// point, and when mounted add more paths that we need to gather.
+		// We also can't do some of the calculations until then. LoadCalculatedDependencies has to wait because
+		// modules might add new entries to RegisteredDependencyGathererClasses as they are loaded. 
+		// We also cannot complete, or run calculations requiring knowledge of all assets, if we have a preloaded
+		// AssetRegistry and we have not finished preloading it.
+
+		// Since we can not progress on some of the calculations, return now.
+		return TickContext.UpdateIntermediateStatus();
+	}
+
+#if WITH_EDITOR
+	// Load Calculated Dependencies for types that register for them via REGISTER_ASSETDEPENDENCY_GATHERER
+	if (PackagesNeedingDependencyCalculation.Num() || PackagesNeedingDependencyCalculationOnGameThread.Num())
+	{
+		TickContext.LazyStartTimer();
+		// Only assets whose classes have a RegisteredDependencyGathererClasses entry need to run through
+		// LoadCalculatedDependencies. Furthermore, we must always perform their gather on the game thread.
+		// PruneAndCoalesce does the non-game thread calculations we can, including removing assets that don't need
+		// any calculations, and moves the remaining issues to the container that needs action from the game thread.
 		PruneAndCoalescePackagesRequiringDependencyCalculation(PackagesNeedingDependencyCalculation,
 			PackagesNeedingDependencyCalculationOnGameThread, InOutInterruptionContext);
 		if (InOutInterruptionContext.ShouldExitEarly())
 		{
-			CalculateStatus(NumGatherFromDiskPending);
-			UpdateStatus();
-			return OutStatus;
+			return TickContext.UpdateIntermediateStatus();
 		}
-		// As currently implemented, we cannot perform the dependency calculations on a background thread. Once we've
-		// called PruneAndCoalesce, all packages that actually need calculations will be in the OnGameThread container
-		// and all other packages will have been removed.
+		// All remaining assets should now be in PackagesNeedingDependencyCalculationOnGameThread.
 		ensure(PackagesNeedingDependencyCalculation.Num() == 0);
 
-		if (PackagesNeedingDependencyCalculationOnGameThread.Num() && bLocalIsInGameThread)
+		if (PackagesNeedingDependencyCalculationOnGameThread.Num() && !TickContext.bIsInGameThread)
 		{
-			LoadCalculatedDependencies(nullptr, TickContext.InheritanceContext,
-				&PackagesNeedingDependencyCalculationOnGameThread, InOutInterruptionContext);
-			if (InOutInterruptionContext.ShouldExitEarly())
-			{
-				CalculateStatus(NumGatherFromDiskPending);
-				UpdateStatus();
-				return OutStatus;
-			}
+			// Since we can not progress, return now.
+			return TickContext.UpdateIntermediateStatus();
+		}
+
+		LoadCalculatedDependencies(nullptr, TickContext.InheritanceContext,
+			&PackagesNeedingDependencyCalculationOnGameThread, InOutInterruptionContext);
+		if (InOutInterruptionContext.ShouldExitEarly())
+		{
+			return TickContext.UpdateIntermediateStatus();
 		}
 	}
 #endif
 
-	CalculateStatus(NumGatherFromDiskPending);
+	// Finishing the gather is only possible on the game thread (and only for callsites of TickGather that opt into it)
+	// because we need to not mark the gather complete until we have broadcast the events, and broadcasting events by
+	// contract is only allowed from game thread.
+	if (!TickContext.bIsInGameThread || !TickContext.bHandleCompletion)
+	{
+		return TickContext.SetIntermediateStatus(EGatherStatus::TickGameThreadActiveGatherIdle);
+	}
 
-	if (OutStatus == EGatherStatus::Complete)
+	if (!IsInitialSearchCompleted())
 	{
-		if (!IsInitialSearchCompleted())
+		if (!EventContext.AssetEvents.IsEmpty())
 		{
-			// Finishing the background search is blocked until preloading complete because plugins can be mounted during
-			// startup up until that point, and we need to wait for all the plugins to load before declaring completion.
-			// Only the main thread can know that we're complete because we need to wait until we've broadcast the events
-			bool bCanCompleteInitialSearch = bPreloadingComplete && IsEngineStartupModuleLoadingComplete()
-				&& bLocalIsInGameThread && TickContext.bHandleCompletion;
+			// Don't mark the InitialSearch completed until we have sent all the AssetDataAdded events
+			// that arose from the final tick of the gatherer. Some callers might do more expensive
+			// work for assets added after the initial search completed, and we don't want them to do
+			// that more expensive work on the last batch of assets before completion.
+			return TickContext.SetIntermediateStatus(EGatherStatus::WaitingForEvents);
+		}
 
-			if (bCanCompleteInitialSearch)
-			{
-				if (!EventContext.AssetEvents.IsEmpty())
-				{
-					// Don't mark the InitialSearch completed until we have sent all the AssetDataAdded events
-					// that arose from the final tick of the gatherer. Some callers might do more expensive
-					// work for assets added after the initial search completed, and we don't want them to do
-					// that more expensive work on the last batch of assets before completion.
-					OutStatus = EGatherStatus::WaitingForEvents;
-					bCanCompleteInitialSearch = false;
-				}
-			}
-			else
-			{
-				if (bLocalIsInGameThread && TickContext.bHandleCompletion)
-				{
-					UE_LOG(LogAssetRegistry, Display, TEXT("TickGatherer returning UnableToProgress because bCanCompleteInitialSearch is false but our work is otherwise complete. "
-						"bPreloadingComplete == %s; IsEngineStartupModuleLoadingComplete() == %s; bLocalIsInGameThread == %s"),
-						bPreloadingComplete ? TEXT("TRUE") : TEXT("FALSE"),
-						IsEngineStartupModuleLoadingComplete() ? TEXT("TRUE") : TEXT("FALSE"),
-						bLocalIsInGameThread ? TEXT("TRUE") : TEXT("FALSE"));
-				}
+		TickContext.RecordTimer(); // OnInitialSearchComplete reads data set by RecordTimer
+		OnInitialSearchCompleted(EventContext);
+	}
+	else if (IsAdditionalMountSearchInProgress())
+	{
+		TickContext.RecordTimer(); // OnAdditionalMountSearchComplete reads data set by RecordTimer
+		OnAdditionalMountSearchCompleted(EventContext);
+	}
 
-				OutStatus = EGatherStatus::UnableToProgress;
-			}
-			if (bCanCompleteInitialSearch)
-			{
-				RecordTimer(); // OnInitialSearchComplete reads data set by RecordTimer
-				OnInitialSearchCompleted(EventContext);
-			}
+	// Give ProgressUpdateData a final call before going idle so that listeners can record 0 remaining.
+	EventContext.ProgressUpdateData.Emplace(
+		HighestPending,	// NumTotalAssets
+		HighestPending,	// NumAssetsProcessedByAssetRegistry
+		0,				// NumAssetsPendingDataLoad
+		false			// bIsDiscoveringAssetFiles
+	);
+	// Clear HighestPending 
+	HighestPending = 0;
+	this->GatherStatus = EGatherStatus::Complete;
+	return EGatherStatus::Complete;
+}
+
+namespace Impl // UE::AssetRegistry::Impl
+{
+
+FTickContext::FTickContext(FAssetRegistryImpl& InGuardedData, Impl::FEventContext& InEventContext,
+	Impl::FClassInheritanceContext& InInheritanceContext)
+	: GuardedData(InGuardedData)
+	, EventContext(InEventContext)
+	, InheritanceContext(InInheritanceContext)
+{
+}
+
+void FTickContext::LazyStartTimer()
+{
+	if (TimingStartTime <= 0.)
+	{
+		TimingStartTime = FPlatformTime::Seconds();
+	}
+};
+void FTickContext::RecordTimer()
+{
+	if (TimingStartTime > 0.)
+	{
+		GuardedData.StoreGatherResultsTimeSeconds += static_cast<float>(FPlatformTime::Seconds() - TimingStartTime);
+		TimingStartTime = -1.;
+	}
+};
+
+void FTickContext::SetNumGatherFromDiskPendingDirty()
+{
+	bNumGatherFromDiskPendingDirty = true;
+}
+
+int32 FTickContext::GetNumGatherFromDiskPending()
+{
+	if (bNumGatherFromDiskPendingDirty)
+	{
+		bNumGatherFromDiskPendingDirty = false;
+		NumGatherFromDiskPending = CalculateNumGatherFromDiskPending();
+	}
+	return NumGatherFromDiskPending;
+}
+
+int32 FTickContext::GetNumPending()
+{
+	int32 NumPending = GetNumGatherFromDiskPending();
+#if WITH_EDITOR
+	int32 NumNeedingCalculation = GuardedData.PackagesNeedingDependencyCalculation.Num()
+		+ GuardedData.PackagesNeedingDependencyCalculationOnGameThread.Num();
+	NumPending += NumNeedingCalculation > 0 ? 1 : 0;
+#endif
+	return NumPending;
+}
+
+int32 FTickContext::CalculateNumGatherFromDiskPending()
+{
+	// Structure that holds temporary data for the current call to TickGatherer
+	UE::AssetDataGather::FResultContext& C = this->ResultContext;
+	// Structure that holds persistent data stored across ticks.
+	UE::AssetDataGather::FResults& B = GuardedData.BackgroundResults;
+	// Reference to GuardedData for persistent data stored as fields on GuardedData rather than fields on the
+	// FResults struct
+	FAssetRegistryImpl& G = GuardedData;
+
+	return
+		C.NumFilesToSearch + C.NumPathsToSearch
+
+		+ B.Paths.Num() + B.Assets.Num() + B.AssetsForGameThread.Num() + B.Dependencies.Num()
+		+ B.DependenciesForGameThread.Num() + B.CookedPackageNamesWithoutAssetData.Num()
+		+ B.VerseFiles.Num() + B.BlockedFiles.Num()
+
+		+ G.DeferredAssets.Num() + G.DeferredAssetsForGameThread.Num() + G.DeferredDependencies.Num()
+		+ G.DeferredDependenciesForGameThread.Num()
+		;
+};
+
+EGatherStatus FTickContext::SetIntermediateStatus(EGatherStatus Status)
+{
+	// Report total pending and highest pending for this run so we can show a good progress bar
+	int32 LocalNumPending = GetNumPending();
+	int32& HighestPending = GuardedData.HighestPending;
+	HighestPending = FMath::Max(HighestPending, LocalNumPending);
+	if (ResultContext.bIsSearching || bHadAssetsToProcess)
+	{
+		EventContext.ProgressUpdateData.Emplace(
+			HighestPending,						// NumTotalAssets
+			HighestPending - LocalNumPending,	// NumAssetsProcessedByAssetRegistry
+			LocalNumPending / 2,				// NumAssetsPendingDataLoad, divided by 2 because assets are double counted due to dependencies
+			ResultContext.bIsDiscoveringFiles	// bIsDiscoveringAssetFiles
+		);
+	}
+	GuardedData.GatherStatus = Status;
+	return Status;
+};
+
+EGatherStatus FTickContext::UpdateIntermediateStatus()
+{
+	EGatherStatus Status;
+	if (InterruptionContext.WasInterrupted())
+	{
+		// When interrupted we don't know the current status, so just keep the previous status, unless
+		// the previous status is a temporary status, in which case just switch it over to TickActive
+		switch (GuardedData.GatherStatus)
+		{
+		case EGatherStatus::WaitingForEvents: [[fallthrough]];
+		case EGatherStatus::UnableToProgress:
+			Status = ResultContext.bAbleToProgress ? EGatherStatus::TickActiveGatherActive
+				: EGatherStatus::TickActiveGatherIdle;
+			break;
+		default:
+			Status = GuardedData.GatherStatus;
+			break;
 		}
-		else if (IsAdditionalMountSearchInProgress())
+	}
+	else if (ResultContext.bIsSearching && !ResultContext.bAbleToProgress)
+	{
+		// Gather from disk is still in progress but is blocked
+		Status = EGatherStatus::UnableToProgress;
+	}
+	else if (ResultContext.bIsSearching)
+	{
+		// We still have work we can do in the gather
+		Status = EGatherStatus::TickActiveGatherActive;
+	}
+	else if (GetNumPending() != 0)
+	{
+		// No work remaining to gather from disk, but we still have calculations we can run
+		Status = EGatherStatus::TickActiveGatherIdle;
+	}
+	else
+	{
+		// We have completed the gather from disk, but we still have calculations we need to run, but we can't
+		// run them now because we have to wait for EngineStartup or some other condition outside our control.
+		Status = EGatherStatus::UnableToProgress;
+	}
+
+	if (Status == Impl::EGatherStatus::TickActiveGatherIdle)
+	{
+		// if there's no additional work the gatherer thread can perform, change the status from TickActiveGatherIdle
+		// to TickGameThreadActiveGatherIdle.
+		UE::AssetDataGather::FResults& B = GuardedData.BackgroundResults;
+		if (GuardedData.DeferredAssets.Num() == 0
+#if WITH_EDITOR
+			&& GuardedData.PackagesNeedingDependencyCalculation.Num() == 0
+#endif
+			&& GuardedData.DeferredDependencies.Num() == 0
+			&& B.Assets.Num() == 0
+			&& B.Dependencies.Num() == 0
+			&& B.CookedPackageNamesWithoutAssetData.Num() == 0
+			&& B.Paths.Num() == 0)
+		{
+			Status = Impl::EGatherStatus::TickGameThreadActiveGatherIdle;
+		}
+	}
+	return SetIntermediateStatus(Status);
+};
+
+void FTickContext::RunAssetSearchDataGathered(TMultiMap<FName, TUniquePtr<FAssetData>>& InAssetResults,
+	TMultiMap<FName, TUniquePtr<FAssetData>>& OutDeferredAssetResults)
+{
+	// Process the asset results
+	if (InAssetResults.Num())
+	{
+		LazyStartTimer();
+		// Mark the first amortize time
+		if (AssetsFoundCallback.IsSet())
 		{
-			bool bCanCompleteAdditionalMountSearch = bLocalIsInGameThread && TickContext.bHandleCompletion;
-			if (bCanCompleteAdditionalMountSearch)
+			TMultiMap<FName, FAssetData*> NonOwningContainer;
+			for (auto Iter = InAssetResults.CreateIterator(); Iter; ++Iter)
 			{
-				RecordTimer(); // OnAdditionalMountSearchComplete reads data set by RecordTimer
-				OnAdditionalMountSearchCompleted(EventContext);
+				NonOwningContainer.Add(Iter.Key(), Iter.Value().Get());
 			}
+			AssetsFoundCallback.GetValue()(NonOwningContainer);
 		}
+
+		GuardedData.AssetSearchDataGathered(EventContext, InAssetResults, OutDeferredAssetResults,
+			InterruptionContext, MountPointsForVerifyAfterGather);
 	}
+};
+
+void FTickContext::RunDependencyDataGathered(TMultiMap<FName, FPackageDependencyData>& DependenciesToProcess,
+	TMultiMap<FName, FPackageDependencyData>& OutDeferredDependencies,
+	TSet<FName>* OutPackagesNeedingDependencyCalculation)
+{
+	// Add dependencies
+	if (DependenciesToProcess.Num())
+	{
+		LazyStartTimer();
 
-	UpdateStatus();
+		GuardedData.DependencyDataGathered(DependenciesToProcess, OutDeferredDependencies,
+			OutPackagesNeedingDependencyCalculation, InterruptionContext, MountPointsForVerifyAfterGather);
+	}
+};
 
-	return OutStatus;
-}
+} // namespace Impl within UE::AssetRegistry
 
 void FAssetRegistryImpl::Shrink()
 {
@@ -6273,7 +6317,7 @@ void FAssetRegistryImpl::ScanPathsSynchronous(Impl::FScanPathContext& Context)
 			}
 		};
 
-	Impl::FTickContext TickContext(Context.EventContext, Context.InheritanceContext);
+	Impl::FTickContext TickContext(*this, Context.EventContext, Context.InheritanceContext);
 	TickContext.AssetsFoundCallback = Impl::FAssetsFoundCallback(AssetsFoundCallback);
 	TickContext.VerseFilesFoundCallback = Impl::FVerseFilesFoundCallback(VerseFileFoundCallback);
 	Context.Status = TickGatherer(TickContext);
@@ -6332,7 +6376,7 @@ void FAssetRegistryImpl::ScanPathsSynchronous(Impl::FScanPathContext& Context)
 
 		bForceCompletionEvenIfPostLoadsFail = bOldForceCompletionEvenIfPostLoadsFail;
 		// Tick to perform any subsequent processing required for these assets beyond AssetSearchDataGathered
-		Impl::FTickContext AssetTickContext(Context.EventContext, Context.InheritanceContext);
+		Impl::FTickContext AssetTickContext(*this, Context.EventContext, Context.InheritanceContext);
 		AssetTickContext.AssetsFoundCallback = Impl::FAssetsFoundCallback(AssetsFoundCallback);
 		Context.Status = TickGatherer(AssetTickContext);
 	}
diff --git a/Engine/Source/Runtime/AssetRegistry/Private/AssetRegistryImpl.h b/Engine/Source/Runtime/AssetRegistry/Private/AssetRegistryImpl.h
--- a/Engine/Source/Runtime/AssetRegistry/Private/AssetRegistryImpl.h
+++ b/Engine/Source/Runtime/AssetRegistry/Private/AssetRegistryImpl.h
@@ -709,7 +709,7 @@ private:
 	TUniquePtr<FAssetDataGatherer> GlobalGatherer;
 
 	/** Lists of results from the gatherer thread that are waiting to get processed */
-	FAssetDataGatherer::FResults BackgroundResults;
+	UE::AssetDataGather::FResults BackgroundResults;
 
 	/** Assets and dependencies that are not ready for processing because they can't yet run PostLoadAssetRegistrytags */
 	TMultiMap<FName, TUniquePtr<FAssetData>> DeferredAssets;
@@ -828,8 +828,10 @@ private:
 	TSet<FTopLevelAssetPath> SkipCookedClasses;
 #endif
 	UE::AssetRegistry::Premade::FAsyncConsumer AsyncConsumer;
-	friend struct UE::AssetRegistry::Premade::FAsyncConsumer;
-	friend struct Impl::FClassInheritanceContext;
+	friend UE::AssetRegistry::Premade::FAsyncConsumer;
+	friend UE::AssetRegistry::Impl::FClassInheritanceContext;
+	friend UE::AssetRegistry::Impl::FTickContext;
+
 };
 
 namespace Impl
@@ -924,18 +926,69 @@ typedef TFunctionRef<void(const TRingBuffer<FName>&)> FVerseFilesFoundCallback;
 /** Input and output variables for the TickGatherer function. */
 struct FTickContext
 {
-	FTickContext(Impl::FEventContext& InEventContext, Impl::FClassInheritanceContext& InInheritanceContext)
-		: EventContext(InEventContext)
-		, InheritanceContext(InInheritanceContext)
-	{
-	}
+public:
+	FTickContext(FAssetRegistryImpl& InGuardedData, Impl::FEventContext& InEventContext,
+		Impl::FClassInheritanceContext& InInheritanceContext);
+
+	/** Start the timer to measure time spent in the current call to TickGatherer, if not already started. */
+	void LazyStartTimer();
+	/**
+	 * Record the time that was spent in the current call to TickGatherer onto the metric for it on GuardedData, if
+	 * LazyStartTimer has been called and RecordTimer has not already been called.
+	 */
+	void RecordTimer();
+
+	/**
+	 * Return progress-bar metric for how many elements remain in the various lists of work to do. A portion of this
+	 * calculation is cached (the GetNumGatherFromDiskPending) and it will not get the most up to date data until
+	 * SetNumGatherFromDiskPendingDirty is called.
+	 */
+	int32 GetNumPending();
+	/**
+	 * Return the portion of GetNumPending that is handled during gathering from disk; does not include the
+	 * calculations that we do once the gather from disk is done. This function is cached, it does not redo the
+	 * calculation until SetNumPendingDirty is called.
+	 */
+	int32 GetNumGatherFromDiskPending();
+	/** Invalidate the cache in GetNumGatherFromDiskPending  */
+	void SetNumGatherFromDiskPendingDirty();
+	/** Run the expression to calculate the value used in GetNumGatherFromDiskPending. */
+	int32 CalculateNumGatherFromDiskPending();
+
+	/**
+	 * Store the result of TickGatherer in the AssetRegistry's status variable and send the progress event.
+	 * Returns the input Status so it can be used in a return line.
+	 */
+	EGatherStatus SetIntermediateStatus(EGatherStatus Status);
+	/**
+	 * Calculate the result of TickGatherer from the current gather data on the TickContext and call
+	 * SetIntermediateStatus. Returns the input Status so it can be used in a return line.
+	 */
+	EGatherStatus UpdateIntermediateStatus();
+
+
+	void RunAssetSearchDataGathered(TMultiMap<FName, TUniquePtr<FAssetData>>& InAssetResults,
+		TMultiMap<FName, TUniquePtr<FAssetData>>& OutDeferredAssetResults);
+	void RunDependencyDataGathered(TMultiMap<FName, FPackageDependencyData>& DependenciesToProcess,
+		TMultiMap<FName, FPackageDependencyData>& OutDeferredDependencies,
+		TSet<FName>* OutPackagesNeedingDependencyCalculation);
+
+public:
+	UE::AssetDataGather::FResultContext ResultContext;
 	Impl::FInterruptionContext InterruptionContext;
+	FAssetRegistryImpl& GuardedData;
 	Impl::FEventContext& EventContext;
 	Impl::FClassInheritanceContext& InheritanceContext;
 	TOptional<FAssetsFoundCallback> AssetsFoundCallback;
 	TOptional<FVerseFilesFoundCallback> VerseFilesFoundCallback;
+	TOptional<TSet<FString>> MountPointsForVerifyAfterGather;
+	double TimingStartTime = -1.;
+	int32 NumGatherFromDiskPending = 0;
 	bool bHandleCompletion = false;
 	bool bHandleDeferred = false;
+	bool bNumGatherFromDiskPendingDirty = true;
+	bool bHadAssetsToProcess = false;
+	bool bIsInGameThread = false;
 };
 
 } // namespace Impl within namespace UE::AssetRegistry
