[Backout] - CL42326408 (Manual - fix incorrectly robomerged backout)
#rnx
Original CL Desc
-----------------------------------------------------------------
AssetRegistry: TickGatherer: Refactor  to reduce complexity and duplication, and to reduce the amount of data that needs to be read each Tick after gathering is complete.
Moved lambdas into member function FTickContext.
Reorded the calculation of status for clarity.
Changed the non-gathering case to read the minimal amount of data possible and early exit.
#rnx
#incrementalcook

diff --git a/Engine/Source/Runtime/AssetRegistry/Private/AssetRegistry.cpp b/Engine/Source/Runtime/AssetRegistry/Private/AssetRegistry.cpp
--- a/Engine/Source/Runtime/AssetRegistry/Private/AssetRegistry.cpp
+++ b/Engine/Source/Runtime/AssetRegistry/Private/AssetRegistry.cpp
@@ -5031,9 +5031,6 @@ Impl::EGatherStatus FAssetRegistryImpl::TickGatherer(Impl::FTickContext& TickCon
 					NumPending / 2,					// NumAssetsPendingDataLoad, divided by 2 because assets are double counted due to dependencies
 					ResultContext.bIsDiscoveringFiles // bIsDiscoveringAssetFiles
 				);
-
-				// Clear highest pending if we are complete
-				HighestPending = (OutStatus != EGatherStatus::Complete) * HighestPending;
 			}
 			this->GatherStatus = OutStatus;
 		};
@@ -5323,6 +5320,19 @@ Impl::EGatherStatus FAssetRegistryImpl::TickGatherer(Impl::FTickContext& TickCon
 	}
 
 	UpdateStatus();
+	if (OutStatus == EGatherStatus::Complete)
+	{
+		HighestPending = 0;
+		BackgroundResults.Shrink();
+		DeferredAssets.Shrink();
+		DeferredAssetsForGameThread.Shrink();
+		DeferredDependencies.Shrink();
+		DeferredDependenciesForGameThread.Shrink();
+#if WITH_EDITOR
+		PackagesNeedingDependencyCalculation.Shrink();
+		PackagesNeedingDependencyCalculationOnGameThread.Shrink();
+#endif
+	}
 
 	return OutStatus;
 }
