[Backout] - CL42326408 (Manual - fix incorrectly robomerged backout)
#rnx
#incrementalcook

diff --git a/Engine/Source/Runtime/AssetRegistry/Private/AssetRegistry.cpp b/Engine/Source/Runtime/AssetRegistry/Private/AssetRegistry.cpp
--- a/Engine/Source/Runtime/AssetRegistry/Private/AssetRegistry.cpp
+++ b/Engine/Source/Runtime/AssetRegistry/Private/AssetRegistry.cpp
@@ -5031,6 +5031,9 @@ Impl::EGatherStatus FAssetRegistryImpl::TickGatherer(Impl::FTickContext& TickCon
 					NumPending / 2,					// NumAssetsPendingDataLoad, divided by 2 because assets are double counted due to dependencies
 					ResultContext.bIsDiscoveringFiles // bIsDiscoveringAssetFiles
 				);
+
+				// Clear highest pending if we are complete
+				HighestPending = (OutStatus != EGatherStatus::Complete) * HighestPending;
 			}
 			this->GatherStatus = OutStatus;
 		};
@@ -5320,19 +5323,6 @@ Impl::EGatherStatus FAssetRegistryImpl::TickGatherer(Impl::FTickContext& TickCon
 	}
 
 	UpdateStatus();
-	if (OutStatus == EGatherStatus::Complete)
-	{
-		HighestPending = 0;
-		BackgroundResults.Shrink();
-		DeferredAssets.Shrink();
-		DeferredAssetsForGameThread.Shrink();
-		DeferredDependencies.Shrink();
-		DeferredDependenciesForGameThread.Shrink();
-#if WITH_EDITOR
-		PackagesNeedingDependencyCalculation.Shrink();
-		PackagesNeedingDependencyCalculationOnGameThread.Shrink();
-#endif
-	}
 
 	return OutStatus;
 }
