[AR] Add ability to exclude some asset bundles from AssetRegistry

diff --git a/Engine/Source/Editor/UnrealEd/Private/Commandlets/AssetRegistryGenerator.cpp b/Engine/Source/Editor/UnrealEd/Private/Commandlets/AssetRegistryGenerator.cpp
--- a/Engine/Source/Editor/UnrealEd/Private/Commandlets/AssetRegistryGenerator.cpp
+++ b/Engine/Source/Editor/UnrealEd/Private/Commandlets/AssetRegistryGenerator.cpp
@@ -1288,7 +1288,7 @@ void FAssetRegistryGenerator::Initialize(bool bInitializeFromExisting)
 
 	FAssetRegistrySerializationOptions SaveOptions;
 
-	AssetRegistry.InitializeSerializationOptions(SaveOptions, TargetPlatform->IniPlatformName());
+	AssetRegistry.InitializeSerializationOptions(SaveOptions, TargetPlatform, UE::AssetRegistry::ESerializationTarget::ForGame);
 
 	if (bInitializeFromExisting)
 	{
@@ -1306,7 +1306,7 @@ void FAssetRegistryGenerator::Initialize(bool bInitializeFromExisting)
 void FAssetRegistryGenerator::CloneGlobalAssetRegistryFilteredByPreviousState(const FAssetRegistryState& PreviousState)
 {
 	FAssetRegistrySerializationOptions SaveOptions;
-	AssetRegistry.InitializeSerializationOptions(SaveOptions, TargetPlatform->IniPlatformName());
+	AssetRegistry.InitializeSerializationOptions(SaveOptions, TargetPlatform, UE::AssetRegistry::ESerializationTarget::ForGame);
 
 	TSet<FName> KeepPackages;
 	const TMap<FName, const FAssetPackageData*> &PreviousPackageDatas = PreviousState.GetAssetPackageDataMap();
@@ -1755,7 +1755,7 @@ void FAssetRegistryGenerator::FinalizeChunkIDs(const TSet<FName>& InCookedPackag
 	// Prune our asset registry to cooked + dev only list
 	{
 		FAssetRegistrySerializationOptions DevelopmentSaveOptions;
-		AssetRegistry.InitializeSerializationOptions(DevelopmentSaveOptions, TargetPlatform->IniPlatformName(), UE::AssetRegistry::ESerializationTarget::ForDevelopment);
+		AssetRegistry.InitializeSerializationOptions(DevelopmentSaveOptions, TargetPlatform, UE::AssetRegistry::ESerializationTarget::ForDevelopment);
 		DevelopmentSaveOptions.bKeepDevelopmentAssetRegistryTags = true;
 
 		// Create a new FAssetRegistryState and call InitializeFromExistingAndPrune, then move the result to State.
@@ -1893,12 +1893,12 @@ bool FAssetRegistryGenerator::SaveAssetRegistry(const FString& SandboxPath, bool
 	
 	// Write development first, this will always write
 	FAssetRegistrySerializationOptions DevelopmentSaveOptions;
-	AssetRegistry.InitializeSerializationOptions(DevelopmentSaveOptions, TargetPlatform->IniPlatformName(), UE::AssetRegistry::ESerializationTarget::ForDevelopment);
+	AssetRegistry.InitializeSerializationOptions(DevelopmentSaveOptions, TargetPlatform, UE::AssetRegistry::ESerializationTarget::ForDevelopment);
 	DevelopmentSaveOptions.bKeepDevelopmentAssetRegistryTags = true;
 
 	// Write runtime registry, this can be excluded per game/platform
 	FAssetRegistrySerializationOptions SaveOptions;
-	AssetRegistry.InitializeSerializationOptions(SaveOptions, TargetPlatform->IniPlatformName());
+	AssetRegistry.InitializeSerializationOptions(SaveOptions, TargetPlatform, UE::AssetRegistry::ESerializationTarget::ForGame);
 	SaveOptions.bKeepDevelopmentAssetRegistryTags = FParse::Param(FCommandLine::Get(), TEXT("ARKeepDevTags"));
 
 	if (bForceNoFilter)
diff --git a/Engine/Source/Runtime/AssetRegistry/AssetRegistry.Build.cs b/Engine/Source/Runtime/AssetRegistry/AssetRegistry.Build.cs
--- a/Engine/Source/Runtime/AssetRegistry/AssetRegistry.Build.cs
+++ b/Engine/Source/Runtime/AssetRegistry/AssetRegistry.Build.cs
@@ -19,8 +19,8 @@ public class AssetRegistry : ModuleRules
 
 		if (Target.bBuildEditor == true)
 		{
-			PrivateIncludePathModuleNames.AddRange(new string[] { "DirectoryWatcher" });
-			DynamicallyLoadedModuleNames.AddRange(new string[] { "DirectoryWatcher" });
+			PrivateIncludePathModuleNames.AddRange(new string[] { "DirectoryWatcher", "TargetPlatform" });
+			DynamicallyLoadedModuleNames.AddRange(new string[] { "DirectoryWatcher", "TargetPlatform" });
 		}
 
 		CppCompileWarningSettings.UnsafeTypeCastWarningLevel = WarningLevel.Error;
diff --git a/Engine/Source/Runtime/AssetRegistry/Private/AssetRegistry.cpp b/Engine/Source/Runtime/AssetRegistry/Private/AssetRegistry.cpp
--- a/Engine/Source/Runtime/AssetRegistry/Private/AssetRegistry.cpp
+++ b/Engine/Source/Runtime/AssetRegistry/Private/AssetRegistry.cpp
@@ -17,6 +17,7 @@
 #include "DependsNode.h"
 #include "GenericPlatform/GenericPlatformChunkInstall.h"
 #include "GenericPlatform/GenericPlatformFile.h"
+#include "GenericPlatform/GenericPlatformMisc.h"
 #include "HAL/PlatformMisc.h"
 #include "HAL/ThreadHeartBeat.h"
 #include "Interfaces/IPluginManager.h"
@@ -60,6 +61,8 @@
 #include "HAL/IConsoleManager.h"
 #include "HAL/PlatformProcess.h"
 #include "IDirectoryWatcher.h"
+#include "Interfaces/ITargetPlatform.h"
+#include "Interfaces/ITargetPlatformManagerModule.h"
 #endif // WITH_EDITOR
 
 #define UE_ENABLE_DIRECTORYWATCH_ROOTS !UE_IS_COOKED_EDITOR
@@ -771,6 +774,28 @@ FAssetData IAssetRegistry::K2_GetAssetByObjectPath(const FSoftObjectPath& Object
 	return GetAssetByObjectPath(ObjectPath, bIncludeOnlyOnDiskAssets, bSkipARFilteredAssets);
 }
 
+void IAssetRegistry::InitializeSerializationOptions(FAssetRegistrySerializationOptions& Options,
+	const FString& PlatformIniName,
+	UE::AssetRegistry::ESerializationTarget Target) const
+{
+#if WITH_EDITOR
+	if (ITargetPlatformManagerModule* TargetPlatformManager = GetTargetPlatformManager())
+	{
+		for (ITargetPlatform* TargetPlatform : TargetPlatformManager->GetActiveTargetPlatforms())
+		{
+			if (TargetPlatform->IniPlatformName() == PlatformIniName)
+			{
+				return InitializeSerializationOptions(Options, TargetPlatform, Target);
+			}
+		}
+	}
+#endif
+
+	if (ensureMsgf(PlatformIniName.IsEmpty(), TEXT("Couldn't find TargetPlatform from ini `%s`"), *PlatformIniName)) {
+		return InitializeSerializationOptions(Options, nullptr, Target);
+	}
+}
+
 IAssetRegistry::FLoadPackageRegistryData::FLoadPackageRegistryData(bool bInGetDependencies)
 	: bGetDependencies(bInGetDependencies)
 {
@@ -994,8 +1019,8 @@ void FAssetRegistryImpl::Initialize(Impl::FInitializeContext& Context)
 #endif
 
 	// Read default serialization options
-	Utils::InitializeSerializationOptionsFromIni(SerializationOptions, FString());
-	Utils::InitializeSerializationOptionsFromIni(DevelopmentSerializationOptions, FString(), UE::AssetRegistry::ESerializationTarget::ForDevelopment);
+	Utils::InitializeSerializationOptionsFromIni(SerializationOptions, nullptr, UE::AssetRegistry::ESerializationTarget::ForGame);
+	Utils::InitializeSerializationOptionsFromIni(DevelopmentSerializationOptions, nullptr, UE::AssetRegistry::ESerializationTarget::ForDevelopment);
 
 	bool bStartedAsyncGather = false;
 	if (ShouldSearchAllAssetsAtStart())
@@ -1377,9 +1402,9 @@ void FAssetRegistryImpl::ReadScriptPackages()
 
 }
 
-void UAssetRegistryImpl::InitializeSerializationOptions(FAssetRegistrySerializationOptions& Options, const FString& PlatformIniName, UE::AssetRegistry::ESerializationTarget Target) const
+void UAssetRegistryImpl::InitializeSerializationOptions(FAssetRegistrySerializationOptions& Options, const ITargetPlatform* TargetPlatform, UE::AssetRegistry::ESerializationTarget Target) const
 {
-	if (PlatformIniName.IsEmpty())
+	if (!TargetPlatform)
 	{
 		UE::AssetRegistry::FInterfaceReadScopeLock InterfaceScopeLock(InterfaceLock);
 		// Use options we already loaded, the first pass for this happens at object creation time so this is always valid when queried externally
@@ -1387,7 +1412,7 @@ void UAssetRegistryImpl::InitializeSerializationOptions(FAssetRegistrySerializat
 	}
 	else
 	{
-		UE::AssetRegistry::Utils::InitializeSerializationOptionsFromIni(Options, PlatformIniName, Target);
+		UE::AssetRegistry::Utils::InitializeSerializationOptionsFromIni(Options, TargetPlatform, Target);
 	}
 }
 
@@ -1421,11 +1446,30 @@ static TSet<FName> MakeNameSet(const TArray<FString>& Strings)
 	return Out;
 }
 
-void InitializeSerializationOptionsFromIni(FAssetRegistrySerializationOptions& Options, const FString& PlatformIniName, UE::AssetRegistry::ESerializationTarget Target)
+void InitializeSerializationOptionsFromIni(
+	FAssetRegistrySerializationOptions& Options,
+	const ITargetPlatform* TargetPlatform,
+	UE::AssetRegistry::ESerializationTarget Target)
 {
+	FString PlatformIniName;
+	EBuildTargetType TargetType = EBuildTargetType::Game;
+	if (TargetPlatform)
+	{
+#if WITH_EDITOR
+		PlatformIniName = TargetPlatform->IniPlatformName();
+		TargetType = TargetPlatform->GetRuntimePlatformType();
+#endif
+	}
+	else
+	{
+		PlatformIniName = ANSI_TO_TCHAR(FPlatformProperties::IniPlatformName());
+		TargetType = FPlatformProperties::IsClientOnly() ? EBuildTargetType::Client
+			: FPlatformProperties::IsServerOnly() ? EBuildTargetType::Server
+			: EBuildTargetType::Game;
+	}
 	// Use passed in platform, or current platform if empty
 	FConfigFile LocalEngineIni;
-	FConfigFile* EngineIni = FConfigCacheIni::FindOrLoadPlatformConfig(LocalEngineIni, TEXT("Engine"), (!PlatformIniName.IsEmpty() ? *PlatformIniName : ANSI_TO_TCHAR(FPlatformProperties::IniPlatformName())));
+	FConfigFile* EngineIni = FConfigCacheIni::FindOrLoadPlatformConfig(LocalEngineIni, TEXT("Engine"), *PlatformIniName);
 
 	Options = FAssetRegistrySerializationOptions(Target);
 	// For DevelopmentAssetRegistry, all non-tag options are overridden in the constructor
@@ -1440,6 +1484,24 @@ void InitializeSerializationOptionsFromIni(FAssetRegistrySerializationOptions& O
 		EngineIni->GetBool(TEXT("AssetRegistry"), TEXT("bFilterAssetDataWithNoTags"), Options.bFilterAssetDataWithNoTags);
 		EngineIni->GetBool(TEXT("AssetRegistry"), TEXT("bFilterDependenciesWithNoTags"), Options.bFilterDependenciesWithNoTags);
 		EngineIni->GetBool(TEXT("AssetRegistry"), TEXT("bFilterSearchableNames"), Options.bFilterSearchableNames);
+
+		{
+			TArray<FString> AssetBundlesDenyList;
+			EngineIni->GetArray(TEXT("AssetRegistry"), TEXT("AssetBundlesDenyList"), AssetBundlesDenyList);
+			Options.AssetBundlesDenyList.Append(AssetBundlesDenyList);
+		}
+		if (TargetType == EBuildTargetType::Server)
+		{
+			TArray<FString> AssetBundlesDenyList;
+			EngineIni->GetArray(TEXT("AssetRegistry"), TEXT("ServerAssetBundlesDenyList"), AssetBundlesDenyList);
+			Options.AssetBundlesDenyList.Append(AssetBundlesDenyList);
+		}
+		if (TargetType == EBuildTargetType::Client)
+		{
+			TArray<FString> AssetBundlesDenyList;
+			EngineIni->GetArray(TEXT("AssetRegistry"), TEXT("ClientAssetBundlesDenyList"), AssetBundlesDenyList);
+			Options.AssetBundlesDenyList.Append(AssetBundlesDenyList);
+		}
 	}
 
 	EngineIni->GetBool(TEXT("AssetRegistry"), TEXT("bUseAssetRegistryTagsWhitelistInsteadOfBlacklist"), Options.bUseAssetRegistryTagsAllowListInsteadOfDenyList);
@@ -1660,8 +1722,8 @@ void FAssetRegistryImpl::RefreshNativeClasses()
 	TempCachedInheritanceBuffer.bDirty = true;
 
 	// Read default serialization options
-	Utils::InitializeSerializationOptionsFromIni(SerializationOptions, FString());
-	Utils::InitializeSerializationOptionsFromIni(DevelopmentSerializationOptions, FString(), UE::AssetRegistry::ESerializationTarget::ForDevelopment);
+	Utils::InitializeSerializationOptionsFromIni(SerializationOptions, nullptr, UE::AssetRegistry::ESerializationTarget::ForGame);
+	Utils::InitializeSerializationOptionsFromIni(DevelopmentSerializationOptions, nullptr, UE::AssetRegistry::ESerializationTarget::ForDevelopment);
 }
 
 }
@@ -5917,6 +5979,11 @@ void FAssetRegistryImpl::GetAllocatedSize(bool bLogDetailed, SIZE_T& StateSize,
 	{
 		StaticSize += Pair.Value.GetAllocatedSize();
 	}
+	StaticSize += SerializationOptions.AssetBundlesDenyList.GetAllocatedSize();
+	for (const FString& Filter : SerializationOptions.AssetBundlesDenyList)
+	{
+		StaticSize += Filter.GetAllocatedSize();
+	}
 }
 
 }
diff --git a/Engine/Source/Runtime/AssetRegistry/Private/AssetRegistry.h b/Engine/Source/Runtime/AssetRegistry/Private/AssetRegistry.h
--- a/Engine/Source/Runtime/AssetRegistry/Private/AssetRegistry.h
+++ b/Engine/Source/Runtime/AssetRegistry/Private/AssetRegistry.h
@@ -125,7 +125,7 @@ public:
 #endif
 	virtual TSet<FName> GetCachedEmptyPackagesCopy() const override;
 	virtual bool ContainsTag(FName TagName) const override;
-	virtual void InitializeSerializationOptions(FAssetRegistrySerializationOptions& Options, const FString& PlatformIniName = FString(), UE::AssetRegistry::ESerializationTarget Target = UE::AssetRegistry::ESerializationTarget::ForGame) const override;
+	virtual void InitializeSerializationOptions(FAssetRegistrySerializationOptions& Options, const ITargetPlatform* TargetPlatform = nullptr, UE::AssetRegistry::ESerializationTarget Target = UE::AssetRegistry::ESerializationTarget::ForGame) const override;
 
 	DECLARE_DERIVED_EVENT( UAssetRegistryImpl, IAssetRegistry::FFilesBlockedEvent, FFilesBlockedEvent);
 	virtual FFilesBlockedEvent& OnFilesBlocked() override;
diff --git a/Engine/Source/Runtime/AssetRegistry/Private/AssetRegistryImpl.h b/Engine/Source/Runtime/AssetRegistry/Private/AssetRegistryImpl.h
--- a/Engine/Source/Runtime/AssetRegistry/Private/AssetRegistryImpl.h
+++ b/Engine/Source/Runtime/AssetRegistry/Private/AssetRegistryImpl.h
@@ -31,6 +31,7 @@
 #define ASSETREGISTRY_FORCE_PREMADE_REGISTRY_IN_EDITOR 0
 #endif
 
+class ITargetPlatform;
 class FDependsNode;
 struct FARFilter;
 class FAssetDataGatherer;
@@ -1019,7 +1020,9 @@ bool RunAssetThroughFilter_Unchecked(const FAssetData& AssetData, const FARCompi
 void RunAssetsThroughFilter(TArray<FAssetData>& AssetDataList, const FARCompiledFilter& Filter, const EFilterMode FilterMode);
 
 /** This will always read the ini, public version may return cache */
-void InitializeSerializationOptionsFromIni(FAssetRegistrySerializationOptions& Options, const FString& PlatformIniName,
+void InitializeSerializationOptionsFromIni(
+	FAssetRegistrySerializationOptions& Options,
+	const ITargetPlatform* TargetPlatform = nullptr,
 	ESerializationTarget Target = ESerializationTarget::ForGame);
 
 /** Gets the current availability of an asset, primarily for streaming install purposes. */
diff --git a/Engine/Source/Runtime/AssetRegistry/Private/AssetRegistryState.cpp b/Engine/Source/Runtime/AssetRegistry/Private/AssetRegistryState.cpp
--- a/Engine/Source/Runtime/AssetRegistry/Private/AssetRegistryState.cpp
+++ b/Engine/Source/Runtime/AssetRegistry/Private/AssetRegistryState.cpp
@@ -6,6 +6,7 @@
 #include "Algo/Sort.h"
 #include "Algo/Unique.h"
 #include "AssetRegistry/ARFilter.h"
+#include "AssetRegistry/AssetBundleData.h"
 #include "AssetRegistryArchive.h"
 #include "AssetRegistryImpl.h"
 #include "AssetRegistryPrivate.h"
@@ -18,11 +19,13 @@
 #include "Misc/Paths.h"
 #include "Misc/PathViews.h"
 #include "Misc/StringBuilder.h"
+#include "Misc/WildcardString.h"
 #include "NameTableArchive.h"
 #include "AssetRegistry/PackageReader.h"
 #include "Serialization/ArrayReader.h"
 #include "Serialization/LargeMemoryReader.h"
 #include "Serialization/MemoryWriter.h"
+#include "Templates/SharedPointer.h"
 #include "UObject/AssetRegistryTagsContext.h"
 #include "UObject/MetaData.h"
 #include "UObject/NameBatchSerialization.h"
@@ -212,6 +215,43 @@ void FAssetRegistryState::FilterTags(const FAssetDataTagMapSharedView& InTagsAnd
 	}
 }
 
+TSharedPtr<FAssetBundleData> FAssetRegistryState::FilterBundles(const TSharedPtr<FAssetBundleData>& InAssetBundles, const FAssetRegistrySerializationOptions& Options)
+{
+	if (InAssetBundles)
+	{
+		TArray<FAssetBundleEntry> Bundles;
+		Bundles.Reserve(InAssetBundles->Bundles.Num());
+		bool bIsDifferent = false;
+		for (auto It = InAssetBundles->Bundles.CreateConstIterator(); It; ++It)
+		{
+			TStringBuilder<NAME_SIZE> Name;
+			It->BundleName.AppendString(Name);
+			bool bPassesFilters = true;
+			for (const FString& Filter : Options.AssetBundlesDenyList)
+			{
+				if (FWildcardString::IsMatchSubstring(*Filter, Name.GetData(), Name.GetData() + Name.Len(), ESearchCase::IgnoreCase))
+				{
+					bPassesFilters = false;
+					bIsDifferent = true;
+					break;
+				}
+			}
+			if (bPassesFilters)
+			{
+				Bundles.Add(*It);
+			}
+		}
+		if (bIsDifferent)
+		{
+			Bundles.Shrink();
+			TSharedPtr<FAssetBundleData> RetData = MakeShared<FAssetBundleData>();
+			RetData->Bundles = MoveTemp(Bundles);
+			return RetData;
+		}
+	}
+	return InAssetBundles;
+}
+
 void FAssetRegistryState::InitializeFromExistingAndPrune(const FAssetRegistryState & ExistingState,
 	const TSet<FName>& RequiredPackages, const TSet<FName>& RemovePackages,
 	const TSet<int32> ChunksToKeep, const FAssetRegistrySerializationOptions& Options)
@@ -280,7 +320,7 @@ void FAssetRegistryState::InitializeFromExistingAndPrune(const FAssetRegistrySta
 				AssetData.AssetClassPath, NewTagsAndValues, AssetData.GetChunkIDs(), AssetData.PackageFlags);
 		}
 
-		NewAssetData->TaggedAssetBundles = AssetData.TaggedAssetBundles;
+		NewAssetData->TaggedAssetBundles = FAssetRegistryState::FilterBundles(AssetData.TaggedAssetBundles, Options);
 
 		// Add asset to new state
 		AddAssetData(NewAssetData);
@@ -413,11 +453,15 @@ void FAssetRegistryState::InitializeFromExisting(const FAssetDataMap& AssetDataM
 		FAssetDataTagMap LocalTagsAndValues;
 		FAssetRegistryState::FilterTags(AssetData.TagsAndValues, LocalTagsAndValues,
 			Options.CookFilterlistTagsByClass.Find(AssetData.AssetClassPath), Options);
-		
+
+		// Filter AssetBundles separately, LocalAssetBundles will point to the same TaggedAssetBundles if it is not different
+		TSharedPtr<FAssetBundleData> LocalAssetBundles = FAssetRegistryState::FilterBundles(AssetData.TaggedAssetBundles, Options);
+
 		if (ExistingData)
 		{
 			FAssetData NewData(AssetData);
 			NewData.TagsAndValues = FAssetDataTagMapSharedView(MoveTemp(LocalTagsAndValues));
+			NewData.TaggedAssetBundles = MoveTemp(LocalAssetBundles);
 			bool bIsModified = false;
 			UpdateAssetData(ExistingData, MoveTemp(NewData), &bIsModified);
 			if (OutAppendResult && bIsModified)
@@ -429,6 +473,7 @@ void FAssetRegistryState::InitializeFromExisting(const FAssetDataMap& AssetDataM
 		{
 			FAssetData* NewData = new FAssetData(AssetData);
 			NewData->TagsAndValues = FAssetDataTagMapSharedView(MoveTemp(LocalTagsAndValues));
+			NewData->TaggedAssetBundles = MoveTemp(LocalAssetBundles);
 			AddAssetData(NewData);
 			if (OutAppendResult)
 			{
@@ -2652,6 +2697,8 @@ void FAssetRegistryState::FilterTags(const FAssetRegistrySerializationOptions& O
 			AssetData.TagsAndValues = FAssetDataTagMapSharedView(MoveTemp(LocalTagsAndValues));
 		}
 
+		AssetData.TaggedAssetBundles = FAssetRegistryState::FilterBundles(AssetData.TaggedAssetBundles, Options);
+
 		// Add the AssetData to all its CachedAssetsByTag/CachedClassesByTag keys even if nothing changed, because
 		// we are reconstructing all CachedAssetsByTag.
 		for (auto TagIt = AssetData.TagsAndValues.CreateConstIterator(); TagIt; ++TagIt)
diff --git a/Engine/Source/Runtime/AssetRegistry/Public/AssetRegistry/AssetRegistryState.h b/Engine/Source/Runtime/AssetRegistry/Public/AssetRegistry/AssetRegistryState.h
--- a/Engine/Source/Runtime/AssetRegistry/Public/AssetRegistry/AssetRegistryState.h
+++ b/Engine/Source/Runtime/AssetRegistry/Public/AssetRegistry/AssetRegistryState.h
@@ -110,6 +110,9 @@ struct FAssetRegistrySerializationOptions
 	/** Tag keys whose values should be stored as FRegistryExportPath in cooked builds. */
 	TSet<FName> CookTagsAsPath;
 
+	/** List of asset bundles (could contain wildcards) that will be removed from the AssetRegistry */
+	TSet<FString> AssetBundlesDenyList;
+
 	/** Disable all filters */
 	void DisableFilters();
 
@@ -673,6 +676,10 @@ private:
 	static void FilterTags(const FAssetDataTagMapSharedView& InTagsAndValues, FAssetDataTagMap& OutTagsAndValues,
 		const TSet<FName>* ClassSpecificFilterList, const FAssetRegistrySerializationOptions & Options);
 
+	/** Filter asset bundles and output a copy of the filtered set or input AssetBundles if nothing was filtered */
+	static TSharedPtr<FAssetBundleData> FilterBundles(const TSharedPtr<FAssetBundleData>& InAssetBundles,
+		const FAssetRegistrySerializationOptions& Options);
+
 	void LoadDependencies(FArchive& Ar);
 	void LoadDependencies_BeforeFlags(FArchive& Ar, bool bSerializeDependencies, FAssetRegistryVersion::Type Version);
 
diff --git a/Engine/Source/Runtime/AssetRegistry/Public/AssetRegistry/IAssetRegistry.h b/Engine/Source/Runtime/AssetRegistry/Public/AssetRegistry/IAssetRegistry.h
--- a/Engine/Source/Runtime/AssetRegistry/Public/AssetRegistry/IAssetRegistry.h
+++ b/Engine/Source/Runtime/AssetRegistry/Public/AssetRegistry/IAssetRegistry.h
@@ -32,6 +32,7 @@ struct FAssetRegistrySerializationOptions;
 struct FPackageFileSummary;
 struct FObjectExport;
 struct FObjectImport;
+class ITargetPlatform;
 
 namespace EAssetAvailability
 {
@@ -1073,8 +1074,14 @@ public:
 	virtual bool ContainsTag(FName TagName) const = 0;
 
 	/** Fills in FAssetRegistrySerializationOptions from ini, optionally using a target platform ini name */
+	UE_DEPRECATED(5.7, "Use version taking ITargetPlatform parameter instead")
+	ASSETREGISTRY_API virtual void InitializeSerializationOptions(FAssetRegistrySerializationOptions& Options,
+		const FString& PlatformIniName,
+		UE::AssetRegistry::ESerializationTarget Target) const;
+
+	/** Fills in FAssetRegistrySerializationOptions from ini, optionally using a target platform */
 	virtual void InitializeSerializationOptions(FAssetRegistrySerializationOptions& Options,
-		const FString& PlatformIniName = FString(),
+		const ITargetPlatform* TargetPlatform = nullptr,
 		UE::AssetRegistry::ESerializationTarget Target = UE::AssetRegistry::ESerializationTarget::ForGame) const = 0;
 
 	struct FLoadPackageRegistryData
