Parallelize ShouldSkipAsset checks during FAssetRegistryState::EnumerateAssets.

#rnx

diff --git a/Engine/Source/Runtime/AssetRegistry/Private/AssetRegistryState.cpp b/Engine/Source/Runtime/AssetRegistry/Private/AssetRegistryState.cpp
--- a/Engine/Source/Runtime/AssetRegistry/Private/AssetRegistryState.cpp
+++ b/Engine/Source/Runtime/AssetRegistry/Private/AssetRegistryState.cpp
@@ -10,6 +10,7 @@
 #include "AssetRegistryArchive.h"
 #include "AssetRegistryImpl.h"
 #include "AssetRegistryPrivate.h"
+#include "Async/ParallelFor.h"
 #include "Blueprint/BlueprintSupport.h"
 #include "DependsNode.h"
 #include "GenericPlatform/GenericPlatformChunkInstall.h"
@@ -1443,16 +1444,64 @@ bool FAssetRegistryState::EnumerateAssets(const FARCompiledFilter& Filter, const
 	}
 
 	// Run the remaining non-accelerated filters on every element of AccumulatedResults
-	for (const FAssetData* AssetData : AccumulatedResults)
+
+#if WITH_ENGINE && WITH_EDITOR
+	// Parallelize the non-accelerated filter. This optimizes large queries coming from e.g. CookByTheBookFinished in the
+	// functions FAssetRegistryGenerator::PreSave and GenerateLocalizationReferences through FAssetRegistryImpl::EnumerateDiskAssets.
+	
+	// Each thread will work on a full cache line to avoid false sharing.
+	constexpr int32 ChecksPerCacheLine = PLATFORM_CACHE_LINE_SIZE / sizeof(bool);
+	if (IsInGameThread() && AccumulatedResults.Num() >= ChecksPerCacheLine)
 	{
-		if (ShouldSkipAssetData(AssetData))
+		// Ensure UE::AssetRegistry::FFiltering::ShouldSkipAsset is thread safe.
+		UE::AssetRegistry::FFiltering::InitializeShouldSkipAsset();
+
+		TArray<bool, TAlignedHeapAllocator<PLATFORM_CACHE_LINE_SIZE>> ShouldSkip;
+		ShouldSkip.SetNum(AccumulatedResults.Num());
+		
+		const int32 CacheLineBatches = FMath::DivideAndRoundUp(AccumulatedResults.Num(), ChecksPerCacheLine);
+
+		ParallelFor(CacheLineBatches, [&ShouldSkipAssetData, &AccumulatedResults, &ShouldSkip](int32 CacheLineIndex)
 		{
-			continue;
+			const int32 CacheLineOffset = CacheLineIndex * ChecksPerCacheLine;
+			for (int32 CheckIndex = 0; CheckIndex < ChecksPerCacheLine; ++CheckIndex)
+			{
+				const int32 AccumulatedResultsIndex = CacheLineOffset + CheckIndex;
+				if (AccumulatedResultsIndex < AccumulatedResults.Num())
+				{
+					ShouldSkip[AccumulatedResultsIndex] = ShouldSkipAssetData(AccumulatedResults[AccumulatedResultsIndex]);
+				}
+			}
+		});
+
+		for (int32 Index = 0; Index < AccumulatedResults.Num(); ++Index)
+		{
+			if (ShouldSkip[Index])
+			{
+				continue;
+			}
+			const FAssetData* AssetData = AccumulatedResults[Index];
+			bool bContinueFiltering = Callback(*AssetData);
+			if (!bContinueFiltering)
+			{
+				return true;
+			}
 		}
-		bool bContinueFiltering = Callback(*AssetData);
-		if (!bContinueFiltering)
+	}
+	else
+#endif 
+	{
+		for (const FAssetData* AssetData : AccumulatedResults)
 		{
-			return true;
+			if (ShouldSkipAssetData(AssetData))
+			{
+				continue;
+			}
+			bool bContinueFiltering = Callback(*AssetData);
+			if (!bContinueFiltering)
+			{
+				return true;
+			}
 		}
 	}
 
