[AssetRegistry]

- No longer force cook workers to read the discovery cache as there isn't a guarantee that the director wrote out a new version.
- Added HasSerializedDiscoveryCache to IAssetRegistry interface to allow querying if the discovery cache has been written.
- The cook director will now allow cook workers to skip registry invalidation if the director wrote out a discovery cache.
- Remove hack disabling the discovery cache on build machines

#rnx

diff --git a/Engine/Source/Editor/UnrealEd/Private/Cooker/CookDirector.cpp b/Engine/Source/Editor/UnrealEd/Private/Cooker/CookDirector.cpp
--- a/Engine/Source/Editor/UnrealEd/Private/Cooker/CookDirector.cpp
+++ b/Engine/Source/Editor/UnrealEd/Private/Cooker/CookDirector.cpp
@@ -1581,6 +1581,7 @@ FString FCookDirector::GetWorkerCommandLine(FWorkerId WorkerId, int32 ProfileId)
 	// does handle them
 	FString Token;
 	TArray<FString> Tokens;
+	bool bAssetRegistryCacheWritingDisabled = false;
 	while (FParse::Token(CommandLine, Token, false /* bUseEscape */))
 	{
 		if (Token.IsEmpty())
@@ -1625,6 +1626,7 @@ FString FCookDirector::GetWorkerCommandLine(FWorkerId WorkerId, int32 ProfileId)
 			// If cache writing has been disabled, we must also disable reading of the discovery cache on workers.
 			// Otherwise workers could be looking at old discovery results that no longer match what's on disk.
 			Tokens.Add(TEXT("-NoAssetRegistryDiscoveryCache"));
+			bAssetRegistryCacheWritingDisabled = true;
 		}
 		Tokens.Add(MoveTemp(Token));
 	}
@@ -1647,6 +1649,18 @@ FString FCookDirector::GetWorkerCommandLine(FWorkerId WorkerId, int32 ProfileId)
 	Tokens.Add(FString::Printf(TEXT("-CookDirectorHost=%s"), *WorkerConnectAuthority));
 	Tokens.Add(TEXT("-unattended"));
 	Tokens.Add(FString::Printf(TEXT("-abslog=%s"), *GetWorkerLogFileName(ProfileId)));
+
+	IAssetRegistry& AssetRegistry = IAssetRegistry::GetChecked();
+	if (AssetRegistry.HasSerializedDiscoveryCache())
+	{
+		UE_CLOG(bAssetRegistryCacheWritingDisabled, LogCook, Warning, TEXT("Invalid configuration: The Cook Director has written an AssetRegistry discovery cache when explicitly told not to via -NoAssetRegistryCacheWrite"));
+		if (!bAssetRegistryCacheWritingDisabled)
+		{
+			// If the director was successful in writing a discovery cache, let workers re-use the director's cache
+			Tokens.Add(TEXT("-AssetRegistryCacheSkipInvalidate"));
+		}
+	}
+
 	if (CoreLimit > 0)
 	{
 		Tokens.Add(FString::Printf(TEXT("-PhysicalCoreLimit=%d"), CoreLimit));
diff --git a/Engine/Source/Runtime/AssetRegistry/Private/AssetDataGatherer.cpp b/Engine/Source/Runtime/AssetRegistry/Private/AssetDataGatherer.cpp
--- a/Engine/Source/Runtime/AssetRegistry/Private/AssetDataGatherer.cpp
+++ b/Engine/Source/Runtime/AssetRegistry/Private/AssetDataGatherer.cpp
@@ -190,12 +190,6 @@ void FGatherSettings::Initialize()
 #endif
 	bool bNoAssetRegistryCache = FParse::Param(FCommandLine::Get(), TEXT("NoAssetRegistryCache"));
 	bool bNoAssetRegistryDiscoveryCache = bNoAssetRegistryCache || FParse::Param(FCommandLine::Get(), TEXT("NoAssetRegistryDiscoveryCache"));
-	if (GIsBuildMachine)
-	{
-		// Force discovery cache usage off until UE-274858 is investigated
-		bNoAssetRegistryDiscoveryCache = true;
-	}
-	
 	bool bNoAssetRegistryCacheRead = FParse::Param(FCommandLine::Get(), TEXT("NoAssetRegistryCacheRead"));
 	uint32 MultiprocessId = UE::GetMultiprocessId();
 	bool bMultiprocess = MultiprocessId > 0 || FParse::Param(FCommandLine::Get(), TEXT("multiprocess"));
@@ -243,7 +237,7 @@ void FGatherSettings::Initialize()
 	{
 		DiscoverySetting = (DiscoverySetting & ~EFeatureEnabledReadWrite::WriteMask) | EFeatureEnabledReadWrite::NeverWrite;
 	}
-	if (bSkipInvalidate || MultiprocessId != 0)
+	if (bSkipInvalidate)
 	{
 		if ((DiscoverySetting & EFeatureEnabledReadWrite::WriteMask) == EFeatureEnabledReadWrite::DefaultWrite)
 		{
@@ -2855,6 +2849,12 @@ bool FAssetDataDiscovery::IsMonitored(FStringView LocalAbsPath) const
 	return MountDir && MountDir->IsMonitored(LocalAbsPath);
 }
 
+bool FAssetDataDiscovery::IsCacheWriteEnabled() const
+{
+	// Upon initialization, WriteEnabled will be forced to Never if not supported by the platform
+	return Cache.IsInitialized() && Cache.IsWriteEnabled() != EFeatureEnabled::Never;
+}
+
 template <typename ContainerType>
 SIZE_T GetArrayRecursiveAllocatedSize(const ContainerType& Container)
 {
@@ -5000,6 +5000,13 @@ bool FAssetDataGatherer::IsComplete() const
 	return bIsComplete;
 }
 
+bool FAssetDataGatherer::HasSerializedDiscoveryCache() const
+{
+	using namespace UE::AssetDataGather::Private;
+	FGathererScopeLock ResultsScopeLock(&ResultsLock);
+	return bDiscoveryIsComplete && Discovery->IsCacheWriteEnabled();
+}
+
 void FAssetDataGatherer::SetInitialPluginsLoaded()
 {
 	FGathererScopeLock ResultsScopeLock(&ResultsLock);
diff --git a/Engine/Source/Runtime/AssetRegistry/Private/AssetDataGatherer.h b/Engine/Source/Runtime/AssetRegistry/Private/AssetDataGatherer.h
--- a/Engine/Source/Runtime/AssetRegistry/Private/AssetDataGatherer.h
+++ b/Engine/Source/Runtime/AssetRegistry/Private/AssetDataGatherer.h
@@ -190,7 +190,10 @@ public:
 	 * This function can be used to check whether there is any work to be done on the gatherer.
 	 */
 	bool IsComplete() const;
-
+	/** Returns true if the asset registry has written out the discovery cache after discovery has completed. 
+	* This function will return false if discovery is not yet complete, discovery caching isn't supported by the current platform 
+	* or cache writing has been disabled such as when using -NoAssetRegistryCacheWrite. */
+	bool HasSerializedDiscoveryCache() const;
 
 	// Reading/writing/triggering depot-wide properties and events (possibly while tick is running)
 
diff --git a/Engine/Source/Runtime/AssetRegistry/Private/AssetDataGathererDiscoveryCache.h b/Engine/Source/Runtime/AssetRegistry/Private/AssetDataGathererDiscoveryCache.h
--- a/Engine/Source/Runtime/AssetRegistry/Private/AssetDataGathererDiscoveryCache.h
+++ b/Engine/Source/Runtime/AssetRegistry/Private/AssetDataGathererDiscoveryCache.h
@@ -132,8 +132,8 @@ class FAssetDataDiscoveryCache
 public:
 	FString GetCacheFileName() const;
 
-	bool IsInitialized() { return bInitialized; }
-	EFeatureEnabled IsWriteEnabled() { return WriteEnabled; }
+	bool IsInitialized() const { return bInitialized; }
+	EFeatureEnabled IsWriteEnabled() const { return WriteEnabled; }
 	void Shutdown();
 
 	void LoadAndUpdateCache();
diff --git a/Engine/Source/Runtime/AssetRegistry/Private/AssetDataGathererPrivate.h b/Engine/Source/Runtime/AssetRegistry/Private/AssetDataGathererPrivate.h
--- a/Engine/Source/Runtime/AssetRegistry/Private/AssetDataGathererPrivate.h
+++ b/Engine/Source/Runtime/AssetRegistry/Private/AssetDataGathererPrivate.h
@@ -619,6 +619,8 @@ public:
 	bool IsOnDenyList(FStringView LocalAbsPath) const;
 	/** Return whether the path should or has been scanned because it is on the AllowList and not on the DenyList. */
 	bool IsMonitored(FStringView LocalAbsPath) const;
+	/** Return whether the discovery cache can be written. */
+	bool IsCacheWriteEnabled() const;
 	/** Return the memory used by *this. sizeof(*this) is not included. */
 	SIZE_T GetAllocatedSize() const;
 
diff --git a/Engine/Source/Runtime/AssetRegistry/Private/AssetRegistry.cpp b/Engine/Source/Runtime/AssetRegistry/Private/AssetRegistry.cpp
--- a/Engine/Source/Runtime/AssetRegistry/Private/AssetRegistry.cpp
+++ b/Engine/Source/Runtime/AssetRegistry/Private/AssetRegistry.cpp
@@ -4632,6 +4632,11 @@ bool UAssetRegistryImpl::IsGathering() const
 	return GuardedData.IsGathering();
 }
 
+bool UAssetRegistryImpl::HasSerializedDiscoveryCache() const
+{
+	return GuardedData.HasSerializedDiscoveryCache();
+}
+
 namespace UE::AssetRegistry
 {
 
@@ -4645,6 +4650,10 @@ bool FAssetRegistryImpl::IsGathering() const
 	return !IsInitialSearchCompleted() || IsAdditionalMountSearchInProgress();
 }
 
+bool FAssetRegistryImpl::HasSerializedDiscoveryCache() const
+{
+	return GlobalGatherer->HasSerializedDiscoveryCache();
+}
 }
 
 UE::AssetRegistry::Impl::EGatherStatus UAssetRegistryImpl::TickOnBackgroundThread()
diff --git a/Engine/Source/Runtime/AssetRegistry/Private/AssetRegistry.h b/Engine/Source/Runtime/AssetRegistry/Private/AssetRegistry.h
--- a/Engine/Source/Runtime/AssetRegistry/Private/AssetRegistry.h
+++ b/Engine/Source/Runtime/AssetRegistry/Private/AssetRegistry.h
@@ -205,6 +205,7 @@ public:
 
 	virtual bool IsLoadingAssets() const override;
 	virtual bool IsGathering() const override;
+	virtual bool HasSerializedDiscoveryCache() const override;
 	virtual bool ShouldUpdateDiskCacheAfterLoad() const override
 	{
 #if WITH_EDITORONLY_DATA
diff --git a/Engine/Source/Runtime/AssetRegistry/Private/AssetRegistryImpl.h b/Engine/Source/Runtime/AssetRegistry/Private/AssetRegistryImpl.h
--- a/Engine/Source/Runtime/AssetRegistry/Private/AssetRegistryImpl.h
+++ b/Engine/Source/Runtime/AssetRegistry/Private/AssetRegistryImpl.h
@@ -389,6 +389,7 @@ public:
 	void GetAllocatedSize(bool bLogDetailed, SIZE_T& StateSize, SIZE_T& StaticSize, SIZE_T& SearchSize) const;
 	bool IsLoadingAssets() const;
 	bool IsGathering() const;
+	bool HasSerializedDiscoveryCache() const;
 	void SetManageReferences(FSetManageReferencesContext& Context);
 	bool SetPrimaryAssetIdForObjectPath(Impl::FEventContext& EventContext, const FSoftObjectPath& ObjectPath, FPrimaryAssetId PrimaryAssetId);
 	bool ResolveRedirect(const FString& InPackageName, FString& OutPackageName) const;
diff --git a/Engine/Source/Runtime/AssetRegistry/Public/AssetRegistry/IAssetRegistry.h b/Engine/Source/Runtime/AssetRegistry/Public/AssetRegistry/IAssetRegistry.h
--- a/Engine/Source/Runtime/AssetRegistry/Public/AssetRegistry/IAssetRegistry.h
+++ b/Engine/Source/Runtime/AssetRegistry/Public/AssetRegistry/IAssetRegistry.h
@@ -1033,6 +1033,11 @@ public:
 	*/
 	virtual bool IsGathering() const = 0;
 
+	/** Returns true if the asset registry has written out the discovery cache after discovery has completed. 
+	* This function will return false if discovery is not yet complete, discovery caching isn't supported by the current platform 
+	* or cache writing has been disabled such as when using -NoAssetRegistryCacheWrite. */
+	virtual bool HasSerializedDiscoveryCache() const = 0;
+
 	/** If true, the AssetRegistry updates its on-disk information for an Asset whenever that Asset loads. */
 	virtual bool ShouldUpdateDiskCacheAfterLoad() const = 0;
 
