SoftObjectPaths: Mark SoftObjectPaths loaded from PackageHeaders as needing redirects applied on load, which we do through PreSavePath, but SoftObjectPaths loaded outside of PackageHeaders, such as when populating the AssetRegistry, do not need redirects applied. Add a separate pass for the AssetRegistry to populate redirects once it finishes loading all assets. This fixes a problem with the AssetBundles used by PrimaryAssets in the AssetManager search indeterministically having some redirects applied but not others. Add new functions in SoftObjectPath.h to categorize the different types of ESoftObjectPathCollectType to reduce duplication in the callsites that need to handle categories of CollectType differently.
#rnx

diff --git a/Engine/Source/Developer/AssetTools/Private/AssetRenameManager.cpp b/Engine/Source/Developer/AssetTools/Private/AssetRenameManager.cpp
--- a/Engine/Source/Developer/AssetTools/Private/AssetRenameManager.cpp
+++ b/Engine/Source/Developer/AssetTools/Private/AssetRenameManager.cpp
@@ -639,7 +639,7 @@ struct FSoftObjectPathRenameSerializer : public FArchiveUObject
 
 			ThreadContext.GetSerializationOptions(ReferencingPackageName, ReferencingPropertyName, CollectType, SerializeType, this);
 
-			if (CollectType == ESoftObjectPathCollectType::NeverCollect || CollectType == ESoftObjectPathCollectType::NonPackage)
+			if (!UE::SoftObjectPath::IsCollectable(CollectType))
 			{
 				return *this;
 			}
diff --git a/Engine/Source/Runtime/AssetRegistry/Private/AssetRegistry.cpp b/Engine/Source/Runtime/AssetRegistry/Private/AssetRegistry.cpp
--- a/Engine/Source/Runtime/AssetRegistry/Private/AssetRegistry.cpp
+++ b/Engine/Source/Runtime/AssetRegistry/Private/AssetRegistry.cpp
@@ -2,6 +2,7 @@
 
 #include "AssetRegistry.h"
 
+#include "Algo/Transform.h"
 #include "Algo/Unique.h"
 #include "AssetDataGatherer.h"
 #include "AssetDataGathererPrivate.h"
@@ -1995,6 +1996,9 @@ void FAssetRegistryImpl::SearchAllAssets(Impl::FEventContext& EventContext,
 		InitialSearchStartTime = FPlatformTime::Seconds();
 		bInitialSearchStarted = true;
 		bInitialSearchCompleted.store(false, std::memory_order_relaxed);
+#if WITH_EDITOR
+		bGatherNeedsApplyObjectRedirects = true;
+#endif
 		UpdateMaxSecondsPerFrame(); // MaxSecondsPerFrame only depends on initial search
 	}
 	else if (!IsGathering())
@@ -2003,6 +2007,9 @@ void FAssetRegistryImpl::SearchAllAssets(Impl::FEventContext& EventContext,
 		AdditionalMountSearchStartTime = FPlatformTime::Seconds();
 		bAdditionalMountSearchInProgress.store(true, std::memory_order_relaxed);
 		GlobalGatherer->SetIsAdditionalMountSearchInProgress(true);
+#if WITH_EDITOR
+		bGatherNeedsApplyObjectRedirects = true;
+#endif
 	}
 
 	FAssetDataGatherer& Gatherer = *GlobalGatherer;
@@ -5204,6 +5211,15 @@ Impl::EGatherStatus FAssetRegistryImpl::TickGatherer(Impl::FTickContext& TickCon
 			return TickContext.UpdateIntermediateStatus();
 		}
 	}
+
+	if (bGatherNeedsApplyObjectRedirects)
+	{
+		ApplyObjectRedirectsToGatheredAssetDatas();
+		if (InOutInterruptionContext.ShouldExitEarly())
+		{
+			return TickContext.UpdateIntermediateStatus();
+		}
+	}
 #endif
 
 	// Finishing the gather is only possible on the game thread (and only for callsites of TickGather that opt into it)
@@ -5438,6 +5454,102 @@ void FTickContext::RunDependencyDataGathered(TMultiMap<FName, FPackageDependency
 
 } // namespace Impl within UE::AssetRegistry
 
+#if WITH_EDITOR
+void FAssetRegistryImpl::ApplyObjectRedirectsToGatheredAssetDatas()
+{
+	bGatherNeedsApplyObjectRedirects = false;
+
+	State.EnumerateAllMutableAssets([](FAssetData& AssetData)
+		{
+			if (!AssetData.TaggedAssetBundles)
+			{
+				return;
+			}
+
+			TArray<FAssetBundleEntry>& Bundles = AssetData.TaggedAssetBundles->Bundles;
+			int32 NumBundles = Bundles.Num();
+			int32 FirstChangedBundle = -1;
+			int32 FirstChangedAssetPath = -1;
+			FSoftObjectPath FirstChangedAssetPathValue;
+			for (FirstChangedBundle = 0; FirstChangedBundle < NumBundles; ++FirstChangedBundle)
+			{
+				FAssetBundleEntry& Bundle = Bundles[FirstChangedBundle];
+				int32 NumAssetPaths = Bundle.AssetPaths.Num();
+				for (FirstChangedAssetPath = 0; FirstChangedAssetPath < NumAssetPaths; ++FirstChangedAssetPath)
+				{
+					FSoftObjectPath Original(Bundle.AssetPaths[FirstChangedAssetPath]);
+					FSoftObjectPath Redirected = GRedirectCollector.GetAssetPathRedirection(Original);
+					if (Redirected.IsValid())
+					{
+						FirstChangedAssetPathValue = Redirected;
+						break;
+					}
+				}
+				if (FirstChangedAssetPath < NumAssetPaths)
+				{
+					break;
+				}
+				FirstChangedAssetPath = -1;
+			}
+
+			if (FirstChangedBundle == NumBundles)
+			{
+				return;
+			}
+
+			FAssetBundleData NewData;
+			NewData.Bundles.Reserve(NumBundles);
+			int32 BundleIndex;
+			for (BundleIndex = 0; BundleIndex < FirstChangedBundle; ++BundleIndex)
+			{
+				NewData.Bundles.Add(Bundles[BundleIndex]);
+			}
+			for (; BundleIndex < NumBundles; ++BundleIndex)
+			{
+				FAssetBundleEntry& NewBundle = NewData.Bundles.Emplace_GetRef(Bundles[BundleIndex]);
+				TArray<FTopLevelAssetPath>& NewAssetPaths = NewBundle.AssetPaths;
+
+				int32 NumAssetPaths = NewAssetPaths.Num();
+				int32 AssetPathIndex = FirstChangedAssetPath != -1 ? FirstChangedAssetPath : 0;
+				for (; AssetPathIndex < NumAssetPaths; ++AssetPathIndex)
+				{
+					FSoftObjectPath Redirected;
+					FSoftObjectPath Original(NewAssetPaths[AssetPathIndex]);
+					if (FirstChangedAssetPath == -1)
+					{
+						Redirected = GRedirectCollector.GetAssetPathRedirection(Original);
+					}
+					else
+					{
+						check(AssetPathIndex == FirstChangedAssetPath);
+						Redirected = FirstChangedAssetPathValue;
+						FirstChangedAssetPathValue = FSoftObjectPath();
+						FirstChangedAssetPath = -1;
+					}
+
+					if (Redirected.IsNull())
+					{
+						continue;
+					}
+					if (!Redirected.GetSubPathUtf8String().IsEmpty())
+					{
+						UE_LOG(LogAssetRegistry, Warning,
+							TEXT("AssetBundleEntry Name %s.%s[%d] was redirected from %s -> %s, but the redirected target is not a TopLevelAssetPath, and AssetBundleEntry Names must be TopLevelAssetPaths. Ignoring the redirector."),
+							*AssetData.ToSoftObjectPath().ToString(), *NewBundle.BundleName.ToString(), AssetPathIndex, *Original.ToString(), *Redirected.ToString());
+						continue;
+					}
+					NewAssetPaths[AssetPathIndex] = Redirected.GetAssetPath();
+				}
+				PRAGMA_DISABLE_DEPRECATION_WARNINGS
+				NewBundle.BundleAssets.Reset();
+				Algo::Transform(NewBundle.AssetPaths, NewBundle.BundleAssets, [](FTopLevelAssetPath Path) { return FSoftObjectPath(Path, {}); });
+				PRAGMA_ENABLE_DEPRECATION_WARNINGS
+			}
+			AssetData.TaggedAssetBundles = MakeShared<FAssetBundleData, ESPMode::ThreadSafe>(MoveTemp(NewData));
+		});
+}
+#endif
+
 void FAssetRegistryImpl::Shrink()
 {
 	BackgroundResults.Shrink();
@@ -8507,6 +8619,9 @@ void FAssetRegistryImpl::OnContentPathMounted(Impl::FEventContext& EventContext,
 				AdditionalMountSearchStartTime = FPlatformTime::Seconds();
 				bAdditionalMountSearchInProgress.store(true, std::memory_order_relaxed);
 				GlobalGatherer->SetIsAdditionalMountSearchInProgress(true);
+#if WITH_EDITOR
+				bGatherNeedsApplyObjectRedirects = true;
+#endif
 			}
 
 			GlobalGatherer->AddMountPoint(FileSystemPath, InAssetPath);
diff --git a/Engine/Source/Runtime/AssetRegistry/Private/AssetRegistryImpl.h b/Engine/Source/Runtime/AssetRegistry/Private/AssetRegistryImpl.h
--- a/Engine/Source/Runtime/AssetRegistry/Private/AssetRegistryImpl.h
+++ b/Engine/Source/Runtime/AssetRegistry/Private/AssetRegistryImpl.h
@@ -476,6 +476,8 @@ public:
 
 	/** Called when new gatherer is registered. Requires subsequent call to RebuildAssetDependencyGathererMapIfNeeded */
 	void OnAssetDependencyGathererRegistered() { bRegisteredDependencyGathererClassesDirty = true; }
+
+	void ApplyObjectRedirectsToGatheredAssetDatas();
 #endif
 
 	/** Adds an asset to the empty package list which contains packages that have no assets left in them */
@@ -822,6 +824,8 @@ private:
 	TMultiMap<FTopLevelAssetPath, UE::AssetDependencyGatherer::Private::FRegisteredAssetDependencyGatherer*> RegisteredDependencyGathererClasses;
 	mutable FRWLock RegisteredDependencyGathererClassesLock;
 	bool bRegisteredDependencyGathererClassesDirty;
+
+	bool bGatherNeedsApplyObjectRedirects = false;
 #endif
 #if WITH_ENGINE && WITH_EDITOR
 	/** Class names that return true for IsAsset but which should not be treated as assets in uncooked packages */
diff --git a/Engine/Source/Runtime/CoreUObject/Private/Misc/RedirectCollector.cpp b/Engine/Source/Runtime/CoreUObject/Private/Misc/RedirectCollector.cpp
--- a/Engine/Source/Runtime/CoreUObject/Private/Misc/RedirectCollector.cpp
+++ b/Engine/Source/Runtime/CoreUObject/Private/Misc/RedirectCollector.cpp
@@ -75,9 +75,11 @@ void FRedirectCollector::OnSoftObjectPathLoaded(const FSoftObjectPath& ObjectPat
 
 	ThreadContext.GetSerializationOptions(PackageName, PropertyName, CollectType, SerializeType, InArchive);
 
-	if (CollectType == ESoftObjectPathCollectType::NonPackage)
+	// We don't need to keep records for NonPackage Or PackageHeader CollectTypes, but we do need to keep records
+	// for all SoftObjectPaths in UObjects, even for the ones marked NeverCollect, so early exit for !IsObjectType
+	// CollectTypes but not necessarily for !IsCollectable CollectTypes.
+	if (!UE::SoftObjectPath::IsObjectType(CollectType))
 	{
-		// Do not track
 		return;
 	}
 
@@ -85,7 +87,7 @@ void FRedirectCollector::OnSoftObjectPathLoaded(const FSoftObjectPath& ObjectPat
 	FTopLevelAssetPath AssetPath = ObjectPath.GetAssetPath();
 
 	FScopeLock ScopeLock(&CriticalSection);
-	if (CollectType != ESoftObjectPathCollectType::NeverCollect)
+	if (UE::SoftObjectPath::IsCollectable(CollectType))
 	{
 		// Add this reference to the soft object inclusion list for the cook's iterative traversal of the soft dependency graph
 		FSoftObjectPathProperty SoftObjectPathProperty(FSoftObjectPath(AssetPath, {}), PropertyName, bReferencedByEditorOnlyProperty);
diff --git a/Engine/Source/Runtime/CoreUObject/Private/UObject/ImportExportCollector.cpp b/Engine/Source/Runtime/CoreUObject/Private/UObject/ImportExportCollector.cpp
--- a/Engine/Source/Runtime/CoreUObject/Private/UObject/ImportExportCollector.cpp
+++ b/Engine/Source/Runtime/CoreUObject/Private/UObject/ImportExportCollector.cpp
@@ -109,7 +109,7 @@ FArchive& FImportExportCollector::operator<<(FSoftObjectPath& Value)
 	FSoftObjectPathThreadContext& ThreadContext = FSoftObjectPathThreadContext::Get();
 	ThreadContext.GetSerializationOptions(CurrentPackage, PropertyName, CollectType, SerializeType, this);
 
-	if (CollectType != ESoftObjectPathCollectType::NeverCollect && CollectType != ESoftObjectPathCollectType::NonPackage)
+	if (UE::SoftObjectPath::IsCollectable(CollectType))
 	{
 		FName PackageName = Value.GetLongPackageFName();
 		if (PackageName != RootPackageName && !PackageName.IsNone())
diff --git a/Engine/Source/Runtime/CoreUObject/Private/UObject/LinkerLoad.cpp b/Engine/Source/Runtime/CoreUObject/Private/UObject/LinkerLoad.cpp
--- a/Engine/Source/Runtime/CoreUObject/Private/UObject/LinkerLoad.cpp
+++ b/Engine/Source/Runtime/CoreUObject/Private/UObject/LinkerLoad.cpp
@@ -2002,7 +2002,8 @@ FLinkerLoad::ELinkerStatus FLinkerLoad::SerializeSoftObjectPathList()
 	}
 
 #if WITH_EDITOR
-	FSoftObjectPathSerializationScope SerializationScope(NAME_None, NAME_None, ESoftObjectPathCollectType::NonPackage, ESoftObjectPathSerializeType::AlwaysSerialize);
+	FSoftObjectPathSerializationScope SerializationScope(NAME_None, NAME_None,
+		ESoftObjectPathCollectType::PackageHeader, ESoftObjectPathSerializeType::AlwaysSerialize);
 #endif // WITH_EDITOR
 
 	FStructuredArchive::FStream Stream = StructuredArchiveRootRecord->EnterStream(TEXT("SoftObjectPathList"));
diff --git a/Engine/Source/Runtime/CoreUObject/Private/UObject/SavePackage/PackageHarvester.cpp b/Engine/Source/Runtime/CoreUObject/Private/UObject/SavePackage/PackageHarvester.cpp
--- a/Engine/Source/Runtime/CoreUObject/Private/UObject/SavePackage/PackageHarvester.cpp
+++ b/Engine/Source/Runtime/CoreUObject/Private/UObject/SavePackage/PackageHarvester.cpp
@@ -1135,7 +1135,7 @@ FArchive& FPackageHarvester::operator<<(FSoftObjectPath& Value)
 
 		// Don't track as a soft reference to a package if this is a NeverCollect SoftObjectPath
 		ThreadContext.GetSerializationOptions(ReferencingPackageName, ReferencingPropertyName, CollectType, SerializeType, this);
-		if (CollectType != ESoftObjectPathCollectType::NeverCollect && CollectType != ESoftObjectPathCollectType::NonPackage)
+		if (UE::SoftObjectPath::IsCollectable(CollectType))
 		{
 #if WITH_EDITORONLY_DATA
 			// Don't track as a used-in-game soft reference to a package if this is an EditorOnly SoftObjectPath
diff --git a/Engine/Source/Runtime/CoreUObject/Private/UObject/SoftObjectPath.cpp b/Engine/Source/Runtime/CoreUObject/Private/UObject/SoftObjectPath.cpp
--- a/Engine/Source/Runtime/CoreUObject/Private/UObject/SoftObjectPath.cpp
+++ b/Engine/Source/Runtime/CoreUObject/Private/UObject/SoftObjectPath.cpp
@@ -418,11 +418,21 @@ void FSoftObjectPath::SerializePath(FArchive& Ar)
 		{
 			PostLoadPath(&Ar);
 
-			// If we think it's going to work, we try to do the pre-save fixup now. This is important because it helps with blueprint CDO save determinism with redirectors
-			// It's important that the entire CDO hierarchy gets fixed up before an instance in a map gets saved otherwise the delta serialization will save too much
-			// If the asset registry hasn't fully loaded this won't necessarily work, but it won't do any harm
-			// This will never work in -game builds or on initial load so don't try
-			if (GIsEditor && !GIsInitialLoad)
+			// If we think it's going to work, we try to do the pre-save fixup now. This is important because it helps
+			// with blueprint CDO save determinism with redirectors. It's important that the entire CDO hierarchy gets
+			// fixed up before an instance in a map gets saved otherwise the delta serialization will save too much.
+			// 
+			// We need to not do this for SoftObjectPaths loaded in AssetBundle tags, because those are loaded before
+			// the AssetRegistry finishes gathering all of the redirectors, so the results would be indeterministic.
+			// We process those paths for redirectors instead at the end of the AssetRegistry.
+			// We detect whether the SoftObjectPaths came from the AssetRegistry by checking for IsPackageType(CollectType)
+			// Other non-package sources will also be skipped, but that's okay; we only need to do the PreSavePath for
+			// SoftObjectPaths that are in Blueprints.
+			//
+			// For paths in blueprints, the handling of redirects is only needed for editor (game does not save blueprints),
+			// and also even in editor we know the information is not yet ready during engine startup, so skip the PreSave
+			// for !GIsEditor || GIsInitialLoad.
+			if (GIsEditor && !GIsInitialLoad && UE::SoftObjectPath::IsPackageType(CollectType))
 			{
 				PreSavePath(nullptr);
 			}
diff --git a/Engine/Source/Runtime/CoreUObject/Public/UObject/SoftObjectPath.h b/Engine/Source/Runtime/CoreUObject/Public/UObject/SoftObjectPath.h
--- a/Engine/Source/Runtime/CoreUObject/Public/UObject/SoftObjectPath.h
+++ b/Engine/Source/Runtime/CoreUObject/Public/UObject/SoftObjectPath.h
@@ -571,8 +571,16 @@ private:
 /** Options for how to set soft object path collection */
 enum class ESoftObjectPathCollectType : uint8
 {
-	/** The SoftObjectPath being loaded is not in a package, so we do not need to record it in inclusion or exclusion lists*/
+	/**
+	 * The SoftObjectPath being loaded is not in a package, so we do not need to record it in inclusion or exclusion lists,
+	 * or translate redirects for it.
+	 */
 	NonPackage,
+	/**
+	 * The SoftObjectPath being loaded is in a package header, not in an Object, so we do not need to record it in
+	 * inclusion or exclusion lists, but we do need to translate redirects for it.
+	 */
+	PackageHeader,
 	/** References is not tracked in any situation, transient reference */
 	NeverCollect,
 	/** Editor only reference, this is tracked for redirector fixup but not for cooking */
@@ -581,6 +589,38 @@ enum class ESoftObjectPathCollectType : uint8
 	AlwaysCollect,
 };
 
+namespace UE::SoftObjectPath
+{
+
+/**
+ * Whether the CollectType indicates that this SoftObjectPath is from a package, either in its header or in an object
+ * within the package.
+ */
+inline bool IsPackageType(ESoftObjectPathCollectType CollectType)
+{
+	return CollectType != ESoftObjectPathCollectType::NonPackage;
+}
+
+/** Whether the CollectType indicates that this SoftObjectPath is in a UObject in a package. */
+inline bool IsObjectType(ESoftObjectPathCollectType CollectType)
+{
+	return (CollectType != ESoftObjectPathCollectType::NonPackage)
+		& (CollectType != ESoftObjectPathCollectType::PackageHeader);
+}
+
+/**
+ * Whether the CollectType is one of the ones that should always or sometimes be collected. This means only ObjectType
+ * CollectTypes that are not NeverCollect.
+ */
+inline bool IsCollectable(ESoftObjectPathCollectType CollectType)
+{
+	return (CollectType != ESoftObjectPathCollectType::NonPackage)
+		& (CollectType != ESoftObjectPathCollectType::PackageHeader)
+		& (CollectType != ESoftObjectPathCollectType::NeverCollect);
+}
+
+}
+
 /** Rules for actually serializing the internals of soft object paths */
 enum class ESoftObjectPathSerializeType : uint8
 {
