CookCorrectness/AssetRegistry: Refactor DependsNode serialization to make less complicated the upcoming addition of backwards compatibility for a change in Manage dependency flags. No behavioral or performance changes.
#rnx

diff --git a/Engine/Source/Runtime/AssetRegistry/Private/AssetRegistryState.cpp b/Engine/Source/Runtime/AssetRegistry/Private/AssetRegistryState.cpp
--- a/Engine/Source/Runtime/AssetRegistry/Private/AssetRegistryState.cpp
+++ b/Engine/Source/Runtime/AssetRegistry/Private/AssetRegistryState.cpp
@@ -2107,30 +2107,7 @@ void FAssetRegistryState::Load(Archive&& Ar, const FAssetRegistryHeader& Header,
 
 	SetAssetDatas(PreallocatedAssetDataBuffer, Options);
 
-	if (Version < FAssetRegistryVersion::AddedDependencyFlags)
-	{
-		LoadDependencies_BeforeFlags(Ar, Options.bLoadDependencies, Version);
-	}
-	else
-	{
-		int64 DependencySectionSize;
-		Ar << DependencySectionSize;
-		int64 DependencySectionEnd = Ar.Tell() + DependencySectionSize;
-
-#if ASSET_REGISTRY_ALLOW_DEPENDENCY_SERIALIZATION
-		if (Options.bLoadDependencies)
-		{
-			LoadDependencies(Ar);
-		}
-			
-		if (!Options.bLoadDependencies || Ar.IsError())
-		{
-			Ar.Seek(DependencySectionEnd);
-		}
-#else
-		Ar.Seek(DependencySectionEnd);
-#endif
-	}
+	LoadDependencies(Ar, Options, Version);
 
 	int32 LocalNumPackageData = 0;
 	Ar << LocalNumPackageData;
@@ -2172,7 +2149,44 @@ void FAssetRegistryState::Load(Archive&& Ar, const FAssetRegistryHeader& Header,
 	}
 }
 
-void FAssetRegistryState::LoadDependencies(FArchive& Ar)
+void FAssetRegistryState::LoadDependencies(FArchive& Ar, const FAssetRegistryLoadOptions& Options,
+	FAssetRegistryVersion::Type Version)
+{
+	if (Version < FAssetRegistryVersion::AddedDependencyFlags)
+	{
+		LegacyLoadDependencies_BeforeFlags(Ar, Options.bLoadDependencies, Version);
+	}
+	else
+	{
+		int64 DependencySectionSize;
+		Ar << DependencySectionSize;
+		int64 DependencySectionEnd = Ar.Tell() + DependencySectionSize;
+
+#if ASSET_REGISTRY_ALLOW_DEPENDENCY_SERIALIZATION
+		if (Options.bLoadDependencies)
+		{
+			if (Version == FAssetRegistryVersion::LatestVersion)
+			{
+				LoadDependenciesImpl<true>(Ar, Version);
+			}
+			else
+			{
+				LoadDependenciesImpl<false>(Ar, Version);
+			}
+		}
+
+		if (!Options.bLoadDependencies || Ar.IsError())
+		{
+			Ar.Seek(DependencySectionEnd);
+		}
+#else
+		Ar.Seek(DependencySectionEnd);
+#endif
+	}
+}
+
+template <bool bLatestVersion>
+void FAssetRegistryState::LoadDependenciesImpl(FArchive& Ar, FAssetRegistryVersion::Type Version)
 {
 	int32 LocalNumDependsNodes = 0;
 	Ar << LocalNumDependsNodes;
@@ -2200,12 +2214,12 @@ void FAssetRegistryState::LoadDependencies(FArchive& Ar)
 	for (int32 DependsNodeIndex = 0; DependsNodeIndex < LocalNumDependsNodes; DependsNodeIndex++)
 	{
 		FDependsNode* DependsNode = &PreallocatedDependsNodeDataBuffer[DependsNodeIndex];
-		DependsNode->SerializeLoad(Ar, GetNodeFromSerializeIndex, Scratch);
+		DependsNode->SerializeLoad<bLatestVersion>(Ar, GetNodeFromSerializeIndex, Scratch, Version);
 		CachedDependsNodes.Add(DependsNode->GetIdentifier(), DependsNode);
 	}
 }
 
-void FAssetRegistryState::LoadDependencies_BeforeFlags(FArchive& Ar, bool bSerializeDependencies,
+void FAssetRegistryState::LegacyLoadDependencies_BeforeFlags(FArchive& Ar, bool bSerializeDependencies,
 	FAssetRegistryVersion::Type Version)
 {
 	int32 LocalNumDependsNodes = 0;
@@ -2229,9 +2243,6 @@ void FAssetRegistryState::LoadDependencies_BeforeFlags(FArchive& Ar, bool bSeria
 		return &PreallocatedDependsNodeDataBuffer[Index];
 	};
 
-	uint32 HardBits, SoftBits, HardManageBits, SoftManageBits;
-	FDependsNode::GetPropertySetBits_BeforeFlags(HardBits, SoftBits, HardManageBits, SoftManageBits);
-
 	TArray<FDependsNode*> DependsNodes;
 	for (int32 DependsNodeIndex = 0; DependsNodeIndex < LocalNumDependsNodes; DependsNodeIndex++)
 	{
@@ -2247,8 +2258,8 @@ void FAssetRegistryState::LoadDependencies_BeforeFlags(FArchive& Ar, bool bSeria
 		}
 
 		// Call the DependsNode legacy serialization function
-		DependsNode->SerializeLoad_BeforeFlags(Ar, Version, PreallocatedDependsNodeDataBuffer, LocalNumDependsNodes,
-			bSerializeDependencies, HardBits, SoftBits, HardManageBits, SoftManageBits);
+		DependsNode->LegacySerializeLoad_BeforeFlags(Ar, Version, PreallocatedDependsNodeDataBuffer,
+			LocalNumDependsNodes, bSerializeDependencies);
 
 		// Register the DependsNode with its AssetIdentifier
 		if (bSerializeDependencies)
diff --git a/Engine/Source/Runtime/AssetRegistry/Private/DependsNode.cpp b/Engine/Source/Runtime/AssetRegistry/Private/DependsNode.cpp
--- a/Engine/Source/Runtime/AssetRegistry/Private/DependsNode.cpp
+++ b/Engine/Source/Runtime/AssetRegistry/Private/DependsNode.cpp
@@ -863,94 +863,238 @@ void FDependsNode::SerializeSave(FArchive& Ar,
 	WriteDependencies(Referencers, nullptr, 0, true);
 }
 
-void FDependsNode::SerializeLoad(FArchive& Ar,
+template <bool bLatestVersion>
+void FDependsNode::SerializeLoad_PackageDependencies(FArchive& Ar,
 	const TUniqueFunction<FDependsNode* (int32)>& GetNodeFromSerializeIndex,
-	FLoadScratch& Scratch)
+	FLoadScratch& Scratch, FAssetRegistryVersion::Type Version)
 {
-	Ar << Identifier;
-	SetIdentifier(Identifier);
-
-	auto ReadDependencies = [&Ar, &GetNodeFromSerializeIndex, &Scratch](TArray<FDependsNode*>& OutDependencies,
-		TBitArray<>* OutFlagBits,
-		int FlagSetWidth)
-	{
-		TArray<int32>& InDependencies = Scratch.InDependencies;
-		TArray<uint32>& InFlagBits = Scratch.InFlagBits;
-		TArray<FDependsNode*>& PointerDependencies = Scratch.PointerDependencies;
-		TArray<int32>& SortIndexes = Scratch.SortIndexes;
-		int32 NumFlagBits = 0;
-
-		InDependencies.Reset();
-		Ar << InDependencies;
-		int32 NumDependencies = InDependencies.Num();
-		if (OutFlagBits)
+	// This code is duplicated-with-modifications in SerializeLoad_PackageDependencies,
+	// and SerializeLoad_ManageDependencies, because sharing code makes it more difficult to read.
+	TArray<FDependsNode*>& OutDependencies = this->PackageDependencies;
+	TBitArray<>& OutFlagBits = this->PackageFlags;
+	int FlagSetWidth = PackageFlagSetWidth;
+
+	TArray<int32>& InDependencies = Scratch.InDependencies;
+	TArray<uint32>& InFlagBits = Scratch.InFlagBits;
+	TArray<FDependsNode*>& PointerDependencies = Scratch.PointerDependencies;
+	TArray<int32>& SortIndexes = Scratch.SortIndexes;
+	int32 NumFlagBits = 0;
+
+	InDependencies.Reset();
+	Ar << InDependencies;
+	int32 NumDependencies = InDependencies.Num();
+
+	// We don't use BitArray::operator<< because we want to avoid the reallocation that operator<< does on load
+	// and we want to avoid saving BitArray.Num when it can be derived from NumDependencies.
+	NumFlagBits = FlagSetWidth * NumDependencies;
+	const int32 NumFlagWords = FBitSet::CalculateNumWords(NumFlagBits);
+	InFlagBits.SetNumUninitialized(NumFlagWords);
+	Ar.Serialize(InFlagBits.GetData(), NumFlagWords * sizeof(uint32));
+
+	PointerDependencies.Reset(NumDependencies);
+	for (int32 SerializeIndex : InDependencies)
+	{
+		FDependsNode* DependsNode = GetNodeFromSerializeIndex(SerializeIndex);
+		if (!DependsNode)
 		{
-			// We don't use BitArray::operator<< because we want to avoid the reallocation that operator<< does on load
-			// and we want to avoid saving BitArray.Num when it can be derived from NumDependencies.
-			NumFlagBits = FlagSetWidth * NumDependencies;
-			const int32 NumFlagWords = FBitSet::CalculateNumWords(NumFlagBits);
-			InFlagBits.SetNumUninitialized(NumFlagWords);
-			Ar.Serialize(InFlagBits.GetData(), NumFlagWords * sizeof(uint32));
+			Ar.SetError();
+			return;
 		}
-		
-		PointerDependencies.Reset(NumDependencies);
-		for (int32 SerializeIndex : InDependencies)
+		PointerDependencies.Add(DependsNode);
+	}
+
+	SortIndexes.Reset(NumDependencies);
+	for (int Index = 0; Index < NumDependencies; ++Index)
+	{
+		SortIndexes.Add(Index);
+	}
+
+	Algo::Sort(SortIndexes, [&PointerDependencies](int32 A, int32 B)
 		{
-			FDependsNode* DependsNode = GetNodeFromSerializeIndex(SerializeIndex);
-			if (!DependsNode)
-			{
-				Ar.SetError();
-				return;
-			}
-			PointerDependencies.Add(DependsNode);
-		}
+			return PointerDependencies[A] < PointerDependencies[B];
+		});
 
-		SortIndexes.Reset(NumDependencies);
-		for (int Index = 0; Index < NumDependencies; ++Index)
+	OutDependencies.Empty(NumDependencies);
+	for (int32 SortIndex : SortIndexes)
+	{
+		OutDependencies.Add(PointerDependencies[SortIndex]);
+	}
+
+	OutFlagBits.SetNumUninitialized(NumFlagBits);
+	uint32* InFlagBitsData = InFlagBits.GetData();
+	for (int32 WriteIndex = 0; WriteIndex < NumDependencies; ++WriteIndex)
+	{
+		int32 ReadIndex = SortIndexes[WriteIndex];
+		OutFlagBits.SetRangeFromRange(WriteIndex * FlagSetWidth, FlagSetWidth,
+			InFlagBitsData, ReadIndex * FlagSetWidth);
+	}
+}
+
+template <bool bLatestVersion>
+void FDependsNode::SerializeLoad_NameDependencies(FArchive& Ar,
+	const TUniqueFunction<FDependsNode* (int32)>& GetNodeFromSerializeIndex,
+	FLoadScratch& Scratch, FAssetRegistryVersion::Type Version)
+{
+	// This code is duplicated-with-modifications in SerializeLoad_NameDependencies
+	// and SerializeLoad_Referencers, because sharing code makes it more difficult to read.
+	TArray<FDependsNode*>& OutDependencies = this->NameDependencies;
+
+	TArray<int32>& InDependencies = Scratch.InDependencies;
+	InDependencies.Reset();
+	Ar << InDependencies;
+	int32 NumDependencies = InDependencies.Num();
+
+	OutDependencies.Reset(NumDependencies);
+	for (int32 SerializeIndex : InDependencies)
+	{
+		FDependsNode* DependsNode = GetNodeFromSerializeIndex(SerializeIndex);
+		if (!DependsNode)
 		{
-			SortIndexes.Add(Index);
+			Ar.SetError();
+			return;
 		}
+		OutDependencies.Add(DependsNode);
+	}
 
-		Algo::Sort(SortIndexes, [&PointerDependencies](int32 A, int32 B)
-			{
-				return PointerDependencies[A] < PointerDependencies[B];
-			});
+	Algo::Sort(OutDependencies, [](FDependsNode* A, FDependsNode* B)
+		{
+			return A < B;
+		});
+}
+
+template <bool bLatestVersion>
+void FDependsNode::SerializeLoad_ManageDependencies(FArchive& Ar,
+	const TUniqueFunction<FDependsNode* (int32)>& GetNodeFromSerializeIndex,
+	FLoadScratch& Scratch, FAssetRegistryVersion::Type Version)
+{
+	using namespace UE::AssetRegistry;
 
-		OutDependencies.Empty(NumDependencies);
-		for (int32 SortIndex : SortIndexes)
+	// This code is duplicated-with-modifications in SerializeLoad_PackageDependencies,
+	// and SerializeLoad_ManageDependencies, because sharing code makes it more difficult to read.
+	TArray<FDependsNode*>& OutDependencies = this->ManageDependencies;
+	TBitArray<>& OutFlagBits = this->ManageFlags;
+	int OutFlagSetWidth = ManageFlagSetWidth;
+	int InFlagSetWidth = OutFlagSetWidth;
+
+	TArray<int32>& InDependencies = Scratch.InDependencies;
+	TArray<uint32>& InFlagBits = Scratch.InFlagBits;
+	TArray<FDependsNode*>& PointerDependencies = Scratch.PointerDependencies;
+	TArray<int32>& SortIndexes = Scratch.SortIndexes;
+	int32 InNumFlagBits = 0;
+
+	InDependencies.Reset();
+	Ar << InDependencies;
+	int32 NumDependencies = InDependencies.Num();
+
+	// We don't use BitArray::operator<< because we want to avoid the reallocation that operator<< does on load
+	// and we want to avoid saving BitArray.Num when it can be derived from NumDependencies.
+	InNumFlagBits = InFlagSetWidth * NumDependencies;
+	const int32 InNumFlagWords = FBitSet::CalculateNumWords(InNumFlagBits);
+	InFlagBits.SetNumUninitialized(InNumFlagWords);
+	Ar.Serialize(InFlagBits.GetData(), InNumFlagWords * sizeof(uint32));
+
+	PointerDependencies.Reset(NumDependencies);
+	for (int32 SerializeIndex : InDependencies)
+	{
+		FDependsNode* DependsNode = GetNodeFromSerializeIndex(SerializeIndex);
+		if (!DependsNode)
 		{
-			OutDependencies.Add(PointerDependencies[SortIndex]);
+			Ar.SetError();
+			return;
 		}
-		if (OutFlagBits)
+		PointerDependencies.Add(DependsNode);
+	}
+
+	SortIndexes.Reset(NumDependencies);
+	for (int Index = 0; Index < NumDependencies; ++Index)
+	{
+		SortIndexes.Add(Index);
+	}
+
+	Algo::Sort(SortIndexes, [&PointerDependencies](int32 A, int32 B)
 		{
-			OutFlagBits->SetNumUninitialized(NumFlagBits);
-			uint32* InFlagBitsData = InFlagBits.GetData();
-			for (int32 WriteIndex = 0; WriteIndex < NumDependencies; ++WriteIndex)
-			{
-				int32 ReadIndex = SortIndexes[WriteIndex];
-				OutFlagBits->SetRangeFromRange(WriteIndex * FlagSetWidth, FlagSetWidth,
-					InFlagBitsData, ReadIndex * FlagSetWidth);
-			}
+			return PointerDependencies[A] < PointerDependencies[B];
+		});
+
+	OutDependencies.Empty(NumDependencies);
+	for (int32 SortIndex : SortIndexes)
+	{
+		OutDependencies.Add(PointerDependencies[SortIndex]);
+	}
+
+	OutFlagBits.SetNumUninitialized(InNumFlagBits);
+	uint32* InFlagBitsData = InFlagBits.GetData();
+	for (int32 WriteIndex = 0; WriteIndex < NumDependencies; ++WriteIndex)
+	{
+		int32 ReadIndex = SortIndexes[WriteIndex];
+		OutFlagBits.SetRangeFromRange(WriteIndex * OutFlagSetWidth, OutFlagSetWidth,
+			InFlagBitsData, ReadIndex * OutFlagSetWidth);
+	}
+}
+
+template <bool bLatestVersion>
+void FDependsNode::SerializeLoad_Referencers(FArchive& Ar,
+	const TUniqueFunction<FDependsNode* (int32)>& GetNodeFromSerializeIndex,
+	FLoadScratch& Scratch, FAssetRegistryVersion::Type Version)
+{
+	// This code is duplicated-with-modifications in SerializeLoad_NameDependencies
+	// and SerializeLoad_Referencers, because sharing code makes it more difficult to read.
+	TArray<FDependsNode*>& OutDependencies = this->Referencers;
+
+	TArray<int32>& InDependencies = Scratch.InDependencies;
+	InDependencies.Reset();
+	Ar << InDependencies;
+	int32 NumDependencies = InDependencies.Num();
+
+	OutDependencies.Reset(NumDependencies);
+	for (int32 SerializeIndex : InDependencies)
+	{
+		FDependsNode* DependsNode = GetNodeFromSerializeIndex(SerializeIndex);
+		if (!DependsNode)
+		{
+			Ar.SetError();
+			return;
 		}
-	};
+		OutDependencies.Add(DependsNode);
+	}
 
-	ReadDependencies(PackageDependencies, &PackageFlags, PackageFlagSetWidth);
-	ReadDependencies(NameDependencies, nullptr, 0);
-	ReadDependencies(ManageDependencies, &ManageFlags, ManageFlagSetWidth);
-	ReadDependencies(Referencers, nullptr, 0);
+	Algo::Sort(OutDependencies, [](FDependsNode* A, FDependsNode* B)
+		{
+			return A < B;
+		});
+}
+
+template <bool bLatestVersion>
+void FDependsNode::SerializeLoad(FArchive& Ar,
+	const TUniqueFunction<FDependsNode* (int32)>& GetNodeFromSerializeIndex,
+	FLoadScratch& Scratch, FAssetRegistryVersion::Type Version)
+{
+	using namespace UE::AssetRegistry;
+
+	Ar << Identifier;
+	SetIdentifier(Identifier);
+
+	SerializeLoad_PackageDependencies<bLatestVersion>(Ar, GetNodeFromSerializeIndex, Scratch, Version);
+	SerializeLoad_NameDependencies<bLatestVersion>(Ar, GetNodeFromSerializeIndex, Scratch, Version);
+	SerializeLoad_ManageDependencies<bLatestVersion>(Ar, GetNodeFromSerializeIndex, Scratch, Version);
+	SerializeLoad_Referencers<bLatestVersion>(Ar, GetNodeFromSerializeIndex, Scratch, Version);
 
 	SetIsDependenciesInitialized(true);
 }
 
-void FDependsNode::SerializeLoad_BeforeFlags(FArchive& Ar,
+// We need Explicit instantiations of the template function because they are only called from other translation units.
+template void FDependsNode::SerializeLoad<true>(FArchive& Ar,
+	const TUniqueFunction<FDependsNode* (int32)>& GetNodeFromSerializeIndex,
+	FLoadScratch& Scratch, FAssetRegistryVersion::Type Version);
+template void FDependsNode::SerializeLoad<false>(FArchive& Ar,
+	const TUniqueFunction<FDependsNode* (int32)>& GetNodeFromSerializeIndex,
+	FLoadScratch& Scratch, FAssetRegistryVersion::Type Version);
+
+void FDependsNode::LegacySerializeLoad_BeforeFlags(FArchive& Ar,
 	FAssetRegistryVersion::Type Version,
 	FDependsNode* PreallocatedDependsNodeDataBuffer,
 	int32 NumDependsNodes,
-	bool bSerializeDependencies,
-	uint32 HardBits,
-	uint32 SoftBits,
-	uint32 HardManageBits,
-	uint32 SoftManageBits)
+	bool bSerializeDependencies)
 {
 	Ar << Identifier;
 	SetIdentifier(Identifier);
@@ -1017,38 +1161,51 @@ void FDependsNode::SerializeLoad_BeforeFlags(FArchive& Ar,
 	};
 
 	// Read the bits for each type, but don't write anything if serializing that type isn't allowed
-	SerializeNodeArray(NumHard, PackageDependencies, &PackageFlags, PackageFlagSetWidth, HardBits, true, true);
-	SerializeNodeArray(NumSoft, PackageDependencies, &PackageFlags, PackageFlagSetWidth, SoftBits, false, true);
-	SerializeNodeArray(NumName, NameDependencies, nullptr, 0, 0, true, bSerializeDependencies);
-	SerializeNodeArray(NumSoftManage, ManageDependencies, &ManageFlags, ManageFlagSetWidth, SoftManageBits,
-		true, bSerializeDependencies);
-	SerializeNodeArray(NumHardManage, ManageDependencies, &ManageFlags, ManageFlagSetWidth, HardManageBits,
-		true, bSerializeDependencies);
-	SerializeNodeArray(NumReferencers, Referencers, nullptr, 0, 0, true, true);
-
-	SetIsDependenciesInitialized(true);
-}
-
-void FDependsNode::GetPropertySetBits_BeforeFlags(uint32& HardBits,
-	uint32& SoftBits,
-	uint32& HardManageBits,
-	uint32& SoftManageBits)
-{
+	if (NumHard > 0)
 	{
+		// All legacy PackageDependencies were Game and Build. Hard had the same meaning as EDependencyProperty::Hard.
+		uint32 HardBits = 0;
 		FDependsNode::FPackageFlagSet FlagSet;
 		FlagSet.Add(PackagePropertiesToByte(
 			UE::AssetRegistry::EDependencyProperty::Hard
 			| UE::AssetRegistry::EDependencyProperty::Game | UE::AssetRegistry::EDependencyProperty::Build));
 		FlagSet.Save(&HardBits);
+		SerializeNodeArray(NumHard, PackageDependencies, &PackageFlags, PackageFlagSetWidth, HardBits, true, true);
 	}
+	if (NumSoft > 0)
 	{
+		// All legacy PackageDependencies were Game and Build. Soft had the same meaning as !EDependencyProperty::Hard.
+		uint32 SoftBits = 0;
 		FDependsNode::FPackageFlagSet FlagSet;
 		FlagSet.Add(PackagePropertiesToByte(
 			UE::AssetRegistry::EDependencyProperty::Game | UE::AssetRegistry::EDependencyProperty::Build));
 		FlagSet.Save(&SoftBits);
+		SerializeNodeArray(NumSoft, PackageDependencies, &PackageFlags, PackageFlagSetWidth, SoftBits, false, true);
+	}
+	if (NumName > 0)
+	{
+		SerializeNodeArray(NumName, NameDependencies, nullptr, 0, 0, true, bSerializeDependencies);
+	}
+	if (NumSoftManage > 0)
+	{
+		// In legacy ManageDependencies, Soft meant Indirect, which we currently record as 0x0.
+		uint32 SoftManageBits = 0x0;
+		SerializeNodeArray(NumSoftManage, ManageDependencies, &ManageFlags, ManageFlagSetWidth, SoftManageBits,
+			true, bSerializeDependencies);
 	}
-	HardManageBits = 0x1;
-	SoftManageBits = 0x0;
+	if (NumHardManage > 0)
+	{
+		// In legacy ManageDependencies, Hard meant Direct, which we currently record as 0x1.
+		uint32 HardManageBits = 0x1;
+		SerializeNodeArray(NumHardManage, ManageDependencies, &ManageFlags, ManageFlagSetWidth, HardManageBits,
+			true, bSerializeDependencies);
+	}
+	if (NumReferencers > 0)
+	{
+		SerializeNodeArray(NumReferencers, Referencers, nullptr, 0, 0, true, true);
+	}
+
+	SetIsDependenciesInitialized(true);
 }
 
 template <uint32 FlagWidth>
diff --git a/Engine/Source/Runtime/AssetRegistry/Private/DependsNode.h b/Engine/Source/Runtime/AssetRegistry/Private/DependsNode.h
--- a/Engine/Source/Runtime/AssetRegistry/Private/DependsNode.h
+++ b/Engine/Source/Runtime/AssetRegistry/Private/DependsNode.h
@@ -163,14 +163,13 @@ public:
 		TArray<FDependsNode*> PointerDependencies;
 		TArray<int32> SortIndexes;
 	};
+	template <bool bLatestVersion>
 	void SerializeLoad(FArchive& Ar,
-		const TUniqueFunction<FDependsNode*(int32)>& GetNodeFromSerializeIndex, FLoadScratch& Scratch);
+		const TUniqueFunction<FDependsNode* (int32)>& GetNodeFromSerializeIndex,
+		FLoadScratch& Scratch, FAssetRegistryVersion::Type Version);
 
-	void SerializeLoad_BeforeFlags(FArchive& Ar, FAssetRegistryVersion::Type Version,
-		FDependsNode* PreallocatedDependsNodeDataBuffer, int32 NumDependsNodes, bool bSerializeDependencies,
-		uint32 HardBits, uint32 SoftBits, uint32 HardManageBits, uint32 SoftManageBits);
-	static void GetPropertySetBits_BeforeFlags(uint32& HardBits, uint32& SoftBits,
-		uint32& HardManageBits, uint32& SoftManageBits);
+	void LegacySerializeLoad_BeforeFlags(FArchive& Ar, FAssetRegistryVersion::Type Version,
+		FDependsNode* PreallocatedDependsNodeDataBuffer, int32 NumDependsNodes, bool bSerializeDependencies);
 
 	bool IsDependencyListSorted(UE::AssetRegistry::EDependencyCategory Category) const;
 	void SetIsDependencyListSorted(UE::AssetRegistry::EDependencyCategory Category, bool bValue);
@@ -194,6 +193,23 @@ private:
 
 	void ConstructFlags();
 
+	template <bool bLatestVersion>
+	void SerializeLoad_PackageDependencies(FArchive& Ar,
+		const TUniqueFunction<FDependsNode* (int32)>& GetNodeFromSerializeIndex,
+		FLoadScratch& Scratch, FAssetRegistryVersion::Type Version);
+	template <bool bLatestVersion>
+	void SerializeLoad_NameDependencies(FArchive& Ar,
+		const TUniqueFunction<FDependsNode* (int32)>& GetNodeFromSerializeIndex,
+		FLoadScratch& Scratch, FAssetRegistryVersion::Type Version);
+	template <bool bLatestVersion>
+	void SerializeLoad_ManageDependencies(FArchive& Ar,
+		const TUniqueFunction<FDependsNode* (int32)>& GetNodeFromSerializeIndex,
+		FLoadScratch& Scratch, FAssetRegistryVersion::Type Version);
+	template <bool bLatestVersion>
+	void SerializeLoad_Referencers(FArchive& Ar,
+		const TUniqueFunction<FDependsNode* (int32)>& GetNodeFromSerializeIndex,
+		FLoadScratch& Scratch, FAssetRegistryVersion::Type Version);
+
 	/** The name of the package/object this node represents */
 	FAssetIdentifier Identifier;
 	FDependsNodeList PackageDependencies;
diff --git a/Engine/Source/Runtime/AssetRegistry/Private/DependsNodeTest.cpp b/Engine/Source/Runtime/AssetRegistry/Private/DependsNodeTest.cpp
--- a/Engine/Source/Runtime/AssetRegistry/Private/DependsNodeTest.cpp
+++ b/Engine/Source/Runtime/AssetRegistry/Private/DependsNodeTest.cpp
@@ -1130,10 +1130,10 @@ bool FDependsNodeTest::RunTest(const FString& Parameters)
 
 				{
 					FMemoryReader Reader(Bytes);
-					LoadedA.SerializeLoad(Reader, GetNodeFromSerializeIndex, LoadScratch);
+					LoadedA.SerializeLoad<true>(Reader, GetNodeFromSerializeIndex, LoadScratch, FAssetRegistryVersion::LatestVersion);
 					if (KeepNode2 == EKeepNode2::KeepAlways)
 					{
-						LoadedB.SerializeLoad(Reader, GetNodeFromSerializeIndex, LoadScratch);
+						LoadedB.SerializeLoad<true>(Reader, GetNodeFromSerializeIndex, LoadScratch, FAssetRegistryVersion::LatestVersion);
 					}
 				}
 
diff --git a/Engine/Source/Runtime/AssetRegistry/Public/AssetRegistry/AssetRegistryState.h b/Engine/Source/Runtime/AssetRegistry/Public/AssetRegistry/AssetRegistryState.h
--- a/Engine/Source/Runtime/AssetRegistry/Public/AssetRegistry/AssetRegistryState.h
+++ b/Engine/Source/Runtime/AssetRegistry/Public/AssetRegistry/AssetRegistryState.h
@@ -680,8 +680,10 @@ private:
 	static TSharedPtr<FAssetBundleData> FilterBundles(const TSharedPtr<FAssetBundleData>& InAssetBundles,
 		const FAssetRegistrySerializationOptions& Options);
 
-	void LoadDependencies(FArchive& Ar);
-	void LoadDependencies_BeforeFlags(FArchive& Ar, bool bSerializeDependencies, FAssetRegistryVersion::Type Version);
+	void LoadDependencies(FArchive& Ar, const FAssetRegistryLoadOptions& Options, FAssetRegistryVersion::Type Version);
+	template <bool bLatestVersion>
+	void LoadDependenciesImpl(FArchive& Ar, FAssetRegistryVersion::Type Version);
+	void LegacyLoadDependencies_BeforeFlags(FArchive& Ar, bool bSerializeDependencies, FAssetRegistryVersion::Type Version);
 
 	void SetTagsOnExistingAsset(FAssetData* AssetData, FAssetDataTagMap&& NewTags);
 
