[AssetRegistry]
- Fix potential race between default class object creation and asset registry
- Add a test case to stress test the race
- Make sure TSAN properly understands the barriers and that RaceDetector and TSAN are both satisfied

diff --git a/Engine/Source/Runtime/AssetRegistry/Private/AssetRegistry.cpp b/Engine/Source/Runtime/AssetRegistry/Private/AssetRegistry.cpp
--- a/Engine/Source/Runtime/AssetRegistry/Private/AssetRegistry.cpp
+++ b/Engine/Source/Runtime/AssetRegistry/Private/AssetRegistry.cpp
@@ -6752,44 +6752,64 @@ bool FAssetRegistryImpl::TryPostLoadAssetRegistryTags(FAssetData* AssetData)
 
 	if (AssetClass)
 	{
+		bool bCompleted = false;
+
 		UObject* ClassDefaultObject = AssetClass->GetDefaultObject(false);
-		if (ClassDefaultObject && !ClassDefaultObject->HasAnyFlags(RF_NeedInitialization))
+		if (ClassDefaultObject)
 		{
-			// We are using RF_NeedInitialization to guarantee that ClassDefaultObject is fully initialized
-			// potentially on another thread. For weakly ordered memory platforms, we need to 
-			// ensure that our read of the vtable ptr isn't performed prior to the read of the class flags
-			// otherwise we might see a stale vtable despite seeing RF_NeedInit clear.
+			// The flag is only safe to read if we have a barrier between the time where the default object pointer is 
+			// published, and the time we read the flag. Otherwise, we could end up reading a stale flag that hasn't yet
+			// been set to RF_NeedInitialization. It needs to be paired with a release barrier when the default
+			// object is set.
 			std::atomic_thread_fence(std::memory_order_acquire);
-			TArray<UObject::FAssetRegistryTag> TagsToModify;
-			UObject::FPostLoadAssetRegistryTagsContext Context(*AssetData, TagsToModify);
-			ClassDefaultObject->ThreadedPostLoadAssetRegistryTags(Context);
-			if (TagsToModify.Num())
+			TSAN_AFTER(ClassDefaultObject); // TSAN doesn't understand fence yet, help it understand what's going on
+
+			if (!ClassDefaultObject->HasAnyFlags(RF_NeedInitialization))
 			{
-				FAssetDataTagMap TagsAndValues = AssetData->TagsAndValues.CopyMap();
-				for (const UObject::FAssetRegistryTag& Tag : TagsToModify)
+				// We are using RF_NeedInitialization to guarantee that ClassDefaultObject is fully initialized
+				// potentially on another thread. For weakly ordered memory platforms, we need to 
+				// ensure that our read of the vtable ptr isn't performed prior to the read of the class flags
+				// otherwise we might see a stale vtable despite seeing RF_NeedInit clear.
+				std::atomic_thread_fence(std::memory_order_acquire);
+				TSAN_AFTER(ClassDefaultObject); // TSAN doesn't understand fence yet, help it understand what's going on
+
+				TArray<UObject::FAssetRegistryTag> TagsToModify;
+				UObject::FPostLoadAssetRegistryTagsContext Context(*AssetData, TagsToModify);
+				ClassDefaultObject->ThreadedPostLoadAssetRegistryTags(Context);
+				if (TagsToModify.Num())
 				{
-					if (!Tag.Value.IsEmpty())
-					{
-						FString& Value = TagsAndValues.FindOrAdd(Tag.Name);
-						Value = Tag.Value;
-					}
-					else
+					FAssetDataTagMap TagsAndValues = AssetData->TagsAndValues.CopyMap();
+					for (const UObject::FAssetRegistryTag& Tag : TagsToModify)
 					{
-						TagsAndValues.Remove(Tag.Name);
+						if (!Tag.Value.IsEmpty())
+						{
+							FString& Value = TagsAndValues.FindOrAdd(Tag.Name);
+							Value = Tag.Value;
+						}
+						else
+						{
+							TagsAndValues.Remove(Tag.Name);
+						}
 					}
+					AssetData->TagsAndValues = FAssetDataTagMapSharedView(MoveTemp(TagsAndValues));
 				}
-				AssetData->TagsAndValues = FAssetDataTagMapSharedView(MoveTemp(TagsAndValues));
+
+				bCompleted = true;
 			}
 		}
-		else if (!bForceCompletionEvenIfPostLoadsFail)
-		{
-			bCouldPostLoadAssetRegistryTags = false;
-		}
-		else 
+
+		if (!bCompleted)
 		{
-			ensureMsgf(!bMakeFinalChecks,
-				TEXT("Unable to PostLoadAssetRegistryTags for '%s' because the CDO for ancestor class '%s' could not be found or was not ready."),
-				*AssetData->GetObjectPathString(), *AssetClassPath.ToString());
+			if (!bForceCompletionEvenIfPostLoadsFail)
+			{
+				bCouldPostLoadAssetRegistryTags = false;
+			}
+			else 
+			{
+				ensureMsgf(!bMakeFinalChecks,
+					TEXT("Unable to PostLoadAssetRegistryTags for '%s' because the CDO for ancestor class '%s' could not be found or was not ready."),
+					*AssetData->GetObjectPathString(), *AssetClassPath.ToString());
+			}
 		}
 	}
 	return bCouldPostLoadAssetRegistryTags;
diff --git a/Engine/Source/Runtime/CoreUObject/Private/UObject/Class.cpp b/Engine/Source/Runtime/CoreUObject/Private/UObject/Class.cpp
--- a/Engine/Source/Runtime/CoreUObject/Private/UObject/Class.cpp
+++ b/Engine/Source/Runtime/CoreUObject/Private/UObject/Class.cpp
@@ -5000,7 +5000,17 @@ PRAGMA_DISABLE_DEPRECATION_WARNINGS
 			{
 				// RF_ArchetypeObject flag is often redundant to RF_ClassDefaultObject, but we need to tag
 				// the CDO as RF_ArchetypeObject in order to propagate that flag to any default sub objects.
-				ClassDefaultObject = StaticAllocateObject(this, GetOuter(), NAME_None, EObjectFlags(RF_Public|RF_ClassDefaultObject|RF_ArchetypeObject));
+				UObject* NewClassDefaultObject = StaticAllocateObject(this, GetOuter(), NAME_None, EObjectFlags(RF_Public | RF_ClassDefaultObject | RF_ArchetypeObject));
+				// Make sure all memory writes that happened during static allocate objects are visible
+				// to a thread that would try to access it once stored in ClassDefaultObject.
+				// This is particularly important for asset registry that can try to read RF_NeedInitialization
+				// do determine if the class is ready.
+				std::atomic_thread_fence(std::memory_order_release);
+				// TSAN doesn't recognize fence yet, so use annotations. 
+				// Any pointer would do as the fence is global, but needs to be paired with TSAN_AFTER
+				// on the same pointer before consuming any object properties.
+				TSAN_BEFORE(NewClassDefaultObject);
+				ClassDefaultObject = NewClassDefaultObject;
 				check(ClassDefaultObject);
 				// Register the offsets of any sparse delegates this class introduces with the sparse delegate storage
 				for (TFieldIterator<FMulticastSparseDelegateProperty> SparseDelegateIt(this, EFieldIteratorFlags::ExcludeSuper, EFieldIteratorFlags::ExcludeDeprecated); SparseDelegateIt; ++SparseDelegateIt)
diff --git a/Engine/Source/Runtime/CoreUObject/Private/UObject/UObjectGlobals.cpp b/Engine/Source/Runtime/CoreUObject/Private/UObject/UObjectGlobals.cpp
--- a/Engine/Source/Runtime/CoreUObject/Private/UObject/UObjectGlobals.cpp
+++ b/Engine/Source/Runtime/CoreUObject/Private/UObject/UObjectGlobals.cpp
@@ -4400,6 +4400,10 @@ void FObjectInitializer::PostConstructInit()
 	// need a fence to guarantee that the cleared flag is only visible to other threads
 	// after the other initialization-related writes
 	std::atomic_thread_fence(std::memory_order_release);
+	// TSAN doesn't recognize fence yet, so use annotations. 
+	// Any pointer would do as the fence is global, but needs to be paired with TSAN_AFTER
+	// on the same pointer before consuming any object properties.
+	TSAN_BEFORE(Obj); 
 	Obj->ClearFlags(RF_NeedInitialization);
 
 	// clear the object pointer so we can guard against running this function again
diff --git a/Engine/Source/Runtime/CoreUObject/Public/UObject/ObjectHandle.h b/Engine/Source/Runtime/CoreUObject/Public/UObject/ObjectHandle.h
--- a/Engine/Source/Runtime/CoreUObject/Public/UObject/ObjectHandle.h
+++ b/Engine/Source/Runtime/CoreUObject/Public/UObject/ObjectHandle.h
@@ -60,7 +60,12 @@ namespace UE::CoreUObject::Private
 	struct FObjectHandlePrivate
 	{
 		//Stores either FPackedObjectRef or a UObject*
-		UPTRINT PointerOrRef;
+		TSAN_ATOMIC(UPTRINT) PointerOrRef;
+
+		FORCEINLINE UPTRINT GetPointerOrRef() const
+		{
+			return PointerOrRef;
+		}
 
 		explicit inline operator bool() const
 		{
@@ -271,7 +276,7 @@ inline uint32 GetTypeHash(UE::CoreUObject::Private::FObjectHandlePrivate Handle)
 		}
 		return GetTypeHash(PackedObjectRef.EncodedRef);
 	}
-	return GetTypeHash(Handle.PointerOrRef);
+	return GetTypeHash(Handle.GetPointerOrRef());
 }
 #elif UE_WITH_REMOTE_OBJECT_HANDLE
 inline bool operator==(UE::CoreUObject::Private::FRemoteObjectHandlePrivate LHS, UE::CoreUObject::Private::FRemoteObjectHandlePrivate RHS)
@@ -517,7 +522,7 @@ namespace UE::CoreUObject::Private
 	inline UObject* ReadObjectHandlePointerNoCheck(FObjectHandle Handle)
 	{
 #if UE_WITH_OBJECT_HANDLE_LATE_RESOLVE
-		return reinterpret_cast<UObject*>(Handle.PointerOrRef);
+		return reinterpret_cast<UObject*>(Handle.GetPointerOrRef());
 #elif UE_WITH_REMOTE_OBJECT_HANDLE
 		return reinterpret_cast<UObject*>(Handle.PointerOrHandle);
 #else
diff --git a/Engine/Source/Runtime/CoreUObject/Tests/ClassDefaultObjectTest.cpp b/Engine/Source/Runtime/CoreUObject/Tests/ClassDefaultObjectTest.cpp
new file mode 100644
--- /dev/null
+++ b/Engine/Source/Runtime/CoreUObject/Tests/ClassDefaultObjectTest.cpp
@@ -0,0 +1,77 @@
+// Copyright Epic Games, Inc. All Rights Reserved.
+
+#if WITH_LOW_LEVEL_TESTS
+
+#include "ClassDefaultObjectTest.h"
+#include "TestHarness.h"
+#include "TestMacros/Assertions.h"
+#include "Async/ManualResetEvent.h"
+#include "Tasks/Task.h"
+#include "UObject/Package.h"
+
+TEST_CASE("UE::CoreUObject::GetDefaultObjectRace", "[CoreUObject][GetDefaultObjectRace]")
+{
+	for (int32 Index = 0; Index < 1000; ++Index)
+	{
+		// Move old object out of the way and force recreation
+		if (UClassDefaultObjectTest* Obj = (UClassDefaultObjectTest*)UClassDefaultObjectTest::StaticClass()->GetDefaultObject(false))
+		{
+			Obj->Rename(nullptr, GetTransientPackage());
+		}
+		UClassDefaultObjectTest::StaticClass()->SetDefaultObject(nullptr);
+
+		std::atomic<bool> bContinue = true;
+		UE::FManualResetEvent StartedEvent;
+		int32 NumPostInitProperties = 0;
+
+		// One worker thread tries to query the default object if its ready.
+		UE::Tasks::FTask Task =
+			UE::Tasks::Launch(
+				TEXT("Task"),
+				[&]()
+				{
+					StartedEvent.Notify();
+					while (bContinue.load(std::memory_order_relaxed))
+					{
+						UClassDefaultObjectTest* Obj = (UClassDefaultObjectTest*)UClassDefaultObjectTest::StaticClass()->GetDefaultObject(false);
+						if (Obj)
+						{
+							// The flag is only safe to read if we have a barrier between the time where the default object pointer is 
+							// published, and the time we read the flag. It needs to be paired with a release barrier when the default
+							// object is set.
+							std::atomic_thread_fence(std::memory_order_acquire);
+							TSAN_AFTER(Obj); // TSAN doesn't understand fence, help it understand what's going on
+							FPlatformProcess::YieldCycles(FMath::RandRange(0, 4000));
+							if (!Obj->HasAnyFlags(RF_NeedInitialization))
+							{
+								FPlatformProcess::YieldCycles(FMath::RandRange(0, 4000));
+								// This value is only safe to read if RF_NeedInitialization has been removed
+								// and that we applied the required barrier.
+								std::atomic_thread_fence(std::memory_order_acquire);
+								TSAN_AFTER(Obj); // TSAN doesn't understand fence, help it understand what's going on
+								NumPostInitProperties = Obj->NumPostInitProperties;
+								return;
+							}
+						}
+					}
+				}
+			);
+
+		// Make sure the async task has started
+		CHECK(StartedEvent.WaitFor(UE::FMonotonicTimeSpan::FromSeconds(1)));
+
+		CHECK(UClassDefaultObjectTest::StaticClass()->GetDefaultObject(false) == nullptr);
+
+		// Main thread creates the default object.
+		UClassDefaultObjectTest::StaticClass()->GetDefaultObject(true);
+
+		CHECK(Task.Wait(FTimespan::FromSeconds(1.0f)));
+
+		// Store the loop to kill the task if it failed to return on time
+		bContinue = false;
+		CHECK(Task.Wait(FTimespan::FromSeconds(1.0f)));
+		CHECK(NumPostInitProperties != 0);
+	}
+}
+
+#endif // WITH_LOW_LEVEL_TESTS
diff --git a/Engine/Source/Runtime/CoreUObject/Tests/ClassDefaultObjectTest.h b/Engine/Source/Runtime/CoreUObject/Tests/ClassDefaultObjectTest.h
new file mode 100644
--- /dev/null
+++ b/Engine/Source/Runtime/CoreUObject/Tests/ClassDefaultObjectTest.h
@@ -0,0 +1,30 @@
+// Copyright Epic Games, Inc. All Rights Reserved.
+
+#pragma once
+
+#if WITH_TESTS
+
+#include "UObject/Object.h"
+
+#include "ClassDefaultObjectTest.generated.h"
+
+UCLASS()
+class UClassDefaultObjectTest : public UObject
+{
+	GENERATED_BODY()
+
+public:
+	int32 NumPostInitProperties = 0;
+
+	// PostInitProperties is called before RF_NeedInitialization is removed
+	// so its a good place to modify a non-atomic value to see if there is 
+	// any race during default object construction.
+	virtual void PostInitProperties()
+	{
+		Super::PostInitProperties();
+
+		NumPostInitProperties++;
+	}
+};
+
+#endif
\ No newline at end of file
