AssetManager,AssetRegistry: Add Property AffectsCookRule to ManageReferences.
[RN] Minor, Foundation

diff --git a/Engine/Source/Runtime/AssetRegistry/Private/AssetRegistry.cpp b/Engine/Source/Runtime/AssetRegistry/Private/AssetRegistry.cpp
--- a/Engine/Source/Runtime/AssetRegistry/Private/AssetRegistry.cpp
+++ b/Engine/Source/Runtime/AssetRegistry/Private/AssetRegistry.cpp
@@ -9506,9 +9506,10 @@ void FAssetRegistryImpl::SetManageReferences(FSetManageReferencesContext& Contex
 				TargetNodeData.bModifiedByCurrentManager = true;
 				CurrentManagerModifiedNodes.Add(TargetNode);
 
-				EDependencyProperty ManageProperties = 
-				TargetNodeData.CurrentManagerProperties = (EdgeFlags & EAssetSetManagerFlags::IsDirectSet)
-					? EDependencyProperty::Direct : EDependencyProperty::None;
+				TargetNodeData.CurrentManagerProperties =
+					((EdgeFlags & EAssetSetManagerFlags::IsDirectSet)
+						? EDependencyProperty::Direct : EDependencyProperty::None)
+					| EDependencyProperty::CookRule;
 				TargetNodeData.DebugInstigator = SourceNode;
 
 				if (Result == EAssetSetManagerResult::SetAndRecurse)
@@ -9518,7 +9519,8 @@ void FAssetRegistryImpl::SetManageReferences(FSetManageReferencesContext& Contex
 			};
 
 			// Check initial node
-			IterateFunction(ManagerNode, BaseManagedNode, EDependencyCategory::Manage, EDependencyProperty::Direct);
+			IterateFunction(ManagerNode, BaseManagedNode, EDependencyCategory::Manage,
+				EDependencyProperty::Game | EDependencyProperty::Build);
 
 			// Do all recursion first, but only if we have a recurse type
 			if (Context.RecurseType != EDependencyCategory::None)
@@ -9670,7 +9672,7 @@ bool FAssetRegistryDependencyOptions::GetManageQuery(UE::AssetRegistry::FDepende
 	if (bIncludeSoftManagementReferences || bIncludeHardManagementReferences)
 	{
 		if (!bIncludeSoftManagementReferences) Flags.Required |= UE::AssetRegistry::EDependencyProperty::Direct;
-		if (!bIncludeHardPackageReferences) Flags.Excluded|= UE::AssetRegistry::EDependencyProperty::Direct;
+		if (!bIncludeHardManagementReferences) Flags.Excluded |= UE::AssetRegistry::EDependencyProperty::Direct;
 		return true;
 	}
 	return false;
diff --git a/Engine/Source/Runtime/AssetRegistry/Private/AssetRegistryState.cpp b/Engine/Source/Runtime/AssetRegistry/Private/AssetRegistryState.cpp
--- a/Engine/Source/Runtime/AssetRegistry/Private/AssetRegistryState.cpp
+++ b/Engine/Source/Runtime/AssetRegistry/Private/AssetRegistryState.cpp
@@ -3716,6 +3716,7 @@ void PropertiesToString(EDependencyProperty Properties, FStringBuilderBase& Buil
 	if (!!(CategoryFilter & EDependencyCategory::Manage))
 	{
 		AppendPropertyName(EDependencyProperty::Direct, TEXT("Direct"), TEXT("Indirect"));
+		AppendPropertyName(EDependencyProperty::CookRule, TEXT("CookRule"), TEXT("ChunkOnly"));
 	}
 	static_assert((EDependencyProperty::PackageMask
 		| EDependencyProperty::SearchableNameMask
diff --git a/Engine/Source/Runtime/AssetRegistry/Private/DependsNode.cpp b/Engine/Source/Runtime/AssetRegistry/Private/DependsNode.cpp
--- a/Engine/Source/Runtime/AssetRegistry/Private/DependsNode.cpp
+++ b/Engine/Source/Runtime/AssetRegistry/Private/DependsNode.cpp
@@ -995,7 +995,16 @@ void FDependsNode::SerializeLoad_ManageDependencies(FArchive& Ar,
 	TArray<FDependsNode*>& OutDependencies = this->ManageDependencies;
 	TBitArray<>& OutFlagBits = this->ManageFlags;
 	int OutFlagSetWidth = ManageFlagSetWidth;
-	int InFlagSetWidth = OutFlagSetWidth;
+	int InFlagSetWidth;
+	if (bLatestVersion || Version >= FAssetRegistryVersion::ManageDependenciesCookRule)
+	{
+		InFlagSetWidth = OutFlagSetWidth;
+	}
+	else
+	{
+		// Dependencies before ManageDependenciesCookRule were a single bit.
+		InFlagSetWidth = 1;
+	}
 
 	TArray<int32>& InDependencies = Scratch.InDependencies;
 	TArray<uint32>& InFlagBits = Scratch.InFlagBits;
@@ -1043,13 +1052,44 @@ void FDependsNode::SerializeLoad_ManageDependencies(FArchive& Ar,
 		OutDependencies.Add(PointerDependencies[SortIndex]);
 	}
 
-	OutFlagBits.SetNumUninitialized(InNumFlagBits);
-	uint32* InFlagBitsData = InFlagBits.GetData();
-	for (int32 WriteIndex = 0; WriteIndex < NumDependencies; ++WriteIndex)
+	if (bLatestVersion || Version >= FAssetRegistryVersion::ManageDependenciesCookRule)
 	{
-		int32 ReadIndex = SortIndexes[WriteIndex];
-		OutFlagBits.SetRangeFromRange(WriteIndex * OutFlagSetWidth, OutFlagSetWidth,
-			InFlagBitsData, ReadIndex * OutFlagSetWidth);
+		OutFlagBits.SetNumUninitialized(InNumFlagBits);
+		uint32* InFlagBitsData = InFlagBits.GetData();
+		for (int32 WriteIndex = 0; WriteIndex < NumDependencies; ++WriteIndex)
+		{
+			int32 ReadIndex = SortIndexes[WriteIndex];
+			OutFlagBits.SetRangeFromRange(WriteIndex * OutFlagSetWidth, OutFlagSetWidth,
+				InFlagBitsData, ReadIndex * OutFlagSetWidth);
+		}
+	}
+	else
+	{
+		// Manage dependency sets before ManageDependenciesCookRule were lists of exactly 1 property value, with a
+		// bit specifying EDependencyProperty::Direct or EDependencyProperty::None. They were all implicitly
+		// EDependencyProperty::CookRule. We stored each as a single bit rather than storing them as a
+		// TPropertyCombinationSet.
+		// So to store the upgraded version as a TPropertyCombinationSet, upgrade the original bit into a single
+		// EDependencyProperty, and add that single property value to the output TPropertyCombinationSet.
+		TBitArray<> InFlagBitsAsBitArray;
+		InFlagBitsAsBitArray.SetNumUninitialized(InNumFlagBits);
+		InFlagBitsAsBitArray.SetRangeFromRange(0, InNumFlagBits, InFlagBits.GetData(), 0);
+
+		int32 OutNumFlagBits = OutFlagSetWidth * NumDependencies;
+		OutFlagBits.SetNumUninitialized(OutNumFlagBits);
+		for (int32 WriteIndex = 0; WriteIndex < NumDependencies; ++WriteIndex)
+		{
+			int32 ReadIndex = SortIndexes[WriteIndex];
+			bool bInputIsDirect = InFlagBitsAsBitArray[ReadIndex*InFlagSetWidth];
+			EDependencyProperty UpgradedInputProperty =
+				(bInputIsDirect ? EDependencyProperty::Direct : EDependencyProperty::None)
+				| EDependencyProperty::CookRule;
+			FManageFlagSet UpgradedInputFlagSet;
+			UpgradedInputFlagSet.Add(ManagePropertiesToByte(UpgradedInputProperty));
+			
+			// Write the bits of the upgraded TPropertyCombinationSet into the properly sorted location in OutFlagBits.
+			UpgradedInputFlagSet.Save(OutFlagBits, WriteIndex * OutFlagSetWidth);
+		}
 	}
 }
 
@@ -1209,15 +1249,24 @@ void FDependsNode::LegacySerializeLoad_BeforeFlags(FArchive& Ar,
 	}
 	if (NumSoftManage > 0)
 	{
-		// In legacy ManageDependencies, Soft meant Indirect, which we currently record as 0x0.
-		uint32 SoftManageBits = 0x0;
+		// All legacy ManageDependencies were CookRule. Soft meant Indirect.
+		uint32 SoftManageBits = 0;
+		FDependsNode::FManageFlagSet FlagSet;
+		FlagSet.Add(ManagePropertiesToByte(
+			UE::AssetRegistry::EDependencyProperty::CookRule));
+		FlagSet.Save(&SoftManageBits);
 		SerializeNodeArray(NumSoftManage, ManageDependencies, &ManageFlags, ManageFlagSetWidth, SoftManageBits,
 			true, bSerializeDependencies);
 	}
 	if (NumHardManage > 0)
 	{
-		// In legacy ManageDependencies, Hard meant Direct, which we currently record as 0x1.
-		uint32 HardManageBits = 0x1;
+		// All legacy ManageDependencies were CookRule. Hard meant Direct.
+		uint32 HardManageBits = 0;
+		FDependsNode::FManageFlagSet FlagSet;
+		FlagSet.Add(ManagePropertiesToByte(
+			UE::AssetRegistry::EDependencyProperty::Direct
+			| UE::AssetRegistry::EDependencyProperty::CookRule));
+		FlagSet.Save(&HardManageBits);
 		SerializeNodeArray(NumHardManage, ManageDependencies, &ManageFlags, ManageFlagSetWidth, HardManageBits,
 			true, bSerializeDependencies);
 	}
diff --git a/Engine/Source/Runtime/AssetRegistry/Private/DependsNode.h b/Engine/Source/Runtime/AssetRegistry/Private/DependsNode.h
--- a/Engine/Source/Runtime/AssetRegistry/Private/DependsNode.h
+++ b/Engine/Source/Runtime/AssetRegistry/Private/DependsNode.h
@@ -18,11 +18,12 @@ public:
 	typedef TArray<FDependsNode*> FDependsNodeList;
 	static constexpr uint32 PackageFlagWidth = 3;
 	static constexpr uint32 SearchableNameFlagWidth = 0;
-	static constexpr uint32 ManageFlagWidth = 1;
+	static constexpr uint32 ManageFlagWidth = 2;
 	typedef TPropertyCombinationSet<PackageFlagWidth> FPackageFlagSet;
 	static constexpr uint32 PackageFlagSetWidth = FPackageFlagSet::StorageBitCount;
 	static constexpr uint32 SearchableNameFlagSetWidth = 0;
-	static constexpr uint32 ManageFlagSetWidth = TPropertyCombinationSet<ManageFlagWidth>::StorageBitCount;
+	typedef TPropertyCombinationSet<ManageFlagWidth> FManageFlagSet;
+	static constexpr uint32 ManageFlagSetWidth = FManageFlagSet::StorageBitCount;
 
 public:
 	FDependsNode();
@@ -353,13 +354,15 @@ inline UE::AssetRegistry::EDependencyProperty FDependsNode::ByteToPackagePropert
 }
 inline uint8 FDependsNode::ManagePropertiesToByte(UE::AssetRegistry::EDependencyProperty Properties)
 {
-	return (0x01 * (static_cast<uint8>(Properties & UE::AssetRegistry::EDependencyProperty::Direct) != 0));
+	return (0x01 * (static_cast<uint8>(Properties & UE::AssetRegistry::EDependencyProperty::Direct) != 0))
+		| (0x02 * (static_cast<uint8>(Properties & UE::AssetRegistry::EDependencyProperty::CookRule) != 0));
 }
 
 inline UE::AssetRegistry::EDependencyProperty FDependsNode::ByteToManageProperties(uint8 Bits)
 {
 	return static_cast<UE::AssetRegistry::EDependencyProperty>(
 		(static_cast<uint32>(UE::AssetRegistry::EDependencyProperty::Direct) * ((Bits & 0x01) != 0))
+		| (static_cast<uint32>(UE::AssetRegistry::EDependencyProperty::CookRule) * ((Bits & 0x02) != 0))
 		);
 }
 
diff --git a/Engine/Source/Runtime/AssetRegistry/Private/DependsNodeTest.cpp b/Engine/Source/Runtime/AssetRegistry/Private/DependsNodeTest.cpp
--- a/Engine/Source/Runtime/AssetRegistry/Private/DependsNodeTest.cpp
+++ b/Engine/Source/Runtime/AssetRegistry/Private/DependsNodeTest.cpp
@@ -153,20 +153,6 @@ bool FDependsNodeTest::RunTest(const FString& Parameters)
 		FName IdentifierA(TEXT("IdentifierA"));
 		A.SetIdentifier(IdentifierA);
 		FDependsNode* PointerA = &A;
-		int NumScratch = 6;
-		if (SortTestType == ESortTestType::AlwaysUnsorted || SortTestType == ESortTestType::UnsortedThenSorted)
-		{
-			A.SetIsDependencyListSorted(EDependencyCategory::All, false);
-			for (int n = 0; n < NumScratch; ++n)
-			{
-				ScratchNodes[n].SetIsReferencersSorted(false);
-			}
-		}
-		for (EDependencyCategory Category : { EDependencyCategory::Package, EDependencyCategory::SearchableName, EDependencyCategory::Manage})
-		{
-			TestEqual(TEXT("IsDependencyListSorted"), A.IsDependencyListSorted(Category), SortTestType == ESortTestType::AlwaysSorted);
-		}
-		TestEqual(TEXT("IsReferencersSorted"), ScratchNodes[0].IsReferencersSorted(), SortTestType == ESortTestType::AlwaysSorted);
 
 		struct FTestDependency
 		{
@@ -194,15 +180,33 @@ bool FDependsNodeTest::RunTest(const FString& Parameters)
 			{&ScratchNodes[1], EDependencyCategory::Package, EDependencyProperty::Hard},
 			{&ScratchNodes[2], EDependencyCategory::Package, EDependencyProperty::None},
 			{&ScratchNodes[3], EDependencyCategory::SearchableName, EDependencyProperty::None},
-			{&ScratchNodes[4], EDependencyCategory::Manage, EDependencyProperty::Direct},
-			{&ScratchNodes[5], EDependencyCategory::Manage, EDependencyProperty::None }
+			{&ScratchNodes[4], EDependencyCategory::Manage, EDependencyProperty::Direct | EDependencyProperty::CookRule},
+			{&ScratchNodes[5], EDependencyCategory::Manage, EDependencyProperty::Direct},
+			{&ScratchNodes[6], EDependencyCategory::Manage, EDependencyProperty::None }
 		};
 
+		constexpr int NumUsedScratchNodes = UE_ARRAY_COUNT(TestDependencies);
+		static_assert(NumUsedScratchNodes <= UE_ARRAY_COUNT(ScratchNodes));
+		if (SortTestType == ESortTestType::AlwaysUnsorted || SortTestType == ESortTestType::UnsortedThenSorted)
+		{
+			A.SetIsDependencyListSorted(EDependencyCategory::All, false);
+			for (int n = 0; n < NumUsedScratchNodes; ++n)
+			{
+				ScratchNodes[n].SetIsReferencersSorted(false);
+			}
+		}
+		for (EDependencyCategory Category : { EDependencyCategory::Package, EDependencyCategory::SearchableName, EDependencyCategory::Manage})
+		{
+			TestEqual(TEXT("IsDependencyListSorted"), A.IsDependencyListSorted(Category), SortTestType == ESortTestType::AlwaysSorted);
+		}
+		TestEqual(TEXT("IsReferencersSorted"), ScratchNodes[0].IsReferencersSorted(), SortTestType == ESortTestType::AlwaysSorted);
+
+
 		for (FTestDependency& TestDependency : TestDependencies)
 		{
 			A.AddDependency(TestDependency.Node, TestDependency.Category, TestDependency.Properties);
 		}
-		for (int n = 0; n < NumScratch; ++n)
+		for (int n = 0; n < NumUsedScratchNodes; ++n)
 		{
 			ScratchNodes[n].AddReferencer(&A);
 		}
@@ -210,7 +214,7 @@ bool FDependsNodeTest::RunTest(const FString& Parameters)
 		if (SortTestType == ESortTestType::UnsortedThenSorted)
 		{
 			A.SetIsDependencyListSorted(UE::AssetRegistry::EDependencyCategory::All, true);
-			for (int n = 0; n < NumScratch; ++n)
+			for (int n = 0; n < NumUsedScratchNodes; ++n)
 			{
 				ScratchNodes[n].SetIsReferencersSorted(true);
 			}
@@ -446,8 +450,8 @@ bool FDependsNodeTest::RunTest(const FString& Parameters)
 		}
 		// Manage All
 		{
-			FTestDependency MatchingDeps[] = { TestDependencies[4], TestDependencies[5] };
-			FDependsNode* Matching[] = { &ScratchNodes[4], &ScratchNodes[5] };
+			FTestDependency MatchingDeps[] = { TestDependencies[4], TestDependencies[5], TestDependencies[6] };
+			FDependsNode* Matching[] = { &ScratchNodes[4], &ScratchNodes[5], &ScratchNodes[6] };
 			Dependencies.Reset();
 			A.GetDependencies(Dependencies, EDependencyCategory::Manage);
 			TestTrue(TEXT("Manage dependency with no flags"), CombinationEquals<FDependsNode*>(Dependencies, Matching));
@@ -479,9 +483,13 @@ bool FDependsNodeTest::RunTest(const FString& Parameters)
 					++IterateCount;
 					if (Dependency == &ScratchNodes[4])
 					{
-						TestTrue(TEXT("IterateDependencies Manage, no flags"), Properties == EDependencyProperty::Direct && !bDuplicate);
+						TestTrue(TEXT("IterateDependencies Manage, no flags"), Properties == (EDependencyProperty::Direct | EDependencyProperty::CookRule) && !bDuplicate);
 					}
 					else if (Dependency == &ScratchNodes[5])
+					{
+						TestTrue(TEXT("IterateDependencies Manage, no flags"), Properties == EDependencyProperty::Direct && !bDuplicate);
+					}
+					else if (Dependency == &ScratchNodes[6])
 					{
 						TestTrue(TEXT("IterateDependencies Manage, no flags"), Properties == EDependencyProperty::None && !bDuplicate);
 					}
@@ -490,13 +498,13 @@ bool FDependsNodeTest::RunTest(const FString& Parameters)
 						AddError(TEXT("IterateDependencies Manage, no flags"));
 					}
 				}, EDependencyCategory::Manage);
-			TestEqual(TEXT("IterateDependencies Manage, no flags, Count"), IterateCount, 2);
+			TestEqual(TEXT("IterateDependencies Manage, no flags, Count"), IterateCount, 3);
 		}
 		// Manage Direct
 		{
-			FTestDependency MatchingDeps[] = { TestDependencies[4] };
-			FDependsNode* Matching[] = { &ScratchNodes[4] };
-			FDependsNode* NonMatching[] = { &ScratchNodes[5] };
+			FTestDependency MatchingDeps[] = { TestDependencies[4], TestDependencies[5] };
+			FDependsNode* Matching[] = { &ScratchNodes[4], &ScratchNodes[5] };
+			FDependsNode* NonMatching[] = { &ScratchNodes[6] };
 			Dependencies.Reset();
 			A.GetDependencies(Dependencies, EDependencyCategory::Manage, EDependencyQuery::Direct);
 			TestTrue(TEXT("Manage dependency with flags, Direct"), CombinationEquals<FDependsNode*>(Dependencies, Matching));
@@ -543,6 +551,10 @@ bool FDependsNodeTest::RunTest(const FString& Parameters)
 					TestTrue(TEXT("IterateDependencies Manage, Direct, Category"), Category == EDependencyCategory::Manage);
 					++IterateCount;
 					if (Dependency == &ScratchNodes[4])
+					{
+						TestTrue(TEXT("IterateDependencies Manage, Direct"), Properties == (EDependencyProperty::Direct | EDependencyProperty::CookRule) && !bDuplicate);
+					}
+					else if (Dependency == &ScratchNodes[5])
 					{
 						TestTrue(TEXT("IterateDependencies Manage, Direct"), Properties == EDependencyProperty::Direct && !bDuplicate);
 					}
@@ -551,13 +563,13 @@ bool FDependsNodeTest::RunTest(const FString& Parameters)
 						AddError(TEXT("IterateDependencies Manage, Direct"));
 					}
 				}, EDependencyCategory::Manage, EDependencyQuery::Direct);
-			TestEqual(TEXT("IterateDependencies Manage, Direct, Count"), IterateCount, 1);
+			TestEqual(TEXT("IterateDependencies Manage, Direct, Count"), IterateCount, 2);
 		}
 		// Manage Indirect
 		{
-			FTestDependency MatchingDeps[] = { TestDependencies[5] };
-			FDependsNode* Matching[] = { &ScratchNodes[5] };
-			FDependsNode* NonMatching[] = { &ScratchNodes[4] };
+			FTestDependency MatchingDeps[] = { TestDependencies[6] };
+			FDependsNode* Matching[] = { &ScratchNodes[6] };
+			FDependsNode* NonMatching[] = { &ScratchNodes[4], &ScratchNodes[5] };
 			Dependencies.Reset();
 			A.GetDependencies(Dependencies, EDependencyCategory::Manage, EDependencyQuery::Indirect);
 			TestTrue(TEXT("Manage dependency with flags, Indirect"), CombinationEquals<FDependsNode*>(Dependencies, Matching));
@@ -603,7 +615,7 @@ bool FDependsNodeTest::RunTest(const FString& Parameters)
 				{
 					TestTrue(TEXT("IterateDependencies Manage, Indirect, Category"), Category == EDependencyCategory::Manage);
 					++IterateCount;
-					if (Dependency == &ScratchNodes[5])
+					if (Dependency == &ScratchNodes[6])
 					{
 						TestTrue(TEXT("IterateDependencies Manage, Indirect"), Properties == EDependencyProperty::None && !bDuplicate);
 					}
@@ -615,13 +627,13 @@ bool FDependsNodeTest::RunTest(const FString& Parameters)
 			TestEqual(TEXT("IterateDependencies Manage, Indirect, Count"), IterateCount, 1);
 		}
 
-		TestEqual(TEXT("RemoveDependency in dependency and referencer filtering loop"), A.GetConnectionCount(), NumScratch);
-		for (int n = 0; n < NumScratch; ++n)
+		TestEqual(TEXT("RemoveDependency in dependency and referencer filtering loop"), A.GetConnectionCount(), NumUsedScratchNodes);
+		for (int n = 0; n < NumUsedScratchNodes; ++n)
 		{
 			ScratchNodes[n].RemoveReferencer(&A);
 			ScratchNodes[n].SetIsReferencersSorted(true);
 			A.RemoveDependency(&ScratchNodes[n]);
-			TestEqual(TEXT("RemoveDependency in dependency and referencer filtering loop"), A.GetConnectionCount(), NumScratch - 1 - n);
+			TestEqual(TEXT("RemoveDependency in dependency and referencer filtering loop"), A.GetConnectionCount(), NumUsedScratchNodes - 1 - n);
 		}
 	}
 
@@ -939,6 +951,7 @@ bool FDependsNodeTest::RunTest(const FString& Parameters)
 			PackageProperties[n] = (n & 0x2) ? EDependencyProperty::Game : EDependencyProperty::None;
 			PackageProperties[n] = (n & 0x4) ? EDependencyProperty::Build : EDependencyProperty::None;
 			ManageProperties[n] = (n & 0x1) ? EDependencyProperty::Direct : EDependencyProperty::None;
+			ManageProperties[n] = (n & 0x2) ? EDependencyProperty::CookRule : EDependencyProperty::None;
 		}
 
 		for (ESortTestType SortTestType : { ESortTestType::AlwaysSorted, ESortTestType::UnsortedThenSorted, ESortTestType::AlwaysUnsorted})
diff --git a/Engine/Source/Runtime/CoreUObject/Public/AssetRegistry/AssetData.h b/Engine/Source/Runtime/CoreUObject/Public/AssetRegistry/AssetData.h
--- a/Engine/Source/Runtime/CoreUObject/Public/AssetRegistry/AssetData.h
+++ b/Engine/Source/Runtime/CoreUObject/Public/AssetRegistry/AssetData.h
@@ -102,6 +102,7 @@ struct FAssetRegistryVersion
 		MarshalledTextAsUTF8String,			// Replaced 2 byte wide string with UTF8 String
 		PackageSavedHash,					// Replaced FAssetPackageData::PackageGuid with PackageSavedHash
 		ExternalActorToWorldIsEditorOnly,   // FPackageDependencyData::LoadDependenciesFromPackageHeader changed how it calculates PackageDependencies
+		ManageDependenciesCookRule,			// CookRule property added to EDependencyProperties for EDependencyCategory::Manage
 
 		// -----<new versions can be added above this line>-------------------------------------------------
 		VersionPlusOne,
diff --git a/Engine/Source/Runtime/CoreUObject/Public/Misc/AssetRegistryInterface.h b/Engine/Source/Runtime/CoreUObject/Public/Misc/AssetRegistryInterface.h
--- a/Engine/Source/Runtime/CoreUObject/Public/Misc/AssetRegistryInterface.h
+++ b/Engine/Source/Runtime/CoreUObject/Public/Misc/AssetRegistryInterface.h
@@ -91,17 +91,43 @@ namespace UE::AssetRegistry
 		// If any values change or are removed, bump FAssetRegistryVersion and add backwards compatibility in FAssetRegistryState::Serialize.
 
 		// Package Dependencies
-		PackageMask = 0x7,
-		Hard = 0x1,			// The target asset must be loaded before the source asset can finish loading. The lack of this property is known as a Soft dependency, and indicates only that the source asset expects the target asset to be loadable on demand. 
-		Game = 0x2,			// The target asset is needed in the game as well as the editor. The lack of this property is known as an EditorOnly dependency.
-		Build = 0x4,		// Fields on the target asset are used in the transformation of the source asset during cooking in addition to being required in the game or editor. The lack of this property indicates that the target asset is required in game or editor, but is not required during cooking.
+		/**
+		 * The target asset must be loaded before the source asset can finish loading.The lack of this property is
+		 * known as a Soft dependency, and indicates only that the source asset expects the target asset to be
+		 * loadable on demand.
+		 */
+		Hard = 1 << 0,
+		/**
+		 * The target asset is needed in the game as well as the editor. The lack of this property is known as an
+		 * EditorOnly dependency.
+		 */
+		Game = 1 << 1,
+		/**
+		 * Fields on the target asset are used in the transformation of the source asset during cooking in addition
+		 * to being required in the game or editor. The lack of this property indicates that the target asset is
+		 * required in game or editor, but is not required during cooking.
+		 */
+		Build = 1 << 2,
+		PackageMask = Hard | Game | Build,
 
-		// SearchableName Dependencies
-		SearchableNameMask = 0x0, // None yet
+		// SearchableName Dependencies (None Yet)
+		SearchableNameMask = None,
 
 		// ManageDependencies
-		ManageMask = 0x8,
-		Direct = 0x8,		// The target asset was specified explicitly as a managee by the source asset. Lack of this property is known as an indirect dependency; the target asset is reachable by following the transitive closure of Direct Manage Dependencies and Package dependencies from the source asset.
+		/**
+		 * The target asset was specified explicitly as a managee by the source asset. Lack of this property is known 
+		 * as an indirect dependency; the target asset is reachable by following the transitive closure of Direct
+		 * Manage Dependencies and Package dependencies from the source asset.
+		 */
+		Direct = 1 << 3,
+		/**
+		 * The manager manages the target for the purposes of deciding its CookRule in addition to deciding its Chunk
+		 * assignment. Lack of this property is known as a ChunkOnly dependency. Every manage dependency used for
+		 * deciding the CookRule also decides the Chunk; Chunk assignment is a superset of CookRule assignment.
+		 */
+		CookRule = 1 << 4,
+		ManageMask = Direct | CookRule,
+
 
 		AllMask = PackageMask | SearchableNameMask | ManageMask,
 	};
@@ -118,32 +144,49 @@ namespace UE::AssetRegistry
 		NoRequirements = 0,
 
 		// Package Dependencies Only
-		Hard = 0x0001,			// Return only dependencies with EDependencyProperty::Hard
-		NotHard = 0x002,		// Return only dependencies without EDependencyProperty::Hard
+		/** Return only dependencies with EDependencyProperty::Hard. */
+		Hard = 1 << 0,
+		/** Return only dependencies without EDependencyProperty::Hard. */
+		NotHard = 1 << 1,
 		Soft = NotHard,
 
-		Game = 0x004,			// Return only dependencies with EDependencyProperty::Game
-		NotGame = 0x008,		// Return only dependencies without EDependencyProperty::Game
+		/** Return only dependencies with EDependencyProperty::Game. */
+		Game = 1 << 2,
+		/** Return only dependencies without EDependencyProperty::Game. */
+		NotGame = 1 << 3,
 		EditorOnly = NotGame,
 
-		Build = 0x010,			// Return only dependencies with EDependencyProperty::Build
-		NotBuild = 0x020,		// Return only dependencies without EDependencyProperty::Build
+		/** Return only dependencies with EDependencyProperty::Build. */
+		Build = 1 << 4,
+		/** Return only dependencies without EDependencyProperty::Build. */
+		NotBuild = 1 << 5,
 
-		Propagation = 0x040,    // Return only dependencies that cause propagation of manage dependencies, which means
-		// either Game or Build. Presence of the Propagation flag in a query causes the Game, NotGame, EditorOnly,
-		// Build, and NotBuild flags to be ignored in the query if present. Either Game or Build is sufficient to pass
-		// the query no matter which of those other flags are present. Hard vs Soft is still respected, and Soft Build
-		// or Soft Game dependencies will be skipped query if Hard is required.
+		/**
+		 * Return only dependencies that cause propagation of manage dependencies, which means
+		 * either Game or Build. Presence of the Propagation flag in a query causes the Game, NotGame, EditorOnly,
+		 * Build, and NotBuild flags to be ignored in the query if present. Either Game or Build is sufficient to pass
+		 * the query no matter which of those other flags are present. Hard vs Soft is still respected, and Soft Build
+		 * or Soft Game dependencies will be skipped query if Hard is required.
+		 */
+		Propagation = 1 << 6,
 
 		// Manage Dependencies Only
-		Direct = 0x0400,		// Return only dependencies with EDependencyProperty::Direct
-		NotDirect = 0x0800,		// Return only dependencies without EDependencyProperty::Direct
+		/** Return only dependencies with EDependencyProperty::Direct. */
+		Direct = 1 << 8,
+		/** Return only dependencies without EDependencyProperty::Direct. */
+		NotDirect = 1 << 9,
 		Indirect = NotDirect,
 
+		/** Return only dependencies with EDependencyProperty::CookRule. */
+		CookRule = 1 << 10,
+		/** Return only dependencies without EDependencyProperty::CookRule. */
+		NotCookRule = 1 << 11,
+		ChunkOnly = NotCookRule,
+
 		// Masks used for manipulating EDependencyQuerys
-		PackageMask = 0x00ff,
+		PackageMask = Hard | NotHard | Game | NotGame | Build | NotBuild,
 		SearchableNameMask = 0x0000, // None yet
-		ManageMask = 0x0f00,
+		ManageMask = Direct | NotDirect | CookRule | NotCookRule,
 	};
 	ENUM_CLASS_FLAGS(EDependencyQuery);
 
@@ -225,11 +268,14 @@ namespace UE::AssetRegistry
 				RequiredUnions.Add(UE::AssetRegistry::EDependencyProperty::Game | UE::AssetRegistry::EDependencyProperty::Build);
 			}
 			Required |= (!!(QueryFlags & EDependencyQuery::Hard) ? UE::AssetRegistry::EDependencyProperty::Hard : UE::AssetRegistry::EDependencyProperty::None)
-				| (!!(QueryFlags & EDependencyQuery::Direct) ? UE::AssetRegistry::EDependencyProperty::Direct : UE::AssetRegistry::EDependencyProperty::None);
+				| (!!(QueryFlags & EDependencyQuery::Direct) ? UE::AssetRegistry::EDependencyProperty::Direct : UE::AssetRegistry::EDependencyProperty::None)
+				| (!!(QueryFlags & EDependencyQuery::CookRule) ? UE::AssetRegistry::EDependencyProperty::CookRule : UE::AssetRegistry::EDependencyProperty::None);
+
 			Excluded = (!!(QueryFlags & EDependencyQuery::NotHard) ? UE::AssetRegistry::EDependencyProperty::Hard : UE::AssetRegistry::EDependencyProperty::None)
 				| (!!(QueryFlags & EDependencyQuery::NotGame) ? UE::AssetRegistry::EDependencyProperty::Game : UE::AssetRegistry::EDependencyProperty::None)
 				| (!!(QueryFlags & EDependencyQuery::NotBuild) ? UE::AssetRegistry::EDependencyProperty::Build : UE::AssetRegistry::EDependencyProperty::None)
-				| (!!(QueryFlags & EDependencyQuery::NotDirect) ? UE::AssetRegistry::EDependencyProperty::Direct : UE::AssetRegistry::EDependencyProperty::None);
+				| (!!(QueryFlags & EDependencyQuery::NotDirect) ? UE::AssetRegistry::EDependencyProperty::Direct : UE::AssetRegistry::EDependencyProperty::None)
+				| (!!(QueryFlags & EDependencyQuery::NotCookRule) ? UE::AssetRegistry::EDependencyProperty::CookRule : UE::AssetRegistry::EDependencyProperty::None);
 		}
 	};
 
