AssetRegistry: DependsNodeTest: Fix typo in test data found by static analysis tool.
#rnx

diff --git a/Engine/Source/Runtime/AssetRegistry/Private/DependsNodeTest.cpp b/Engine/Source/Runtime/AssetRegistry/Private/DependsNodeTest.cpp
--- a/Engine/Source/Runtime/AssetRegistry/Private/DependsNodeTest.cpp
+++ b/Engine/Source/Runtime/AssetRegistry/Private/DependsNodeTest.cpp
@@ -948,10 +948,10 @@ bool FDependsNodeTest::RunTest(const FString& Parameters)
 		for (int n = 0; n < NumLinkNodes; ++n)
 		{
 			PackageProperties[n] = (n & 0x1) ? EDependencyProperty::Hard : EDependencyProperty::None;
-			PackageProperties[n] = (n & 0x2) ? EDependencyProperty::Game : EDependencyProperty::None;
-			PackageProperties[n] = (n & 0x4) ? EDependencyProperty::Build : EDependencyProperty::None;
+			PackageProperties[n] |= (n & 0x2) ? EDependencyProperty::Game : EDependencyProperty::None;
+			PackageProperties[n] |= (n & 0x4) ? EDependencyProperty::Build : EDependencyProperty::None;
 			ManageProperties[n] = (n & 0x1) ? EDependencyProperty::Direct : EDependencyProperty::None;
-			ManageProperties[n] = (n & 0x2) ? EDependencyProperty::CookRule : EDependencyProperty::None;
+			ManageProperties[n] |= (n & 0x2) ? EDependencyProperty::CookRule : EDependencyProperty::None;
 		}
 
 		for (ESortTestType SortTestType : { ESortTestType::AlwaysSorted, ESortTestType::UnsortedThenSorted, ESortTestType::AlwaysUnsorted})
