AssetManager/AssetRegistry: Change the graph search conducted in UpdateManagementDatabase -> SetManageReferences to do two separate graph searches: one for Chunking that traverses Game and Build dependencies, and one for CookRules that traverses only Game dependencies. The two graph searches are conducted at the same time for better performance by using a graph search painting algorithm.
In this initial submission the separate traversal rules are toggled off so the new algorithm can be tested separately from the desired behavior change.
#rnx

diff --git a/Engine/Source/Runtime/AssetRegistry/Private/AssetRegistry.cpp b/Engine/Source/Runtime/AssetRegistry/Private/AssetRegistry.cpp
--- a/Engine/Source/Runtime/AssetRegistry/Private/AssetRegistry.cpp
+++ b/Engine/Source/Runtime/AssetRegistry/Private/AssetRegistry.cpp
@@ -9223,6 +9223,31 @@ static FAutoConsoleVariableRef CVarAssetRegistryManagementPathsPackageDebugName(
 namespace UE::AssetRegistry
 {
 
+/**
+ * In SetManagerReferences, we execute a graph search within the graph of PackageDependencies for Assets to decide
+ * which Assets should be marked as having Managed properties. We use a graph search painting algorithm to
+ * apply multiple properties within a single graph search. A painting algorithm keeps track of the current list of
+ * properties to apply to each visited vertex, and filters that list of properties depending on the edge properties
+ * when traversing an edge. The list of Managed properties we paint during SetManageReferences is given by the values
+ * of EManageSearchColor.
+ */
+enum class EManageSearchColor : uint8
+{
+	CookRule = 0,
+	Chunking = 1,
+	Num,
+};
+constexpr int32 ManageSearchColorsNum = static_cast<int32>(EManageSearchColor::Num);
+static_assert(ManageSearchColorsNum <= 8 * sizeof(uint32), "We store lists of SearchColors as a uint32 bitfield.");
+enum class EManageSearchColorBit : uint32
+{
+	None = 0,
+	CookRule = (1 << static_cast<int32>(EManageSearchColor::CookRule)),
+	Chunking = (1 << static_cast<int32>(EManageSearchColor::Chunking)),
+	All = CookRule | Chunking,
+};
+ENUM_CLASS_FLAGS(EManageSearchColorBit);
+
 /**
  * Data held for each Asset (aka node in our graph search) across the multiple graph searches conducted within
  * SetManageReferences, and across multiple calls to SetManageReferences
@@ -9236,10 +9261,11 @@ struct FSetManageReferencesNodeData
 {
 	// Values preserved across multiple rounds
 	/**
-	 * bManagedInEarlierRound is set for every asset found managed in a round, and later rounds can use that information
-	 * in ShouldSetManager calls to e.g. ignore the management of a node by a lower-priority manager.
+	 * ManagedInEarlierRound bit (1 << N) is set at the end of a round for assets painted with ((EManageSearchColor)N) 
+	 * in the round, and later rounds can use that information in ShouldSetManager calls to e.g. ignore the management
+	 * of a node by a lower-priority manager.
 	 */
-	bool bManagedInEarlierRound = false;
+	EManageSearchColorBit ManagedInEarlierRound = EManageSearchColorBit::None;
 
 	// Values reset for each round.
 	/**
@@ -9248,20 +9274,28 @@ struct FSetManageReferencesNodeData
 	 */
 	bool bModifiedByRound = false;
 	/**
-	 * Whether the asset was found to be managed by the round (and therefore needs to be added to
-	 * bManagedInEarlierRound).
+	 * ManagedInThisRound bit (1 << N) is set during a round's graph search for the asset painted with
+	 * ((EManageSearchColor)N) in the round.
 	 */
-	bool bManagedInThisRound = false;
+	EManageSearchColorBit ManagedInThisRound = EManageSearchColorBit::None;
 	/**
-	 * The direct managers of the asset in the current round, as described by the caller of the round in
-	 * FSetManageReferencesContext.ManagerMap.
+	 * Whether the asset has a direct manager in the current round, as described by the caller of the round in
+	 * FSetManageReferencesContext.ManagerMap. Direct Managers apply to all EManageSearchColors.
 	 */
-	TArray<FDependsNode*> DirectManagersThisRound;
+	bool bDirectManagedInThisRound = false;
+
 	/**
-	 * The source node (aka asset) that caused this node to be referenced in the current round, used to provide
-	 * instigator chains for what caused an asset to be added to the cook or to a chunk.
+	 * For a manager node in the current round, the list of directly managed assets. Directly managed assets paint
+	 * with all EManageSearchColors.
 	 */
-	FDependsNode* DebugInstigator = nullptr;
+	TArray<FDependsNode*> DirectManagedAssetsThisRound;
+
+	/**
+	 * DebugInstigator[N] is the source node (aka asset) that caused this node to be painted with
+	 * ((EManageSearchColor)N) in the current round. This information is used to provide instigator chains for what
+	 * caused an asset to be added to the cook or to a chunk.
+	 */
+	FDependsNode* DebugInstigator[ManageSearchColorsNum] = {};
 
 	// Values reset for the graph search from each manager with a round.
 	/**
@@ -9269,11 +9303,15 @@ struct FSetManageReferencesNodeData
 	 * restored at the end of the search).
 	 */
 	bool bModifiedByCurrentManager = false;
-	/** Whether the asset was found to be managed by the manager. */
-	bool bVisitedByCurrentManager = false;
+	/**
+	 * VisitedByCurrentManager bit (1 <<N) records whether the asset was found to be painted with
+	 * ((EManageSearchColor)N) by the manager.
+	 */
+	EManageSearchColorBit VisitedByCurrentManager = EManageSearchColorBit::None;
 	/**
 	 * The AssetRegistry dependency properties that should be assigned to the targetnode for the Management edge we
-	 * create for the node from the current manager.
+	 * create for the node from the current manager; this accumulates the properties we set for each painted color,
+	 * as well as the special EDependencyProperty::Direct property.
 	 */
 	EDependencyProperty CurrentManagerProperties = EDependencyProperty::None;
 };
@@ -9298,7 +9336,7 @@ void PrintAssetRegistryManagementPathsPackageDebugInfo(FDependsNode* Node, UE::A
 			{
 				break;
 			}
-			FDependsNode* ReferencingNode = (*CurrentNodeData)->DebugInstigator;
+			FDependsNode* ReferencingNode = (*CurrentNodeData)->DebugInstigator[(int32)EManageSearchColor::Chunking];
 			if (!ReferencingNode)
 			{
 				break;
@@ -9420,7 +9458,7 @@ void FAssetRegistryImpl::SetManageReferences(FSetManageReferencesContext& Contex
 		// Clear information from Scratch about Earlier rounds
 		for (TPair<FDependsNode*, TUniquePtr<FSetManageReferencesNodeData>>& NodePair : *NodeData)
 		{
-			NodePair.Value->bManagedInEarlierRound = false;
+			NodePair.Value->ManagedInEarlierRound = EManageSearchColorBit::None;
 		}
 
 		// Clear all outgoing manage dependencies, and update the referencers of any nodes with incoming
@@ -9448,9 +9486,14 @@ void FAssetRegistryImpl::SetManageReferences(FSetManageReferencesContext& Contex
 		}
 	}
 
-	TArray<FDependsNode*> CurrentRoundExplicitlyManagedNodes;
+	TArray<FDependsNode*> CurrentRoundManagers;
 	TArray<FDependsNode*> CurrentManagerModifiedNodes;
-	TArray<FDependsNode*> CurrentManagerVisitQueue;
+	TArray<TPair<FDependsNode*, EManageSearchColorBit>> CurrentManagerVisitQueue;
+
+	// Convert ManageMap from edges in the FAssetIdentifier graph into edges in the FDependsNode graph.
+	// Also mark each explicitly managed node; AssetManager's ShouldSetManager function might decide to skip
+	// an indirect manage reference if there is a direct manage reference in the same round.
+	CurrentRoundManagers.Reserve(Context.ManagerMap->Num());
 	for (const TPair<FAssetIdentifier, FAssetIdentifier>& Pair : *Context.ManagerMap)
 	{
 		FDependsNode* ManagedNode = State.FindDependsNode(Pair.Value);
@@ -9462,38 +9505,41 @@ void FAssetRegistryImpl::SetManageReferences(FSetManageReferencesContext& Contex
 		}
 		FDependsNode* ManagerNode = State.CreateOrFindDependsNode(Pair.Key);
 
-		CurrentRoundExplicitlyManagedNodes.Add(ManagedNode);
-		FindOrAddNodeData(ManagedNode).DirectManagersThisRound.Add(ManagerNode);
+		CurrentRoundManagers.Add(ManagerNode);
+		FindOrAddNodeData(ManagedNode).bDirectManagedInThisRound = true;
+		FindOrAddNodeData(ManagerNode).DirectManagedAssetsThisRound.Add(ManagedNode);
 	}
 
-	Algo::Sort(CurrentRoundExplicitlyManagedNodes);
-	CurrentRoundExplicitlyManagedNodes.SetNum(Algo::Unique(CurrentRoundExplicitlyManagedNodes));
+	Algo::Sort(CurrentRoundManagers);
+	CurrentRoundManagers.SetNum(Algo::Unique(CurrentRoundManagers));
 
-	// For each directly mangaged node, do a graph search for each of its managers.
-	for (FDependsNode* BaseManagedNode : CurrentRoundExplicitlyManagedNodes)
+	// For each manager, do a graph search starting from its directly managed nodes.
+	for (FDependsNode* ManagerNode : CurrentRoundManagers)
 	{
-		FSetManageReferencesNodeData& BaseManagedNodeData = FindOrAddNodeData(BaseManagedNode);
-		for (FDependsNode* ManagerNode : BaseManagedNodeData.DirectManagersThisRound)
-		{
-			check(CurrentManagerModifiedNodes.IsEmpty());
-			check(CurrentManagerVisitQueue.IsEmpty());
+		FSetManageReferencesNodeData& ManagerNodeData = FindOrAddNodeData(ManagerNode);
+		check(CurrentManagerModifiedNodes.IsEmpty());
+		check(CurrentManagerVisitQueue.IsEmpty());
 
-			auto IterateFunction = [&ManagerNode, &Context, NodeData, &CurrentManagerModifiedNodes,
-				&CurrentManagerVisitQueue, &FindOrAddNodeData]
+		auto VisitNode = [ManagerNode, &Context, &CurrentManagerModifiedNodes,
+			&CurrentManagerVisitQueue, &FindOrAddNodeData]
 			(FDependsNode* SourceNode, FDependsNode* TargetNode, EDependencyCategory DependencyType,
-				EDependencyProperty DependencyProperties)
+				EDependencyProperty PackageDependencyProperties, EManageSearchColorBit PaintColors)
 			{
 				// Only recurse if we haven't already visited, and this node passes recursion test
 				FSetManageReferencesNodeData& TargetNodeData = FindOrAddNodeData(TargetNode);
-				if (TargetNodeData.bVisitedByCurrentManager)
+				PaintColors &= ~TargetNodeData.VisitedByCurrentManager;
+				if (PaintColors == EManageSearchColorBit::None)
 				{
 					return;
 				}
 
-				EAssetSetManagerFlags::Type EdgeFlags = (EAssetSetManagerFlags::Type)(
-					(SourceNode == ManagerNode ? EAssetSetManagerFlags::IsDirectSet : 0)
-					| (TargetNodeData.bManagedInEarlierRound ? EAssetSetManagerFlags::TargetHasExistingManager : 0)
-					| (TargetNodeData.DirectManagersThisRound.Num() > 0 && SourceNode != ManagerNode
+				EManageSearchColorBit ColorsManagedInEarlierRound = TargetNodeData.ManagedInEarlierRound & PaintColors;
+				EManageSearchColorBit ColorsNotManagedInEarlierRound = (~TargetNodeData.ManagedInEarlierRound) & PaintColors;
+
+				bool bDirectSet = SourceNode == ManagerNode;
+				EAssetSetManagerFlags::Type EdgeFlagsOtherThanManagedInEarlierRound = (EAssetSetManagerFlags::Type)(
+					(bDirectSet ? EAssetSetManagerFlags::IsDirectSet : 0)
+					| (TargetNodeData.bDirectManagedInThisRound && !bDirectSet
 						? EAssetSetManagerFlags::TargetHasDirectManager : 0));
 				FShouldSetManagerContext ShouldSetContext
 				{
@@ -9503,89 +9549,160 @@ void FAssetRegistryImpl::SetManageReferences(FSetManageReferencesContext& Contex
 					.ManagerAssetFlags = ConstructAssetManagerAssetFlags(*ManagerNode),
 					.SourceAssetFlags = ConstructAssetManagerAssetFlags(*SourceNode),
 					.TargetAssetFlags = ConstructAssetManagerAssetFlags(*TargetNode),
-					.EdgeFlags = EdgeFlags,
+					.EdgeFlags = EdgeFlagsOtherThanManagedInEarlierRound,
 					.EdgeARCategory = DependencyType,
-					.EdgeARProperties = DependencyProperties
+					.EdgeARProperties = PackageDependencyProperties
 				};
 
-				EAssetSetManagerResult::Type Result = (*Context.ShouldSetManager)(ShouldSetContext);
-				if (Result == EAssetSetManagerResult::DoNotSet)
+				EManageSearchColorBit ColorsToSet = EManageSearchColorBit::None;
+				EManageSearchColorBit ColorsToRecurse = EManageSearchColorBit::None;
+
+				if (ColorsNotManagedInEarlierRound != EManageSearchColorBit::None)
 				{
-					return;
+					EAssetSetManagerResult::Type Result = (*Context.ShouldSetManager)(ShouldSetContext);
+					switch (Result)
+					{
+					case EAssetSetManagerResult::DoNotSet:
+						break;
+					case EAssetSetManagerResult::SetButDoNotRecurse:
+						ColorsToSet |= ColorsNotManagedInEarlierRound;
+						break;
+					case EAssetSetManagerResult::SetAndRecurse:
+						ColorsToSet |= ColorsNotManagedInEarlierRound;
+						ColorsToRecurse |= ColorsNotManagedInEarlierRound;
+						break;
+					default:
+						checkNoEntry();
+						break;
+					}
+				}
+				if (ColorsManagedInEarlierRound != EManageSearchColorBit::None)
+				{
+					ShouldSetContext.EdgeFlags = (EAssetSetManagerFlags::Type)
+						(ShouldSetContext.EdgeFlags | EAssetSetManagerFlags::TargetHasExistingManager);
+					EAssetSetManagerResult::Type Result = (*Context.ShouldSetManager)(ShouldSetContext);
+					switch (Result)
+					{
+					case EAssetSetManagerResult::DoNotSet:
+						break;
+					case EAssetSetManagerResult::SetButDoNotRecurse:
+						ColorsToSet |= ColorsManagedInEarlierRound;
+						break;
+					case EAssetSetManagerResult::SetAndRecurse:
+						ColorsToSet |= ColorsManagedInEarlierRound;
+						ColorsToRecurse |= ColorsManagedInEarlierRound;
+						break;
+					default:
+						checkNoEntry();
+						break;
+					}
 				}
 
-				TargetNodeData.bVisitedByCurrentManager = true;
-				TargetNodeData.bModifiedByCurrentManager = true;
-				CurrentManagerModifiedNodes.Add(TargetNode);
-
-				TargetNodeData.CurrentManagerProperties =
-					((EdgeFlags & EAssetSetManagerFlags::IsDirectSet)
-						? EDependencyProperty::Direct : EDependencyProperty::None)
-					| EDependencyProperty::CookRule;
-				TargetNodeData.DebugInstigator = SourceNode;
+				if (ColorsToSet != EManageSearchColorBit::None)
+				{
+					TargetNodeData.VisitedByCurrentManager |= ColorsToSet;
+					TargetNodeData.bModifiedByCurrentManager = true;
+					CurrentManagerModifiedNodes.Add(TargetNode);
+
+					TargetNodeData.CurrentManagerProperties |=
+						(bDirectSet ? EDependencyProperty::Direct : EDependencyProperty::None)
+						| (EnumHasAnyFlags(ColorsToSet, EManageSearchColorBit::CookRule)
+							? EDependencyProperty::CookRule : EDependencyProperty::None)
+						// Painting with any color, or marking it as Direct, also marks it as painted with Chunking.
+						// There is not a separate flag to store this in EDependencyProperty; the existence of the
+						// dependency implies it is at least a Chunking dependency.
+						;
+					for (int32 PaintColor = 0; PaintColor < ManageSearchColorsNum; ++PaintColor)
+					{
+						if (EnumHasAnyFlags(ColorsToSet, (EManageSearchColorBit)(1 << ManageSearchColorsNum))
+							&& !TargetNodeData.DebugInstigator[PaintColor])
+						{
+							TargetNodeData.DebugInstigator[PaintColor] = SourceNode;
+						}
+					}
+				}
 
-				if (Result == EAssetSetManagerResult::SetAndRecurse)
+				if (ColorsToRecurse != EManageSearchColorBit::None)
 				{
-					CurrentManagerVisitQueue.Push(TargetNode);
+					CurrentManagerVisitQueue.Push({ TargetNode, ColorsToRecurse });
 				}
 			};
 
-			// Check initial node
-			IterateFunction(ManagerNode, BaseManagedNode, EDependencyCategory::Manage,
-				EDependencyProperty::Game | EDependencyProperty::Build);
+		// Visit the directly managed nodes and populate the starting vertices for the graph search of indirectly
+		// managed nodes
+		for (FDependsNode* ManagedNode : ManagerNodeData.DirectManagedAssetsThisRound)
+		{
+			VisitNode(ManagerNode, ManagedNode, EDependencyCategory::Manage,
+				EDependencyProperty::Game | EDependencyProperty::Build, EManageSearchColorBit::All);
+		}
 
-			// Do all recursion first, but only if we have a recurse type
-			if (Context.RecurseType != EDependencyCategory::None)
+		// Find and Visit all indirectly managed nodes, but only if we have a recurse type
+		if (Context.RecurseType != EDependencyCategory::None)
+		{
+			while (CurrentManagerVisitQueue.Num())
 			{
-				while (CurrentManagerVisitQueue.Num())
-				{
-					// Pull off end of array to avoid the cost of shifts; order of visitation doesn't matter
-					FDependsNode* SourceNode = CurrentManagerVisitQueue.Pop();
+				// Pull off end of array to avoid the cost of shifts; order of visitation doesn't matter
+				TPair<FDependsNode*, EManageSearchColorBit> VisitPair = CurrentManagerVisitQueue.Pop();
+				FDependsNode* SourceNode = VisitPair.Key;
+				EManageSearchColorBit VertexPaintColors = VisitPair.Value;
 
-					SourceNode->IterateOverDependencies([&IterateFunction, SourceNode]
-					(FDependsNode* TargetNode, EDependencyCategory DependencyCategory,
-						EDependencyProperty DependencyProperties, bool bDuplicate)
+				SourceNode->IterateOverDependencies([&VisitNode, SourceNode, VertexPaintColors]
+				(FDependsNode* TargetNode, EDependencyCategory DependencyCategory,
+					EDependencyProperty DependencyProperties, bool bDuplicate)
+					{
+						// Skip dependencies that are EditorOnly and non-build.
+						// Propagate only through UsedInGame or build dependencies.
+						if (EnumHasAnyFlags(DependencyProperties, EDependencyProperty::Game | EDependencyProperty::Build))
 						{
-							// Skip dependencies that are EditorOnly and non-build.
-							// Propagate only through UsedInGame or build dependencies.
-							if (EnumHasAnyFlags(DependencyProperties, EDependencyProperty::Game | EDependencyProperty::Build))
+							// Different behavior in the graph searches is temporarily disabled in the first submit,
+							// So we can test the new graph search independently of the desired behavior change.
+							static constexpr bool bCookRulesPaintedOnlyWithGame = false;
+
+							EManageSearchColorBit EdgePaintColors;
+							if (bCookRulesPaintedOnlyWithGame)
 							{
-								IterateFunction(SourceNode, TargetNode, DependencyCategory, DependencyProperties);
+								EdgePaintColors = EnumHasAnyFlags(DependencyProperties, EDependencyProperty::Game) ?
+									EManageSearchColorBit::All : EManageSearchColorBit::Chunking;
 							}
-						}, Context.RecurseType);
-				}
+							else
+							{
+								EdgePaintColors = EManageSearchColorBit::All;
+							}
+							EManageSearchColorBit TargetPaintColors = VertexPaintColors & EdgePaintColors;
+							VisitNode(SourceNode, TargetNode, DependencyCategory, DependencyProperties, TargetPaintColors);
+						}
+					}, Context.RecurseType);
 			}
+		}
 
-			ManagerNode->SetIsDependencyListSorted(UE::AssetRegistry::EDependencyCategory::Manage, false);
-			FSetManageReferencesNodeData& ManagerData = FindOrAddNodeData(ManagerNode);
-			ManagerData.bModifiedByRound = true;
-			CurrentRoundModifiedNodes.Add(ManagerNode);
+		ManagerNode->SetIsDependencyListSorted(UE::AssetRegistry::EDependencyCategory::Manage, false);
+		ManagerNodeData.bModifiedByRound = true;
+		CurrentRoundModifiedNodes.Add(ManagerNode);
 
-			for (FDependsNode* ModifiedNode : CurrentManagerModifiedNodes)
+		for (FDependsNode* ModifiedNode : CurrentManagerModifiedNodes)
+		{
+			FSetManageReferencesNodeData& ModifiedData = FindOrAddNodeData(ModifiedNode);
+			if (!ModifiedData.bModifiedByCurrentManager)
 			{
-				FSetManageReferencesNodeData& ModifiedData = FindOrAddNodeData(ModifiedNode);
-				if (!ModifiedData.bModifiedByCurrentManager)
-				{
-					// A duplicate of a NodeData we already handled earlier in the list
-					continue;
-				}
+				// A duplicate of a NodeData we already handled earlier in the list
+				continue;
+			}
 
-				ModifiedNode->SetIsReferencersSorted(false);
-				ModifiedNode->AddReferencer(ManagerNode);
-				ManagerNode->AddDependency(ModifiedNode, EDependencyCategory::Manage,
-					ModifiedData.CurrentManagerProperties);
+			ModifiedNode->SetIsReferencersSorted(false);
+			ModifiedNode->AddReferencer(ManagerNode);
+			ManagerNode->AddDependency(ModifiedNode, EDependencyCategory::Manage,
+				ModifiedData.CurrentManagerProperties);
 
-				ModifiedData.bManagedInThisRound = true;
-				ModifiedData.bModifiedByRound = true;
-				CurrentRoundModifiedNodes.Add(ModifiedNode);
+			ModifiedData.ManagedInThisRound |= ModifiedData.VisitedByCurrentManager;
+			ModifiedData.bModifiedByRound = true;
+			CurrentRoundModifiedNodes.Add(ModifiedNode);
 
-				ModifiedData.bModifiedByCurrentManager = false;
-				ModifiedData.bVisitedByCurrentManager = false;
-				ModifiedData.CurrentManagerProperties = EDependencyProperty::None;
-			}
-			CurrentManagerModifiedNodes.Reset();
-			CurrentManagerVisitQueue.Reset();
+			ModifiedData.bModifiedByCurrentManager = false;
+			ModifiedData.VisitedByCurrentManager = EManageSearchColorBit::None;
+			ModifiedData.CurrentManagerProperties = EDependencyProperty::None;
 		}
+		CurrentManagerModifiedNodes.Reset();
+		CurrentManagerVisitQueue.Reset();
 	}
 
 #if WITH_EDITOR
@@ -9609,14 +9726,18 @@ void FAssetRegistryImpl::SetManageReferences(FSetManageReferencesContext& Contex
 
 		// Restore all modified nodes to manage dependencies sorted and references sorted, so we can efficiently read
 		// them in future operations
-		ModifiedData.bManagedInEarlierRound |= ModifiedData.bManagedInThisRound;
+		ModifiedData.ManagedInEarlierRound |= ModifiedData.ManagedInThisRound;
 		ModifiedNode->SetIsDependencyListSorted(EDependencyCategory::Manage, bShouldSortDependencies);
 		ModifiedNode->SetIsReferencersSorted(bShouldSortReferencers);
 
 		ModifiedData.bModifiedByRound = false;
-		ModifiedData.bManagedInThisRound = false;
-		ModifiedData.DirectManagersThisRound.Empty();
-		ModifiedData.DebugInstigator = nullptr;
+		ModifiedData.ManagedInThisRound = EManageSearchColorBit::None;
+		ModifiedData.bDirectManagedInThisRound = false;
+		ModifiedData.DirectManagedAssetsThisRound.Empty();
+		for (FDependsNode*& DebugInstigator : ModifiedData.DebugInstigator)
+		{
+			DebugInstigator = nullptr;
+		}
 	}
 }
 
