[AssetRegistry]

Disable DependencyDataGathering parallelism due to race when gathering

#rnx

diff --git a/Engine/Source/Runtime/AssetRegistry/Private/AssetRegistry.cpp b/Engine/Source/Runtime/AssetRegistry/Private/AssetRegistry.cpp
--- a/Engine/Source/Runtime/AssetRegistry/Private/AssetRegistry.cpp
+++ b/Engine/Source/Runtime/AssetRegistry/Private/AssetRegistry.cpp
@@ -7188,7 +7188,12 @@ void FAssetRegistryImpl::DependencyDataGathered(TMultiMap<FName, FPackageDepende
 	bool bInterrupted = false;
 	int32 IterationCounter = 0;
 	constexpr int32 MinBatchSize = 1024;
-	const bool bUseJobs = DependsResults.Num() > MinBatchSize;
+
+	// Disable parallelism until UE-295029 is resolved
+	const bool bUseJobs = false;
+	//const bool bUseJobs = DependsResults.Num() > MinBatchSize;
+	const EParallelForFlags ParallelForFlags = bUseJobs ? EParallelForFlags::None : EParallelForFlags::ForceSingleThread;
+
 	{
 		TRACE_CPUPROFILER_EVENT_SCOPE(DependencyGatheredData_GatherPackages);
 		const bool bShouldSortReferencers = ShouldSortReferencers();
@@ -7394,7 +7399,7 @@ void FAssetRegistryImpl::DependencyDataGathered(TMultiMap<FName, FPackageDepende
 			{
 				SourceNode->AddDependency(DependsNode, EDependencyCategory::SearchableName, EDependencyProperty::None);
 			}
-		});
+		}, ParallelForFlags);
 
 	// Add referencers - Single threaded due to the poor memory access patterns between nodes and their dependencies
 	{
@@ -7438,7 +7443,7 @@ void FAssetRegistryImpl::DependencyDataGathered(TMultiMap<FName, FPackageDepende
 			// Set AllowShrinking first so SetIsReferencersSorted will reduce our memory use if we have over allocated
 			Node->SetAllowShrinking(true);
 			Node->SetIsReferencersSorted(bShouldSortReferencers);
-		});
+		}, ParallelForFlags);
 }
 
 void FAssetRegistryImpl::CookedPackageNamesWithoutAssetDataGathered(Impl::FEventContext& EventContext,
