- ReadOnly state in the AR, via FAssetPackageData
- Using it in the HybridPacakgeStore to choose to load uncooked packages if the file is checked out

diff --git a/Engine/Source/Runtime/AssetRegistry/Private/AssetDataGatherer.cpp b/Engine/Source/Runtime/AssetRegistry/Private/AssetDataGatherer.cpp
--- a/Engine/Source/Runtime/AssetRegistry/Private/AssetDataGatherer.cpp
+++ b/Engine/Source/Runtime/AssetRegistry/Private/AssetDataGatherer.cpp
@@ -374,11 +374,12 @@ void AssignStringWithoutShrinking(FString& InOutResult, FStringView Value)
 }
 
 FDiscoveredPathData::FDiscoveredPathData(FStringView InLocalAbsPath, FStringView InLongPackageName,
-	FStringView InRelPath, const FDateTime& InPackageTimestamp, EGatherableFileType InType)
+	FStringView InRelPath, const FDateTime& InPackageTimestamp, bool InIsReadOnly, EGatherableFileType InType)
 	: LocalAbsPath(InLocalAbsPath)
 	, LongPackageName(InLongPackageName)
 	, RelPath(InRelPath)
 	, PackageTimestamp(InPackageTimestamp)
+	, bIsReadOnly(InIsReadOnly)
 	, Type(InType)
 {
 }
@@ -388,6 +389,7 @@ FDiscoveredPathData::FDiscoveredPathData(FStringView InLocalAbsPath, FStringView
 	: LocalAbsPath(InLocalAbsPath)
 	, LongPackageName(InLongPackageName)
 	, RelPath(InRelPath)
+	, bIsReadOnly(false)
 	, Type(InType)
 {
 }
@@ -402,10 +404,11 @@ void FDiscoveredPathData::Assign(FStringView InLocalAbsPath, FStringView InLongP
 }
 
 void FDiscoveredPathData::Assign(FStringView InLocalAbsPath, FStringView InLongPackageName, FStringView InRelPath,
-	const FDateTime& InPackageTimestamp, EGatherableFileType InType)
+	const FDateTime& InPackageTimestamp, bool InIsReadOnly, EGatherableFileType InType)
 {
 	Assign(InLocalAbsPath, InLongPackageName, InRelPath, InType);
 	PackageTimestamp = InPackageTimestamp;
+	bIsReadOnly = InIsReadOnly;
 }
 
 SIZE_T FDiscoveredPathData::GetAllocatedSize() const
@@ -414,17 +417,18 @@ SIZE_T FDiscoveredPathData::GetAllocatedSize() const
 }
 
 FGatheredPathData::FGatheredPathData(FStringView InLocalAbsPath, FStringView InLongPackageName,
-	const FDateTime& InPackageTimestamp, EGatherableFileType InType)
+	const FDateTime& InPackageTimestamp, bool InIsReadOnly, EGatherableFileType InType)
 	: LocalAbsPath(InLocalAbsPath)
 	, LongPackageName(InLongPackageName)
 	, PackageTimestamp(InPackageTimestamp)
+	, bIsReadOnly(InIsReadOnly)
 	, Type(InType)
 {
 }
 
 FGatheredPathData::FGatheredPathData(const FDiscoveredPathData& DiscoveredData)
 	: FGatheredPathData(DiscoveredData.LocalAbsPath, DiscoveredData.LongPackageName,
-		DiscoveredData.PackageTimestamp, DiscoveredData.Type)
+		DiscoveredData.PackageTimestamp, DiscoveredData.bIsReadOnly, DiscoveredData.Type)
 {
 }
 
@@ -432,23 +436,25 @@ FGatheredPathData::FGatheredPathData(FDiscoveredPathData&& DiscoveredData)
 	: LocalAbsPath(MoveTemp(DiscoveredData.LocalAbsPath))
 	, LongPackageName(MoveTemp(DiscoveredData.LongPackageName))
 	, PackageTimestamp(MoveTemp(DiscoveredData.PackageTimestamp))
+	, bIsReadOnly(DiscoveredData.bIsReadOnly)
 	, Type(DiscoveredData.Type)
 {
 }
 
 void FGatheredPathData::Assign(FStringView InLocalAbsPath, FStringView InLongPackageName,
-	const FDateTime& InPackageTimestamp, EGatherableFileType InType)
+	const FDateTime& InPackageTimestamp, bool InIsReadOnly, EGatherableFileType InType)
 {
 	AssignStringWithoutShrinking(LocalAbsPath, InLocalAbsPath);
 	AssignStringWithoutShrinking(LongPackageName, InLongPackageName);
 	PackageTimestamp = InPackageTimestamp;
+	bIsReadOnly = InIsReadOnly;
 	Type = InType;
 }
 
 void FGatheredPathData::Assign(const FDiscoveredPathData& DiscoveredData)
 {
 	Assign(DiscoveredData.LocalAbsPath, DiscoveredData.LongPackageName, DiscoveredData.PackageTimestamp,
-		DiscoveredData.Type);
+		DiscoveredData.bIsReadOnly, DiscoveredData.Type);
 }
 
 SIZE_T FGatheredPathData::GetAllocatedSize() const
@@ -2031,7 +2037,7 @@ void FAssetDataDiscovery::TickInternal(bool bTickAll)
 						FPathViews::AppendPath(Data.DirLongPackageName, BaseName);
 						FPathViews::AppendPath(Data.DirLocalAbsPath, RelPath);
 						Data.IteratedFiles[Data.NumIteratedFiles++].Assign(Data.DirLocalAbsPath, Data.DirLongPackageName, RelPath,
-							FileData.ModificationTime, FileType);
+							FileData.ModificationTime, !!FileData.bIsReadOnly, FileType);
 					}
 				}
 			}
@@ -2081,7 +2087,7 @@ void FAssetDataDiscovery::TickInternal(bool bTickAll)
 
 			FCachedDirScanDir CacheDataToAdd;
 			auto ProcessIterData = [this, &Data, &CacheDataToAdd]
-			(const TCHAR* IterFilename, bool bIsDirectory, const FDateTime& ModificationTime, FFileJournalFileHandle JournalHandle, bool bIsReparsePoint)
+			(const TCHAR* IterFilename, bool bIsDirectory, const FDateTime& ModificationTime, bool bIsReadOnly, FFileJournalFileHandle JournalHandle, bool bIsReparsePoint)
 			{
 				FStringView LocalAbsPath(IterFilename);
 				FStringView RelPath;
@@ -2142,6 +2148,7 @@ void FAssetDataDiscovery::TickInternal(bool bTickAll)
 						FCachedDirScanFile& FileData = CacheDataToAdd.Files.Emplace_GetRef();
 						FileData.RelPath = FString(RelPath);
 						FileData.ModificationTime = ModificationTime;
+						FileData.bIsReadOnly = bIsReadOnly;
 					}
 					EGatherableFileType FileType = GetFileType(RelPath);
 					// Don't record files that contain invalid packagepath characters (not counting their extension)
@@ -2158,7 +2165,7 @@ void FAssetDataDiscovery::TickInternal(bool bTickAll)
 							}
 							FPathViews::AppendPath(Data.DirLongPackageName, BaseName);
 							Data.IteratedFiles[Data.NumIteratedFiles++].Assign(LocalAbsPath, Data.DirLongPackageName, RelPath,
-								ModificationTime, FileType);
+								ModificationTime, bIsReadOnly, FileType);
 						}
 					}
 				}
@@ -2173,7 +2180,7 @@ void FAssetDataDiscovery::TickInternal(bool bTickAll)
 					[this, &Data, &CacheDataToAdd, &ProcessIterData]
 					(const TCHAR* IterFilename, const FFileJournalData& IterData)
 					{
-						return ProcessIterData(IterFilename, IterData.bIsDirectory, IterData.ModificationTime, IterData.JournalHandle, IterData.bIsReparsePoint);
+						return ProcessIterData(IterFilename, IterData.bIsDirectory, IterData.ModificationTime, IterData.bIsReadOnly, IterData.JournalHandle, IterData.bIsReparsePoint);
 					}))
 				{
 					bIteratedDirectory = true;
@@ -2213,7 +2220,7 @@ void FAssetDataDiscovery::TickInternal(bool bTickAll)
 					[this, &Data, &ProcessIterData]
 					(const TCHAR* IterFilename, const FFileStatData& IterData)
 					{
-						return ProcessIterData(IterFilename, IterData.bIsDirectory, IterData.ModificationTime, FFileJournalFileHandle(), false /* bIsReparsePoint */);
+						return ProcessIterData(IterFilename, IterData.bIsDirectory, IterData.ModificationTime, IterData.bIsReadOnly, FFileJournalFileHandle(), false /* bIsReparsePoint */);
 					});
 			}
 
@@ -2415,6 +2422,12 @@ FDateTime FPathExistence::GetModificationTime()
 	return ModificationTime;
 }
 
+bool FPathExistence::IsReadOnly()
+{
+	LoadExistenceData();
+	return bIsReadOnly;
+}
+
 void FPathExistence::LoadExistenceData()
 {
 	if (bHasExistenceData)
@@ -2437,6 +2450,7 @@ void FPathExistence::LoadExistenceData()
 		}
 
 		ModificationTime = StatData.ModificationTime;
+		bIsReadOnly = StatData.bIsReadOnly;
 		PathType = StatData.bIsDirectory ? EType::Directory : EType::File;
 	}
 	else
@@ -2465,6 +2479,7 @@ void FPathExistence::LoadExistenceData()
 		{
 			PathType = EType::MissingParentDir;
 		}
+		bIsReadOnly = false;
 	}
 
 	bHasExistenceData = true;
@@ -2651,7 +2666,7 @@ void FAssetDataDiscovery::SetPropertiesAndWait(TArrayView<FPathExistence> QueryP
 							FPathViews::AppendPath(LongPackageName, ScanDir->GetMountRelPath());
 							FPathViews::AppendPath(LongPackageName, FileRelPathNoExt);
 							AddDiscoveredFile(FDiscoveredPathData(QueryPath.GetLocalAbsPath(), LongPackageName,
-								RelPathFromParentDir, QueryPath.GetModificationTime(), FileType));
+								RelPathFromParentDir, QueryPath.GetModificationTime(), QueryPath.IsReadOnly(), FileType));
 							if (FPathViews::IsPathLeaf(RelPathFromControlDir) && !ScanDir->HasScanned())
 							{
 								SetIsIdle(false);
@@ -3128,7 +3143,7 @@ void FAssetDataDiscovery::OnFileCreated(const FString& LocalAbsPath)
 			LongPackageName << MountDir->GetLongPackageName();
 			FPathViews::AppendPath(LongPackageName, ScanDir->GetMountRelPath());
 			FPathViews::AppendPath(LongPackageName, FileRelPathNoExt);
-			AddDiscoveredFile(FDiscoveredPathData(LocalAbsPath, LongPackageName, RelPathFromParentDir, StatData.ModificationTime, FileType));
+			AddDiscoveredFile(FDiscoveredPathData(LocalAbsPath, LongPackageName, RelPathFromParentDir, StatData.ModificationTime, StatData.bIsReadOnly, FileType));
 			if (FPathViews::IsPathLeaf(FileRelPath))
 			{
 				ScanDir->MarkFileAlreadyScanned(FileRelPath);
@@ -4449,7 +4464,8 @@ FAssetDataGatherer::ETickResult FAssetDataGatherer::TickInternal(double& TickSta
 			{
 				TargetAssetResults.Add(MakeUnique<FAssetData>(AssetData));
 			}
-			TargetDependencyResults.Add(DiskCachedAssetData->DependencyData);
+			FPackageDependencyData& DependencyData = TargetDependencyResults.Add_GetRef(DiskCachedAssetData->DependencyData);
+			DependencyData.PackageData.SetIsReadOnly(AssetFileData.bIsReadOnly);
 
 			AddToCache(PackageName, DiskCachedAssetData);
 		}
@@ -4581,7 +4597,8 @@ FAssetDataGatherer::ETickResult FAssetDataGatherer::TickInternal(double& TickSta
 
 				// Add the results from the package into our output results
 				TargetAssetResults.Append(MoveTemp(ReadContext.AssetDataFromFile));
-				TargetDependencyResults.Add(MoveTemp(ReadContext.DependencyData));
+				FPackageDependencyData& DependencyData = TargetDependencyResults.Add_GetRef(MoveTemp(ReadContext.DependencyData));
+				DependencyData.PackageData.SetIsReadOnly(ReadContext.AssetFileData.bIsReadOnly);
 			}
 		}
 		else if (ReadContext.bCanAttemptAssetRetry)
diff --git a/Engine/Source/Runtime/AssetRegistry/Private/AssetDataGathererDiscoveryCache.cpp b/Engine/Source/Runtime/AssetRegistry/Private/AssetDataGathererDiscoveryCache.cpp
--- a/Engine/Source/Runtime/AssetRegistry/Private/AssetDataGathererDiscoveryCache.cpp
+++ b/Engine/Source/Runtime/AssetRegistry/Private/AssetDataGathererDiscoveryCache.cpp
@@ -27,7 +27,7 @@ FArchive& operator<<(FArchive& Ar, FFileJournalFileHandle& H)
 
 namespace AssetDataGathererConstants
 {
-const FGuid DiscoveryCacheVersion(TEXT("4F4C364CC08C47B9BF18278136E1CB6E"));
+const FGuid DiscoveryCacheVersion(TEXT("415D05FD867949DCA49FFDE505005D49"));
 }
 
 namespace UE::AssetDataGather::Private
@@ -656,6 +656,7 @@ FArchive& operator<<(FArchive& Ar, FCachedDirScanFile& Data)
 {
 	Ar << Data.RelPath;
 	Ar << Data.ModificationTime;
+	Ar << Data.bIsReadOnly;
 	return Ar;
 }
 
diff --git a/Engine/Source/Runtime/AssetRegistry/Private/AssetDataGathererDiscoveryCache.h b/Engine/Source/Runtime/AssetRegistry/Private/AssetDataGathererDiscoveryCache.h
--- a/Engine/Source/Runtime/AssetRegistry/Private/AssetDataGathererDiscoveryCache.h
+++ b/Engine/Source/Runtime/AssetRegistry/Private/AssetDataGathererDiscoveryCache.h
@@ -53,6 +53,7 @@ struct FCachedDirScanFile
 {
 	FString RelPath;
 	FDateTime ModificationTime;
+	bool bIsReadOnly = false;
 };
 FArchive& operator<<(FArchive& Ar, FCachedDirScanFile& Data);
 
diff --git a/Engine/Source/Runtime/AssetRegistry/Private/AssetDataGathererPrivate.h b/Engine/Source/Runtime/AssetRegistry/Private/AssetDataGathererPrivate.h
--- a/Engine/Source/Runtime/AssetRegistry/Private/AssetDataGathererPrivate.h
+++ b/Engine/Source/Runtime/AssetRegistry/Private/AssetDataGathererPrivate.h
@@ -83,16 +83,18 @@ struct FDiscoveredPathData
 	FString RelPath;
 	/** If the path is a file, the modification timestamp of the package file (that it had when it was discovered). */
 	FDateTime PackageTimestamp;
+	/** If the file is readonly */
+	bool bIsReadOnly;
 	/** The type of file that was found */
 	EGatherableFileType Type = EGatherableFileType::Invalid;
 
 	FDiscoveredPathData() = default;
 	FDiscoveredPathData(FStringView InLocalAbsPath, FStringView InLongPackageName, FStringView InRelPath,
-		const FDateTime& InPackageTimestamp, EGatherableFileType InType);
+		const FDateTime& InPackageTimestamp, bool InIsReadonly, EGatherableFileType InType);
 	FDiscoveredPathData(FStringView InLocalAbsPath, FStringView InLongPackageName, FStringView InRelPath,
 		EGatherableFileType InType);
 	void Assign(FStringView InLocalAbsPath, FStringView InLongPackageName, FStringView InRelPath,
-		const FDateTime& InPackageTimestamp, EGatherableFileType InType);
+		const FDateTime& InPackageTimestamp, bool InIsReadOnly, EGatherableFileType InType);
 	void Assign(FStringView InLocalAbsPath, FStringView InLongPackageName, FStringView InRelPath,
 		EGatherableFileType InType);
 
@@ -109,16 +111,18 @@ struct FGatheredPathData
 	FString LongPackageName;
 	/** The modification timestamp of the package file (that it had when it was discovered) */
 	FDateTime PackageTimestamp;
+	/** If the file is readonly */
+	bool bIsReadOnly;
 	/** The type of file that was found */
 	EGatherableFileType Type = EGatherableFileType::Invalid;
 
 	FGatheredPathData() = default;
 	FGatheredPathData(FStringView InLocalAbsPath, FStringView InLongPackageName, const FDateTime& InPackageTimestamp,
-		EGatherableFileType InType);
+		bool InIsReadOnly, EGatherableFileType InType);
 	explicit FGatheredPathData(const FDiscoveredPathData& DiscoveredData);
 	explicit FGatheredPathData(FDiscoveredPathData&& DiscoveredData);
 	void Assign(FStringView InLocalAbsPath, FStringView InLongPackageName, const FDateTime& InPackageTimestamp,
-		EGatherableFileType InType);
+		bool InIsReadOnly, EGatherableFileType InType);
 	void Assign(const FDiscoveredPathData& DiscoveredData);
 
 	/**
@@ -252,6 +256,7 @@ struct FPathExistence
 	const FString& GetLocalAbsPath() const;
 	EType GetType();
 	FDateTime GetModificationTime();
+	bool IsReadOnly();
 
 	/** If the path is a directory or file, return the path. Otherwise return the parent directory, if it exists. */
 	FStringView GetLowestExistingPath();
@@ -265,6 +270,7 @@ private:
 	FString LocalAbsPath;
 	FDateTime ModificationTime;
 	EType PathType = EType::MissingParentDir;
+	bool bIsReadOnly = false;
 	bool bHasExistenceData = false;
 };
 
diff --git a/Engine/Source/Runtime/Core/Private/GenericPlatform/GenericPlatformFile.cpp b/Engine/Source/Runtime/Core/Private/GenericPlatform/GenericPlatformFile.cpp
--- a/Engine/Source/Runtime/Core/Private/GenericPlatform/GenericPlatformFile.cpp
+++ b/Engine/Source/Runtime/Core/Private/GenericPlatform/GenericPlatformFile.cpp
@@ -1057,6 +1057,7 @@ FFileJournalData ToJournalData(const FFileStatData& StatData)
 	JournalData.JournalHandle = FileJournalFileHandleInvalid;
 	JournalData.bIsValid = StatData.bIsValid;
 	JournalData.bIsDirectory = StatData.bIsDirectory;
+	JournalData.bIsReadOnly = StatData.bIsReadOnly;
 	return JournalData;
 }
 
diff --git a/Engine/Source/Runtime/Core/Private/Windows/WindowsPlatformFileJournal.cpp b/Engine/Source/Runtime/Core/Private/Windows/WindowsPlatformFileJournal.cpp
--- a/Engine/Source/Runtime/Core/Private/Windows/WindowsPlatformFileJournal.cpp
+++ b/Engine/Source/Runtime/Core/Private/Windows/WindowsPlatformFileJournal.cpp
@@ -305,6 +305,7 @@ bool FWindowsPlatformFile::FileJournalIterateDirectory(const TCHAR* Directory,
 				FileData.bIsValid = true;
 				FileData.bIsDirectory = !!(Info->FileAttributes & FILE_ATTRIBUTE_DIRECTORY);
 				FileData.bIsReparsePoint = !!(Info->FileAttributes & FILE_ATTRIBUTE_REPARSE_POINT);
+				FileData.bIsReadOnly = !!(Info->FileAttributes & FILE_ATTRIBUTE_READONLY);
 				FileData.JournalHandle = ToFileHandle(Info->FileId);
 
 				if (!Visitor(*LocalAbsPath, FileData))
@@ -410,6 +411,7 @@ FFileJournalData FWindowsPlatformFile::FileJournalGetFileData(const TCHAR* Filen
 
 	OutResult.bIsDirectory = !!(Info->FileAttributes & FILE_ATTRIBUTE_DIRECTORY);
 	OutResult.bIsReparsePoint = !!(Info->FileAttributes & FILE_ATTRIBUTE_REPARSE_POINT);
+	OutResult.bIsReadOnly = !!(Info->FileAttributes & FILE_ATTRIBUTE_READONLY);
 	FILETIME LastWriteTime;
 	LastWriteTime.dwLowDateTime = Info->LastWriteTime.LowPart;
 	LastWriteTime.dwHighDateTime = Info->LastWriteTime.HighPart;
diff --git a/Engine/Source/Runtime/Core/Public/GenericPlatform/GenericPlatformFile.h b/Engine/Source/Runtime/Core/Public/GenericPlatform/GenericPlatformFile.h
--- a/Engine/Source/Runtime/Core/Public/GenericPlatform/GenericPlatformFile.h
+++ b/Engine/Source/Runtime/Core/Public/GenericPlatform/GenericPlatformFile.h
@@ -273,6 +273,7 @@ struct FFileJournalData
 	bool bIsValid : 1;
 	bool bIsDirectory : 1;
 	bool bIsReparsePoint : 1;
+	bool bIsReadOnly : 1;
 };
 
 /** Stores custom error messages from the engine along with an optional system error code that can provider more detailed information */
@@ -1042,6 +1043,7 @@ inline FFileJournalData::FFileJournalData()
 	: bIsValid(false)
 	, bIsDirectory(false)
 	, bIsReparsePoint(false)
+	, bIsReadOnly(false)
 {
 }
 
diff --git a/Engine/Source/Runtime/CoreUObject/Private/Serialization/PackageStore.cpp b/Engine/Source/Runtime/CoreUObject/Private/Serialization/PackageStore.cpp
--- a/Engine/Source/Runtime/CoreUObject/Private/Serialization/PackageStore.cpp
+++ b/Engine/Source/Runtime/CoreUObject/Private/Serialization/PackageStore.cpp
@@ -2,14 +2,17 @@
 
 #include "Serialization/PackageStore.h"
 
+#include "AssetRegistry/AssetData.h"
 #include "Containers/StringConv.h"
 #include "Containers/StringView.h"
 #include "Containers/UnrealString.h"
 #include "CoreGlobals.h"
 #include "CoreGlobalsInternal.h"
 #include "HAL/Platform.h"
+#include "HAL/PlatformFileManager.h"
 #include "IO/PackageId.h"
 #include "Misc/AssertionMacros.h"
+#include "Misc/AssetRegistryInterface.h"
 #include "Misc/ConfigCacheIni.h"
 #include "Misc/CoreDelegates.h"
 #include "Misc/PackageName.h"
@@ -20,13 +23,17 @@
 #include "Templates/Greater.h"
 
 
+DEFINE_LOG_CATEGORY_STATIC(LogPackageStore, Log, All);
+
+
 ///////////////////////////////////////
 // FLinkerLoadPackageStoreBackend
 ///////////////////////////////////////
 
 class FLinkerLoadPackageStoreBackend final : public IPackageStoreBackend
 {
-	EPackageLoader GetSupportedLoaders()
+public:
+	virtual EPackageLoader GetSupportedLoaders() override
 	{
 		return EPackageLoader::LinkerLoad;
 	}
@@ -43,7 +50,7 @@ class FLinkerLoadPackageStoreBackend final : public IPackageStoreBackend
 	{
 	}
 
-	virtual EPackageStoreEntryStatus GetPackageStoreEntry(FPackageId PackageId, FName PackageName, FPackageStoreEntry& OutPackageStoreEntry) override
+	EPackageStoreEntryStatus GetPackageStoreEntryEx(FPackageId PackageId, FName PackageName, FPackageStoreEntry& OutPackageStoreEntry, FPackagePath* OutPackagePath)
 	{
 		FPackageStoreEntry* Entry = PackageStoreCache.Find(PackageId);
 		if (Entry)
@@ -51,6 +58,11 @@ class FLinkerLoadPackageStoreBackend final : public IPackageStoreBackend
 			if (Entry->LoaderType == EPackageLoader::LinkerLoad)
 			{
 				OutPackageStoreEntry = *Entry;
+				if (OutPackagePath)
+				{
+					*OutPackagePath = FPackagePath::FromPackageNameUnchecked(PackageName);
+					OutPackagePath->SetHeaderExtension(OutPackageStoreEntry.PackageExtension);
+				}
 				return EPackageStoreEntryStatus::Ok;
 			}
 			return EPackageStoreEntryStatus::Missing;
@@ -68,17 +80,20 @@ class FLinkerLoadPackageStoreBackend final : public IPackageStoreBackend
 		constexpr bool bMatchCaseOnDisk = false;
 #endif
 
-		FPackagePath PackagePath = FPackagePath::FromPackageNameUnchecked(PackageName);
+		FPackagePath LocalPackagePath;
+		FPackagePath* PackagePath = OutPackagePath ? OutPackagePath : &LocalPackagePath;
+
+		*PackagePath = FPackagePath::FromPackageNameUnchecked(PackageName);
 		// now, check to see if the package exists in the local file system, and if it does, correct the case as needed
-		if (FPackageName::DoesPackageExistEx(PackagePath, FPackageName::EPackageLocationFilter::FileSystem,
-			bMatchCaseOnDisk, &PackagePath) != FPackageName::EPackageLocationFilter::None)
+		if (FPackageName::DoesPackageExistEx(*PackagePath, FPackageName::EPackageLocationFilter::FileSystem,
+			bMatchCaseOnDisk, PackagePath) != FPackageName::EPackageLocationFilter::None)
 		{
 			OutPackageStoreEntry.LoaderType = EPackageLoader::LinkerLoad;
-			OutPackageStoreEntry.PackageExtension = PackagePath.GetHeaderExtension();
+			OutPackageStoreEntry.PackageExtension = PackagePath->GetHeaderExtension();
 			check(OutPackageStoreEntry.PackageExtension != EPackageExtension::Unspecified);
 
 #if WITH_EDITORONLY_DATA
-			OutPackageStoreEntry.LinkerLoadCaseCorrectedPackageName = PackagePath.GetPackageFName();
+			OutPackageStoreEntry.LinkerLoadCaseCorrectedPackageName = PackagePath->GetPackageFName();
 #endif
 
 			PackageStoreCache.Add(PackageId, OutPackageStoreEntry);
@@ -89,6 +104,11 @@ class FLinkerLoadPackageStoreBackend final : public IPackageStoreBackend
 		return EPackageStoreEntryStatus::Missing;
 	}
 
+	virtual EPackageStoreEntryStatus GetPackageStoreEntry(FPackageId PackageId, FName PackageName, FPackageStoreEntry& OutPackageStoreEntry) override
+	{
+		return GetPackageStoreEntryEx(PackageId, PackageName, OutPackageStoreEntry, nullptr);
+	}
+
 	virtual bool GetPackageRedirectInfo(FPackageId PackageId, FName& OutSourcePackageName, FPackageId& OutRedirectedToPackageId) override
 	{
 		return false;
@@ -142,17 +162,70 @@ bool FHybridPackageStoreBackend::GetPackageRedirectInfo(FPackageId PackageId, FN
 	return CookedPackageStore->GetPackageRedirectInfo(PackageId, OutSourcePackageName, OutRedirectedToPackageId);
 }
 
-
 EPackageStoreEntryStatus FHybridPackageStoreBackend::GetPackageStoreEntry(FPackageId PackageId, FName PackageName, FPackageStoreEntry& OutPackageStoreEntry)
 {
+	TRACE_CPUPROFILER_EVENT_SCOPE(HybridCookedEditor_GetPackageStoreEntry);
+
 	EPackageStoreEntryStatus InnerStatus = EPackageStoreEntryStatus::Missing;
 
-	if (!ForcedLoosePackages.Contains(PackageId))
+	bool bAllowZenLoad = true;
+	// check if we are forcing this package to load uncooked
+	if (ForcedLoosePackages.Contains(PackageId))
+	{
+		bAllowZenLoad = false;
+	}
+	else
+	{
+		// get the readonly flag (without slowly touching the disk for each package)
+		FAssetPackageData AssetPackageData;
+		UE::AssetRegistry::EExists Result = IAssetRegistryInterface::GetPtr()->TryGetAssetPackageData(PackageName, AssetPackageData);
+		
+		
+		// if we did find it in the AR, and it's writeable, assume user wants local changes, so load loose
+		//
+		// if we don't know the state yet, then instead of just assuming the worst and loading the loose version,
+		// use the process of querying the loose packagestore, but also get the path, and then we can check the 
+		// physical disk for readonly state, and use that to determine if we should load cooked or not (as if 
+		// we had found the ReadOnly state in the AR, as above)
+		//
+		// otherwise, (it's readonly, or doesn't exist locally), we want to load cooked
+		if (Result == UE::AssetRegistry::EExists::Exists && !AssetPackageData.IsReadOnly())
+		{
+			bAllowZenLoad = false;
+		}
+		if (Result == UE::AssetRegistry::EExists::Unknown)
+		{
+			FPackageStoreEntry TempLooseEntry;
+			FPackagePath TempPath;
+			TSharedPtr<FLinkerLoadPackageStoreBackend> LinkerLoadStore = StaticCastSharedPtr<FLinkerLoadPackageStoreBackend>(LoosePackageStore);
+			// check if we can find this package on disk, getting the store entry, while also getting the FPackagePath.
+			// this way, we can just return the PackageEntry directly
+			if (LinkerLoadStore->GetPackageStoreEntryEx(PackageId, PackageName, TempLooseEntry, &TempPath) == EPackageStoreEntryStatus::Ok)
+			{
+				// now that we have determined the local path to the file, check if it's read only
+				const bool bIsReadOnly = FPlatformFileManager::Get().GetPlatformFile().IsReadOnly(*TempPath.GetLocalFullPath());
+				if (!bIsReadOnly)
+				{
+					// if it's writeable, then we want to use the LinkerLoad version, so just return it
+					OutPackageStoreEntry = TempLooseEntry;
+
+					// we are done!
+					InnerStatus = EPackageStoreEntryStatus::Ok;
+					bAllowZenLoad = false;
+				}
+			}
+		}
+
+		UE_LOG(LogPackageStore, Verbose, TEXT("Hybrid Cooked Editor Discovery: Package: %s, AssetPackageData.Exists: %d, bAllowZen: %d"), 
+			*PackageName.ToString(), (int)Result, bAllowZenLoad);
+	}
+
+	if (bAllowZenLoad)
 	{
 		InnerStatus = CookedPackageStore->GetPackageStoreEntry(PackageId, PackageName, OutPackageStoreEntry);
 	}
 
-	if (InnerStatus == EPackageStoreEntryStatus::Missing && FAsyncLoadingThreadSettings::Get().bLooseFileLoadingEnabled)
+	if (InnerStatus == EPackageStoreEntryStatus::Missing)
 	{
 		// reset it in case the cooked store partially filled this out
 		OutPackageStoreEntry = {};
@@ -376,18 +449,33 @@ void FPackageStore::Mount(TSharedRef<IPackageStoreBackend> Backend, int32 Priori
 	// the loose package store
 	if (IsRunningHybridCookedEditor())
 	{
-		if (!HybridBackend && LinkerLoadBackend && !!(Backend->GetSupportedLoaders() & EPackageLoader::Zen))
+		if (!HybridBackend)
 		{
-			// remove the loose backend
-			Backends.RemoveAll([this](const FBackendAndPriority& Pair) { return Pair.Value == LinkerLoadBackend; });
+			// the order of mounting could be LinkerLoad,Zen or Zen,LinkerLoad, so check all combos
+			TSharedPtr<IPackageStoreBackend> LinkerLoad = LinkerLoadBackend;
+			if (!LinkerLoad && !!(Backend->GetSupportedLoaders() & EPackageLoader::LinkerLoad))
+			{
+				LinkerLoad = Backend;
+			}
+			TSharedPtr<IPackageStoreBackend> Zen = Backend;
+			if (!(Zen->GetSupportedLoaders() & EPackageLoader::Zen) && !!(Backends[0].Value->GetSupportedLoaders() & EPackageLoader::Zen))
+			{
+				Zen = Backends[0].Value;
+			}
+
+			if (LinkerLoad.IsValid() && Zen.IsValid())
+			{
+				// remove the loose backend
+				Backends.RemoveAll([LinkerLoad, Zen](const FBackendAndPriority& Pair) { return Pair.Value == LinkerLoad || Pair.Value == Zen; });
 
-			// make a new hybrid backend which wraps the existing loose and the new backend 
-			HybridBackend = MakeShared<FHybridPackageStoreBackend>(LinkerLoadBackend, Backend);
-			Mount(HybridBackend.ToSharedRef(), Priority);
+				// make a new hybrid backend which wraps the existing loose and the new backend 
+				HybridBackend = MakeShared<FHybridPackageStoreBackend>(LinkerLoad, Zen);
+				Mount(HybridBackend.ToSharedRef(), Priority);
 
-			// call callbacks 
-			Backend->OnMounted(BackendContext);
-			return;
+				// call callbacks 
+				Backend->OnMounted(BackendContext);
+				return;
+			}
 		}
 	}
 #endif
diff --git a/Engine/Source/Runtime/CoreUObject/Public/AssetRegistry/AssetData.h b/Engine/Source/Runtime/CoreUObject/Public/AssetRegistry/AssetData.h
--- a/Engine/Source/Runtime/CoreUObject/Public/AssetRegistry/AssetData.h
+++ b/Engine/Source/Runtime/CoreUObject/Public/AssetRegistry/AssetData.h
@@ -1017,6 +1017,10 @@ public:
 	bool HasVirtualizedPayloads() const { return (Flags & FLAG_HAS_VIRTUALIZED_PAYLOADS) != 0; }
 	void SetHasVirtualizedPayloads(bool bValue) { Flags = (Flags & ~FLAG_HAS_VIRTUALIZED_PAYLOADS) | (bValue ? FLAG_HAS_VIRTUALIZED_PAYLOADS : 0); }
 
+	/** Whether the package is readonly or not */
+	bool IsReadOnly() const { return (Flags & FLAG_PACKAGE_READ_ONLY) != 0; }
+	void SetIsReadOnly(bool bValue) { Flags = (Flags & ~FLAG_PACKAGE_READ_ONLY) | (bValue ? FLAG_PACKAGE_READ_ONLY : 0); }
+
 	/**
 	 * Serialize as part of the registry cache. This is not meant to be serialized as part of a package so it does not handle versions normally
 	 * To version this data change FAssetRegistryVersion
@@ -1068,6 +1072,7 @@ private:
 	{
 		FLAG_LICENSEE_VERSION						= 1u << 0,
 		FLAG_HAS_VIRTUALIZED_PAYLOADS				= 1u << 1,
+		FLAG_PACKAGE_READ_ONLY						= 1u << 2,
 
 		FLAG_PACKAGE_LOCATION_IODISPATCHER			= 1u << 30,
 		FLAG_PACKAGE_LOCATION_FILESYSTEM			= 1u << 31,
