[AssetRegistry]

- DependencyDataGathered will aggregate dependency information when the dependency results to process contain more than one entry for the same source node.
-- This prevents inadvertently accessing the FDependsNodes later in parallel

#rnx

diff --git a/Engine/Source/Runtime/AssetRegistry/Private/AssetRegistry.cpp b/Engine/Source/Runtime/AssetRegistry/Private/AssetRegistry.cpp
--- a/Engine/Source/Runtime/AssetRegistry/Private/AssetRegistry.cpp
+++ b/Engine/Source/Runtime/AssetRegistry/Private/AssetRegistry.cpp
@@ -7134,13 +7134,6 @@ struct FConstructData
 	FDependsNode* SourceNode;
 	FPackageDependencyMap PackageDependencies;
 	FDependsNode::FDependsNodeList SearchableNamesNodes;
-
-	FConstructData(FDependsNode* InSourceNode, FPackageDependencyMap InPackageDependencies, FDependsNode::FDependsNodeList InSearchableNamesNodes)
-		: SourceNode(InSourceNode)
-		, PackageDependencies(MoveTemp(InPackageDependencies))
-		, SearchableNamesNodes(MoveTemp(InSearchableNamesNodes))
-	{
-	}
 };
 }
 
@@ -7174,11 +7167,11 @@ void FAssetRegistryImpl::DependencyDataGathered(TMultiMap<FName, FPackageDepende
 		Utils::InitializeMountPoints(MountPointsForVerifyAfterGather);
 	}
 
-	TArray<FConstructData> ConstructDatas;
-	TMap<FName, FName> CachedDepToRedirect;
-	TSet<FDependsNode*> ModifiedDependNodes;
+	TMap<FName, FConstructData, TInlineSetAllocator<1>> AssetNameToConstructDataMap;
+	TMap<FName, FName, TInlineSetAllocator<1>> CachedDepToRedirect;
+	TSet<FDependsNode*, DefaultKeyFuncs<FDependsNode*>, TInlineSetAllocator<1>> ModifiedDependNodes;
 
-	ConstructDatas.Reserve(DependsResults.Num());
+	AssetNameToConstructDataMap.Reserve(DependsResults.Num());
 	CachedDepToRedirect.Reserve(DependsResults.Num());
 	ModifiedDependNodes.Reserve(DependsResults.Num());
 
@@ -7188,10 +7181,7 @@ void FAssetRegistryImpl::DependencyDataGathered(TMultiMap<FName, FPackageDepende
 	bool bInterrupted = false;
 	int32 IterationCounter = 0;
 	constexpr int32 MinBatchSize = 1024;
-
-	// Disable parallelism until UE-295029 is resolved
-	const bool bUseJobs = false;
-	//const bool bUseJobs = DependsResults.Num() > MinBatchSize;
+	const bool bUseJobs = DependsResults.Num() > MinBatchSize;
 	const EParallelForFlags ParallelForFlags = bUseJobs ? EParallelForFlags::None : EParallelForFlags::ForceSingleThread;
 
 	{
@@ -7210,17 +7200,6 @@ void FAssetRegistryImpl::DependencyDataGathered(TMultiMap<FName, FPackageDepende
 				}
 			};
 
-			bool bSourceNodeAlreadyExisted = false;
-			FName SourceNodePackageName = Iter.Value().PackageName;
-			FDependsNode* SourceNode = State.CreateOrFindDependsNode(SourceNodePackageName, &bSourceNodeAlreadyExisted);
-			if (!bSourceNodeAlreadyExisted)
-			{
-				if (FPackageDependencyData* DependencyDependencyData = DependsResults.Find(SourceNodePackageName))
-				{
-					SourceNode->Reserve(DependencyDependencyData->DependsNodeReservations);
-				}
-			}
-
 			if (DeferredAssets.Contains(Iter.Key()) || DeferredAssetsForGameThread.Contains(Iter.Key()))
 			{
 				OutDeferredDependencyResults.Add(MoveTemp(*Iter));
@@ -7231,15 +7210,50 @@ void FAssetRegistryImpl::DependencyDataGathered(TMultiMap<FName, FPackageDepende
 			FPackageDependencyData Result = MoveTemp(Iter.Value());
 			Iter.RemoveCurrent();
 
-			if (bLocalVerifyMountPointAfterGather)
+
+			// DependsResults might contain multiple values with the same SourcePackageName. 
+			// We don't want to create more constructed nodes than necessary, so we create
+			// or find existing FConstructData and append results as needed.
+			bool bSourceNodeAlreadyExisted = false;
+			FName SourceNodePackageName = Iter.Value().PackageName;
+			FConstructData* ConstructData = nullptr;
+			if (FConstructData* ExistingConstructData = AssetNameToConstructDataMap.Find(SourceNodePackageName))
 			{
-				SourceNodePackageName.ToString(PackageNameString);
-				if (!Utils::IsPathMounted(PackageNameString, *MountPointsForVerifyAfterGather, PackageRoot))
+				ConstructData = ExistingConstructData;
+			}
+			else
+			{
+				if (bLocalVerifyMountPointAfterGather)
 				{
-					UE_LOG(LogAssetRegistry, Display,
-						TEXT("AssetRegistry: DependencyData has been loaded with an invalid mount point: '%s', Mount Point: '%s'. Ignoring the DependencyData."),
-						*PackageNameString, *PackageRoot);
-					continue;
+					SourceNodePackageName.ToString(PackageNameString);
+					if (!Utils::IsPathMounted(PackageNameString, *MountPointsForVerifyAfterGather, PackageRoot))
+					{
+						UE_LOG(LogAssetRegistry, Display,
+							TEXT("AssetRegistry: DependencyData has been loaded with an invalid mount point: '%s', Mount Point: '%s'. Ignoring the DependencyData."),
+							*PackageNameString, *PackageRoot);
+						continue;
+					}
+				}
+
+				FDependsNode* SourceNode = State.CreateOrFindDependsNode(SourceNodePackageName, &bSourceNodeAlreadyExisted);
+				if (!bSourceNodeAlreadyExisted)
+				{
+					if (FPackageDependencyData* DependencyDependencyData = DependsResults.Find(SourceNodePackageName))
+					{
+						SourceNode->Reserve(DependencyDependencyData->DependsNodeReservations);
+					}
+
+					if (Result.bHasDependencyData)
+					{
+						SourceNode->Reserve(Result.DependsNodeReservations);
+						SourceNode->SetAllowShrinking(false);
+						SourceNode->SetIsReferencersSorted(bShouldSortReferencers);
+
+						ModifiedDependNodes.Add(SourceNode);
+
+						ConstructData = &AssetNameToConstructDataMap.Add(SourceNodePackageName);
+						ConstructData->SourceNode = SourceNode;
+					}
 				}
 			}
 
@@ -7251,11 +7265,11 @@ void FAssetRegistryImpl::DependencyDataGathered(TMultiMap<FName, FPackageDepende
 				*PackageData = Result.PackageData;
 			}
 
-			if (Result.bHasDependencyData)
+			if (ConstructData)
 			{
 				OutPackagesNeedingDependencyCalculation->Add(SourceNodePackageName);
 
-				FPackageDependencyMap PackageDependencies;
+				FPackageDependencyMap& PackageDependencies = ConstructData->PackageDependencies;
 				PackageDependencies.Reserve(Result.PackageDependencies.Num());
 				for (const FPackageDependencyData::FPackageDependency& DependencyData : Result.PackageDependencies)
 				{
@@ -7317,7 +7331,7 @@ void FAssetRegistryImpl::DependencyDataGathered(TMultiMap<FName, FPackageDepende
 					FlagSet.Add(FDependsNode::PackagePropertiesToByte(DependencyData.Property));
 				}
 
-				FDependsNode::FDependsNodeList SearchNameDependencyNodes;
+				FDependsNode::FDependsNodeList& SearchNameDependencyNodes = ConstructData->SearchableNamesNodes;
 				for (const FPackageDependencyData::FSearchableNamesDependency& NamesDependency : Result.SearchableNameDependencies)
 				{
 					SearchNameDependencyNodes.Reserve(SearchNameDependencyNodes.Num() + NamesDependency.ValueNames.Num());
@@ -7330,31 +7344,30 @@ void FAssetRegistryImpl::DependencyDataGathered(TMultiMap<FName, FPackageDepende
 						SearchNameDependencyNodes.Add(DependsNode);
 					}
 				}
-
-				ModifiedDependNodes.Add(SourceNode);
-
-				SourceNode->Reserve(Result.DependsNodeReservations);
-				SourceNode->SetAllowShrinking(false);
-				SourceNode->SetIsReferencersSorted(bShouldSortReferencers);
-
-				ConstructDatas.Emplace(SourceNode, MoveTemp(PackageDependencies), MoveTemp(SearchNameDependencyNodes));
 			}
 		}
 	}
 
-	if (ConstructDatas.IsEmpty())
+	if (AssetNameToConstructDataMap.IsEmpty())
 	{
 		return;
 	}
 
+	TArray<FConstructData*, TInlineAllocator<1>> ConstructDatas;
+	ConstructDatas.Reserve(AssetNameToConstructDataMap.Num());
+	for (TPair<FName, FConstructData>& Pair : AssetNameToConstructDataMap)
+	{
+		ConstructDatas.Add(&Pair.Value);
+	}
+
 	{
 		TRACE_CPUPROFILER_EVENT_SCOPE(DependencyGatheredData_RemoveReferencers);
 		// Remove referencers - Single threaded due to the poor memory access patterns 
 		// between nodes and their dependencies
-		for (FConstructData& ConstructData : ConstructDatas)
+		for (FConstructData* ConstructData : ConstructDatas)
 		{
-			FDependsNode* SourceNode = ConstructData.SourceNode;
-			const FPackageDependencyMap& PackageDependencies = ConstructData.PackageDependencies;
+			FDependsNode* SourceNode = ConstructData->SourceNode;
+			const FPackageDependencyMap& PackageDependencies = ConstructData->PackageDependencies;
 
 			// We will populate the node dependencies below. Empty the set here in case this file was already read
 			// Also remove references to all existing dependencies, those will be also repopulated below
@@ -7379,10 +7392,10 @@ void FAssetRegistryImpl::DependencyDataGathered(TMultiMap<FName, FPackageDepende
 	ParallelFor(TEXT("DependencyGatheredData_AddDependencies"), ConstructDatas.Num(), MinBatchSize, 
 		[ShouldSortDependencyList = ShouldSortDependencies(), &ConstructDatas] (int Index)
 		{
-			FConstructData& ConstructData = ConstructDatas[Index];
-			FDependsNode* SourceNode = ConstructData.SourceNode;
-			const FPackageDependencyMap& PackageDependencies = ConstructData.PackageDependencies;
-			const FDependsNode::FDependsNodeList& SearchableNameNodes = ConstructData.SearchableNamesNodes;
+			FConstructData* ConstructData = ConstructDatas[Index];
+			FDependsNode* SourceNode = ConstructData->SourceNode;
+			const FPackageDependencyMap& PackageDependencies = ConstructData->PackageDependencies;
+			const FDependsNode::FDependsNodeList& SearchableNameNodes = ConstructData->SearchableNamesNodes;
 
 			SourceNode->ClearDependencies();
 			SourceNode->SetIsDependencyListSorted(EDependencyCategory::All, ShouldSortDependencyList);
@@ -7404,11 +7417,11 @@ void FAssetRegistryImpl::DependencyDataGathered(TMultiMap<FName, FPackageDepende
 	// Add referencers - Single threaded due to the poor memory access patterns between nodes and their dependencies
 	{
 		TRACE_CPUPROFILER_EVENT_SCOPE(DependencyGatheredData_AddReferencers);
-		for (FConstructData& ConstructData : ConstructDatas)
+		for (FConstructData* ConstructData : ConstructDatas)
 		{
-			FDependsNode* SourceNode = ConstructData.SourceNode;
-			const FPackageDependencyMap& PackageDependencies = ConstructData.PackageDependencies;
-			const FDependsNode::FDependsNodeList& SearchableNameNodes = ConstructData.SearchableNamesNodes;
+			FDependsNode* SourceNode = ConstructData->SourceNode;
+			const FPackageDependencyMap& PackageDependencies = ConstructData->PackageDependencies;
+			const FDependsNode::FDependsNodeList& SearchableNameNodes = ConstructData->SearchableNamesNodes;
 
 			// Doubly-link all of the PackageDependencies
 			for (const TPair<FDependsNode*, FDependsNode::FPackageFlagSet>& Pair : PackageDependencies)
