[AssetRegistry]

- If a journal cache has an error while reading, rename the cache on disk so it may be inspected for a post-mortem
- Re-enables the DiscoveryCache on non-build machines

#rnx

diff --git a/Engine/Config/BaseEngine.ini b/Engine/Config/BaseEngine.ini
--- a/Engine/Config/BaseEngine.ini
+++ b/Engine/Config/BaseEngine.ini
@@ -2962,7 +2962,7 @@ GatherOnGameThreadOnly=false
 ; by setting AssetRegistryDiscoveryCache=AlwaysWrite, and it can be force-read with no test for invalidation due to modified
 ; directories by commandline -AssetRegistryCacheSkipInvalidate.
 ; Expected values for AssetRegistryDiscoveryCache are Never,Default,AlwaysWrite.
-AssetRegistryDiscoveryCache=Never
+AssetRegistryDiscoveryCache=Default
 
 [AutomationTesting]
 ImportTestPath=Content/EditorAutomation/
diff --git a/Engine/Source/Runtime/AssetRegistry/Private/AssetDataGatherer.cpp b/Engine/Source/Runtime/AssetRegistry/Private/AssetDataGatherer.cpp
--- a/Engine/Source/Runtime/AssetRegistry/Private/AssetDataGatherer.cpp
+++ b/Engine/Source/Runtime/AssetRegistry/Private/AssetDataGatherer.cpp
@@ -190,6 +190,13 @@ void FGatherSettings::Initialize()
 #endif
 	bool bNoAssetRegistryCache = FParse::Param(FCommandLine::Get(), TEXT("NoAssetRegistryCache"));
 	bool bNoAssetRegistryDiscoveryCache = bNoAssetRegistryCache || FParse::Param(FCommandLine::Get(), TEXT("NoAssetRegistryDiscoveryCache"));
+
+	// Disable DiscoveryCache on build machines until UE-290798 is resolved
+	if (GIsBuildMachine)
+	{
+		bNoAssetRegistryDiscoveryCache = true;
+	}
+
 	bool bNoAssetRegistryCacheRead = FParse::Param(FCommandLine::Get(), TEXT("NoAssetRegistryCacheRead"));
 	uint32 MultiprocessId = UE::GetMultiprocessId();
 	bool bMultiprocess = MultiprocessId > 0 || FParse::Param(FCommandLine::Get(), TEXT("multiprocess"));
diff --git a/Engine/Source/Runtime/AssetRegistry/Private/AssetDataGathererDiscoveryCache.cpp b/Engine/Source/Runtime/AssetRegistry/Private/AssetDataGathererDiscoveryCache.cpp
--- a/Engine/Source/Runtime/AssetRegistry/Private/AssetDataGathererDiscoveryCache.cpp
+++ b/Engine/Source/Runtime/AssetRegistry/Private/AssetDataGathererDiscoveryCache.cpp
@@ -196,6 +196,19 @@ void FAssetDataDiscoveryCache::LoadAndUpdateCache()
 			VolumeInfo.JournalId = VolumeInfo.JournalIdOnDisk;
 			VolumeInfo.NextJournalEntryToScan = VolumeInfo.NextJournalEntryOnDisk;
 			VolumeInfo.Dirs.Empty();
+
+#if WITH_EDITOR
+			// Rename the discovery cache so it may remain on disk for inspection
+			FString CachePath = GetCacheFileName();
+			FStringView CacheDir, CacheFileName, CacheExtension;
+			FPathViews::Split(CachePath, CacheDir, CacheFileName, CacheExtension);
+			FString NewCacheFileName = FString::Printf(TEXT("%.*s-%s.%.*s"), CacheFileName.Len(), CacheFileName.GetData(), *FDateTime::Now().ToString(), CacheExtension.Len(), CacheExtension.GetData());
+			FString NewCachePath = FPaths::Combine(CacheDir, NewCacheFileName);
+			if (!IFileManager::Get().Move(*NewCachePath, *CachePath, true /*Replace*/, true /*EvenIfReadOnly*/))
+			{
+				UE_LOG(LogAssetRegistry, Warning, TEXT("Failed creating a copy of %s at %s"), *CachePath, *NewCachePath);
+			}
+#endif
 		}
 		else
 		{
@@ -587,6 +600,7 @@ void FCachedVolumeInfo::InitializePlatformData()
 		if (NextJournalEntryToScan == FileJournalEntryHandleInvalid)
 		{
 			JournalId = JournalIdOnDisk;
+			check(NextJournalEntryOnDisk != FileJournalEntryHandleInvalid);
 			NextJournalEntryToScan = NextJournalEntryOnDisk;
 		}
 	}
