AssetRegistry: Remove deprecation handling from 5.5 and earlier.
#rnx

diff --git a/Engine/Source/Runtime/AssetRegistry/Private/AssetRegistry.cpp b/Engine/Source/Runtime/AssetRegistry/Private/AssetRegistry.cpp
--- a/Engine/Source/Runtime/AssetRegistry/Private/AssetRegistry.cpp
+++ b/Engine/Source/Runtime/AssetRegistry/Private/AssetRegistry.cpp
@@ -2449,33 +2449,6 @@ bool UAssetRegistryImpl::GetInMemoryAssets(const FARCompiledFilter& CompiledFilt
 	return true;
 }
 
-bool UAssetRegistryImpl::EnumerateAssets(const FARFilter& InFilter, TFunctionRef<bool(const FAssetData&)> Callback,
-	bool bSkipARFilteredAssets) const
-{
-	FARCompiledFilter CompiledFilter;
-	CompileFilter(InFilter, CompiledFilter);
-	return EnumerateAssets(CompiledFilter, Callback, bSkipARFilteredAssets);
-}
-
-bool UAssetRegistryImpl::EnumerateAssets(const FARCompiledFilter& InFilter, TFunctionRef<bool(const FAssetData&)> Callback,
-	bool bSkipARFilteredAssets) const
-{
-	const UE::AssetRegistry::EEnumerateAssetsFlags Flags = bSkipARFilteredAssets ? UE::AssetRegistry::EEnumerateAssetsFlags::None : UE::AssetRegistry::EEnumerateAssetsFlags::AllowUnfilteredArAssets;
-	return EnumerateAssets(InFilter, Callback, Flags);
-}
-
-bool UAssetRegistryImpl::EnumerateAssets(const FARFilter& InFilter, TFunctionRef<bool(const FAssetData&)> Callback) const
-{
-	FARCompiledFilter CompiledFilter;
-	CompileFilter(InFilter, CompiledFilter);
-	return EnumerateAssets(CompiledFilter, Callback, UE::AssetRegistry::EEnumerateAssetsFlags::None);
-}
-
-bool UAssetRegistryImpl::EnumerateAssets(const FARCompiledFilter& InFilter, TFunctionRef<bool(const FAssetData&)> Callback) const
-{
-	return EnumerateAssets(InFilter, Callback, UE::AssetRegistry::EEnumerateAssetsFlags::None);
-}
-
 bool UAssetRegistryImpl::EnumerateAssets(const FARFilter& InFilter, TFunctionRef<bool(const FAssetData&)> Callback,
 	UE::AssetRegistry::EEnumerateAssetsFlags InEnumerateFlags) const
 {
@@ -2960,17 +2933,6 @@ bool UAssetRegistryImpl::GetAllAssets(TArray<FAssetData>& OutAssetData, bool bIn
 	return true;
 }
 
-bool UAssetRegistryImpl::EnumerateAllAssets(TFunctionRef<bool(const FAssetData&)> Callback) const
-{
-	return EnumerateAllAssets(Callback, UE::AssetRegistry::EEnumerateAssetsFlags::None);
-}
-
-bool UAssetRegistryImpl::EnumerateAllAssets(TFunctionRef<bool(const FAssetData&)> Callback, bool bIncludeOnlyOnDiskAssets) const
-{
-	const UE::AssetRegistry::EEnumerateAssetsFlags Flags = bIncludeOnlyOnDiskAssets ? UE::AssetRegistry::EEnumerateAssetsFlags::OnlyOnDiskAssets : UE::AssetRegistry::EEnumerateAssetsFlags::None;
-	return EnumerateAllAssets(Callback, Flags);
-}
-
 bool UAssetRegistryImpl::EnumerateAllAssets(TFunctionRef<bool(const FAssetData&)> Callback, UE::AssetRegistry::EEnumerateAssetsFlags InEnumerateFlags) const
 {
 	TSet<FName> PackageNamesToSkip;
@@ -4502,11 +4464,6 @@ void UAssetRegistryImpl::AssetsSaved(TArray<FAssetData>&& Assets)
 #endif
 }
 
-void UAssetRegistryImpl::AssetFullyUpdateTags(UObject* Object)
-{
-	AssetUpdateTags(Object, EAssetRegistryTagsCaller::Fast);
-}
-
 void UAssetRegistryImpl::AssetUpdateTags(UObject* Object, EAssetRegistryTagsCaller Caller)
 {
 #if WITH_EDITOR
@@ -10024,19 +9981,6 @@ void FEventContext::Append(FEventContext&& Other)
 
 }
 
-void UAssetRegistryImpl::ReadLockEnumerateTagToAssetDatas(TFunctionRef<void(FName TagName, const TArray<const FAssetData*>& Assets)> Callback) const
-{
-	UE_LOG(LogAssetRegistry, Error, TEXT("ReadLockEnumerateTagToAssetDatas has been deprecated. Use ReadLockEnumerateAllTagToAssetDatas instead."));
-
-	UE::AssetRegistry::FInterfaceReadScopeLock InterfaceScopeLock(InterfaceLock);
-	GuardedData.GetState().EnumerateTags([Callback](FName TagName)
-	{
-		TArray<const FAssetData*> EmptyArray;
-		Callback(TagName, EmptyArray);
-		return true;
-	});
-}
-
 void UAssetRegistryImpl::ReadLockEnumerateAllTagToAssetDatas(TFunctionRef<bool(FName TagName, IAssetRegistry::FEnumerateAssetDatasFunc EnumerateAssets)> Callback) const
 {
 	UE::AssetRegistry::FInterfaceReadScopeLock InterfaceScopeLock(InterfaceLock);
diff --git a/Engine/Source/Runtime/AssetRegistry/Private/AssetRegistry.h b/Engine/Source/Runtime/AssetRegistry/Private/AssetRegistry.h
--- a/Engine/Source/Runtime/AssetRegistry/Private/AssetRegistry.h
+++ b/Engine/Source/Runtime/AssetRegistry/Private/AssetRegistry.h
@@ -41,12 +41,10 @@ public:
 	virtual bool GetAssets(const FARCompiledFilter& Filter, TArray<FAssetData>& OutAssetData, bool bSkipARFilteredAssets = true) const override;
 	virtual bool GetInMemoryAssets(const FARFilter& Filter, TArray<FAssetData>& OutAssetData, bool bSkipARFilteredAssets=true) const override;
 	virtual bool GetInMemoryAssets(const FARCompiledFilter& Filter, TArray<FAssetData>& OutAssetData, bool bSkipARFilteredAssets=true) const override;
-	virtual bool EnumerateAssets(const FARFilter& Filter, TFunctionRef<bool(const FAssetData&)> Callback, bool bSkipARFilteredAssets) const override;
-	virtual bool EnumerateAssets(const FARCompiledFilter& Filter, TFunctionRef<bool(const FAssetData&)> Callback, bool bSkipARFilteredAssets) const override;
-	virtual bool EnumerateAssets(const FARFilter& Filter, TFunctionRef<bool(const FAssetData&)> Callback) const override;
-	virtual bool EnumerateAssets(const FARCompiledFilter& Filter, TFunctionRef<bool(const FAssetData&)> Callback) const override;
-	virtual bool EnumerateAssets(const FARFilter& Filter, TFunctionRef<bool(const FAssetData&)> Callback, UE::AssetRegistry::EEnumerateAssetsFlags InEnumerateFlags) const override;
-	virtual bool EnumerateAssets(const FARCompiledFilter& Filter, TFunctionRef<bool(const FAssetData&)> Callback, UE::AssetRegistry::EEnumerateAssetsFlags InEnumerateFlags) const override;
+	virtual bool EnumerateAssets(const FARFilter& Filter, TFunctionRef<bool(const FAssetData&)> Callback,
+		UE::AssetRegistry::EEnumerateAssetsFlags InEnumerateFlags = UE::AssetRegistry::EEnumerateAssetsFlags::None) const override;
+	virtual bool EnumerateAssets(const FARCompiledFilter& Filter, TFunctionRef<bool(const FAssetData&)> Callback,
+		UE::AssetRegistry::EEnumerateAssetsFlags InEnumerateFlags = UE::AssetRegistry::EEnumerateAssetsFlags::None) const override;
 	UE_DEPRECATED(5.1, "Asset path FNames have been deprecated, use FSoftObjectPath instead.")
 	virtual FAssetData GetAssetByObjectPath( const FName ObjectPath, bool bIncludeOnlyOnDiskAssets = false ) const override;
 	virtual FAssetData GetAssetByObjectPath(const FSoftObjectPath& ObjectPath, bool bIncludeOnlyOnDiskAssets = false, bool bSkipARFilteredAssets = true) const override;
@@ -54,9 +52,8 @@ public:
 	virtual UE::AssetRegistry::EExists TryGetAssetPackageData(const FName PackageName, FAssetPackageData& OutAssetPackageData) const override;
 	virtual UE::AssetRegistry::EExists TryGetAssetPackageData(const FName PackageName, FAssetPackageData& OutAssetPackageData, FName& OutCorrectCasePackageName) const override;
 	virtual bool GetAllAssets(TArray<FAssetData>& OutAssetData, bool bIncludeOnlyOnDiskAssets = false) const override;
-	virtual bool EnumerateAllAssets(TFunctionRef<bool(const FAssetData&)> Callback, bool bIncludeOnlyOnDiskAssets) const override;
-	virtual bool EnumerateAllAssets(TFunctionRef<bool(const FAssetData&)> Callback) const override;
-	virtual bool EnumerateAllAssets(TFunctionRef<bool(const FAssetData&)> Callback, UE::AssetRegistry::EEnumerateAssetsFlags InEnumerateFlags) const override;
+	virtual bool EnumerateAllAssets(TFunctionRef<bool(const FAssetData&)> Callback,
+		UE::AssetRegistry::EEnumerateAssetsFlags InEnumerateFlags = UE::AssetRegistry::EEnumerateAssetsFlags::None) const override;
 	virtual void GetPackagesByName(FStringView PackageName, TArray<FName>& OutPackageNames) const override;
 	virtual FName GetFirstPackageByName(FStringView PackageName) const override;
 	virtual bool GetDependencies(const FAssetIdentifier& AssetIdentifier, TArray<FAssetIdentifier>& OutDependencies, UE::AssetRegistry::EDependencyCategory Category = UE::AssetRegistry::EDependencyCategory::All, const UE::AssetRegistry::FDependencyQuery& Flags = UE::AssetRegistry::FDependencyQuery()) const override;
@@ -154,8 +151,6 @@ public:
 	virtual void AssetRenamed(const UObject* RenamedAsset, const FString& OldObjectPath) override;
 	virtual void AssetsSaved(TArray<FAssetData>&& SavedAssets) override;
 	virtual void AssetUpdateTags(UObject* Object, EAssetRegistryTagsCaller Caller) override;
-	UE_DEPRECATED(5.4, "Call AssetUpdateTags with EAssetRegistryTagsCaller::Fast")
-	virtual void AssetFullyUpdateTags(UObject* Object) override;
 	virtual void AssetTagsFinalized(const UObject& FinalizedAsset) override;
 
 	virtual bool VerseCreated(const FString& FilePath) override;
@@ -217,8 +212,6 @@ public:
 
 	virtual void Tick (float DeltaTime) override;
 
-	UE_DEPRECATED(5.5, "ReadLockEnumerateTagToAssetDatas with TArray has been deprecated. Use ReadLockEnumerateAllTagToAssetDatas instead.")
-	virtual void ReadLockEnumerateTagToAssetDatas(TFunctionRef<void(FName TagName, const TArray<const FAssetData*>& Assets)> Callback) const override;
 	virtual void ReadLockEnumerateAllTagToAssetDatas(TFunctionRef<bool(FName TagName, IAssetRegistry::FEnumerateAssetDatasFunc EnumerateAssets)> Callback) const override;
 
 	virtual bool IsPathBeautificationNeeded(const FString& InAssetPath) const override;
diff --git a/Engine/Source/Runtime/AssetRegistry/Private/AssetRegistryState.cpp b/Engine/Source/Runtime/AssetRegistry/Private/AssetRegistryState.cpp
--- a/Engine/Source/Runtime/AssetRegistry/Private/AssetRegistryState.cpp
+++ b/Engine/Source/Runtime/AssetRegistry/Private/AssetRegistryState.cpp
@@ -1296,22 +1296,6 @@ void FilterAssetsByCachedClassesByTag(TArray<const FAssetData*>& InOutResults,
 
 } // namespace UE::AssetRegistry::Private
 
-bool FAssetRegistryState::EnumerateAssets(const FARCompiledFilter& Filter, const TSet<FName>& PackageNamesToSkip,
-	TFunctionRef<bool(const FAssetData&)> Callback, bool bSkipARFilteredAssets) const
-{
-	const UE::AssetRegistry::EEnumerateAssetsFlags Flags = bSkipARFilteredAssets
-		? UE::AssetRegistry::EEnumerateAssetsFlags::None
-		: UE::AssetRegistry::EEnumerateAssetsFlags::AllowUnfilteredArAssets;
-	return EnumerateAssets(Filter, PackageNamesToSkip, Callback, Flags);
-}
-
-bool FAssetRegistryState::EnumerateAssets(const FARCompiledFilter& Filter, const TSet<FName>& PackageNamesToSkip,
-	TFunctionRef<bool(const FAssetData&)> Callback) const
-{
-	return EnumerateAssets(Filter, PackageNamesToSkip, Callback,
-		UE::AssetRegistry::EEnumerateAssetsFlags::AllowUnfilteredArAssets);
-}
-
 bool FAssetRegistryState::EnumerateAssets(const FARCompiledFilter& Filter, const TSet<FName>& PackageNamesToSkip,
 	TFunctionRef<bool(const FAssetData&)> Callback, UE::AssetRegistry::EEnumerateAssetsFlags InEnumerateFlags) const
 {
@@ -1535,22 +1519,6 @@ bool FAssetRegistryState::GetAllAssets(const TSet<FName>& PackageNamesToSkip, TA
 	EnumerateFlags);
 }
 
-bool FAssetRegistryState::EnumerateAllAssets(const TSet<FName>& PackageNamesToSkip,
-	TFunctionRef<bool(const FAssetData&)> Callback, bool bSkipARFilteredAssets) const
-{
-	const UE::AssetRegistry::EEnumerateAssetsFlags EnumerateFlags = bSkipARFilteredAssets
-		? UE::AssetRegistry::EEnumerateAssetsFlags::None
-		: UE::AssetRegistry::EEnumerateAssetsFlags::AllowUnfilteredArAssets;
-	return EnumerateAllAssets(PackageNamesToSkip, Callback, EnumerateFlags);
-}
-
-bool FAssetRegistryState::EnumerateAllAssets(const TSet<FName>& PackageNamesToSkip,
-	TFunctionRef<bool(const FAssetData&)> Callback) const
-{
-	return EnumerateAllAssets(PackageNamesToSkip, Callback,
-		UE::AssetRegistry::EEnumerateAssetsFlags::AllowUnfilteredArAssets);
-}
-
 void FAssetRegistryState::EnumerateAllAssets(TFunctionRef<void(const FAssetData&)> Callback) const
 {
 	EnumerateAllMutableAssets([&Callback](FAssetData& AssetData)
diff --git a/Engine/Source/Runtime/AssetRegistry/Private/PackageReader.cpp b/Engine/Source/Runtime/AssetRegistry/Private/PackageReader.cpp
--- a/Engine/Source/Runtime/AssetRegistry/Private/PackageReader.cpp
+++ b/Engine/Source/Runtime/AssetRegistry/Private/PackageReader.cpp
@@ -1848,13 +1848,6 @@ namespace UE::AssetRegistry
 		return true;
 	}
 
-	bool ReadPackageDataDependencies(FArchive& BinaryArchive, TBitArray<>& OutImportUsedInGame, TBitArray<>& OutSoftPackageUsedInGame)
-	{
-		UE_LOG(LogAssetRegistry, Error,
-			TEXT("This version of ReadPackageDataDependencies is no longer supported since it does not include enough information to know the package's AssetRegistryVersion. Read will be marked as failed."));
-		return false;
-	}
-
 	// See the corresponding WriteAssetRegistryPackageData defined in SavePackageUtilities.cpp in CoreUObject module
 	bool ReadPackageDataDependencies(FReadPackageDataDependenciesArgs& Args)
 	{
diff --git a/Engine/Source/Runtime/AssetRegistry/Public/AssetRegistry/AssetRegistryModule.h b/Engine/Source/Runtime/AssetRegistry/Public/AssetRegistry/AssetRegistryModule.h
--- a/Engine/Source/Runtime/AssetRegistry/Public/AssetRegistry/AssetRegistryModule.h
+++ b/Engine/Source/Runtime/AssetRegistry/Public/AssetRegistry/AssetRegistryModule.h
@@ -116,7 +116,7 @@ public:
 	}
 
 	virtual bool EnumerateAssets(const FARFilter& Filter, TFunctionRef<bool(const FAssetData&)> Callback,
-		UE::AssetRegistry::EEnumerateAssetsFlags InEnumerateFlags) const override
+		UE::AssetRegistry::EEnumerateAssetsFlags InEnumerateFlags = UE::AssetRegistry::EEnumerateAssetsFlags::None) const override
 	{
 		IAssetRegistry* AssetRegistry = IAssetRegistry::Get();
 		if (!AssetRegistry)
diff --git a/Engine/Source/Runtime/AssetRegistry/Public/AssetRegistry/AssetRegistryState.h b/Engine/Source/Runtime/AssetRegistry/Public/AssetRegistry/AssetRegistryState.h
--- a/Engine/Source/Runtime/AssetRegistry/Public/AssetRegistry/AssetRegistryState.h
+++ b/Engine/Source/Runtime/AssetRegistry/Public/AssetRegistry/AssetRegistryState.h
@@ -213,22 +213,6 @@ public:
 	ASSETREGISTRY_API bool GetAssets(const FARCompiledFilter& Filter, const TSet<FName>& PackageNamesToSkip,
 		TArray<FAssetData>& OutAssetData, bool bSkipARFilteredAssets = false) const;
 
-	/**
-	 * Enumerate asset data for all assets that match the filter.
-	 * Assets returned must satisfy every filter component if there is at least one element in the component's array.
-	 * Assets will satisfy a component if they match any of the elements in it.
-	 *
-	 * @param Filter filter to apply to the assets in the AssetRegistry
-	 * @param PackageNamesToSkip explicit list of packages to skip, because they were already added
-	 * @param Callback function to call for each asset data enumerated
-	 * @param bARFiltering Whether to apply filtering from UE::AssetRegistry::FFiltering (false by default)
-	 */
-	UE_DEPRECATED(5.5, "Use EnumerateAssets with InEnumerateFlags instead.")
-	ASSETREGISTRY_API bool EnumerateAssets(const FARCompiledFilter& Filter, const TSet<FName>& PackageNamesToSkip,
-		TFunctionRef<bool(const FAssetData&)> Callback, bool bARFiltering) const;
-	ASSETREGISTRY_API bool EnumerateAssets(const FARCompiledFilter& Filter, const TSet<FName>& PackageNamesToSkip,
-		TFunctionRef<bool(const FAssetData&)> Callback) const;
-
 	/**
 	 * Enumerate asset data for all assets that match the filter.
 	 * Assets returned must satisfy every filter component if there is at least one element in the component's array.
@@ -242,7 +226,7 @@ public:
 	 */
 	ASSETREGISTRY_API bool EnumerateAssets(const FARCompiledFilter& Filter, const TSet<FName>& PackageNamesToSkip, 
 		TFunctionRef<bool(const FAssetData&)> Callback,
-		UE::AssetRegistry::EEnumerateAssetsFlags InEnumerateFlags) const;
+		UE::AssetRegistry::EEnumerateAssetsFlags InEnumerateFlags = UE::AssetRegistry::EEnumerateAssetsFlags::AllowUnfilteredArAssets) const;
 
 	/**
 	 * Gets asset data for all assets in the registry state.
@@ -254,20 +238,6 @@ public:
 	ASSETREGISTRY_API bool GetAllAssets(const TSet<FName>& PackageNamesToSkip, TArray<FAssetData>& OutAssetData,
 		bool bARFiltering = false) const;
 
-	/**
-	 * Enumerates asset data for all assets in the registry state.
-	 *
-	 * @param PackageNamesToSkip explicit list of packages to skip, because they were already added
-	 * @param Callback function to call for each asset data enumerated
-	 * @param bARFiltering Whether to apply filtering from UE::AssetRegistry::FFiltering (false by default)
-	 */
-	UE_DEPRECATED(5.5, "Use EnumerateAllAssets with InEnumerateFlags instead.")
-	ASSETREGISTRY_API bool EnumerateAllAssets(const TSet<FName>& PackageNamesToSkip,
-		TFunctionRef<bool(const FAssetData&)> Callback, bool bARFiltering) const;
-	ASSETREGISTRY_API void EnumerateAllAssets(TFunctionRef<void(const FAssetData&)> Callback) const;
-	ASSETREGISTRY_API bool EnumerateAllAssets(const TSet<FName>& PackageNamesToSkip,
-		TFunctionRef<bool(const FAssetData&)> Callback) const;
-
 	/**
 	 * Enumerates asset data for all assets in the registry state.
 	 *
@@ -278,8 +248,9 @@ public:
 	 */
 	ASSETREGISTRY_API bool EnumerateAllAssets(const TSet<FName>& PackageNamesToSkip,
 		TFunctionRef<bool(const FAssetData&)> Callback,
-		UE::AssetRegistry::EEnumerateAssetsFlags InEnumerateFlags) const;
-	
+		UE::AssetRegistry::EEnumerateAssetsFlags InEnumerateFlags = UE::AssetRegistry::EEnumerateAssetsFlags::AllowUnfilteredArAssets) const;
+	ASSETREGISTRY_API void EnumerateAllAssets(TFunctionRef<void(const FAssetData&)> Callback) const;
+
 
 	/**
 	 * Calls the callback with the LongPackageName of each path that has assets as direct children.
@@ -374,15 +345,11 @@ public:
 	const FAssetData* GetAssetByObjectPath(const UE::AssetRegistry::Private::FCachedAssetKey& Key) const;
 
 	/**
-	 * Gets the asset data for the specified package name
+	 * Enumerates the asset datas for the specified package name
 	 *
 	 * @param PackageName the path of the package to be looked up
-	 * @return an array of AssetData*, empty if nothing found
+	 * Calls the callback once for each AssetData* in the package; callback is not called if nothing found.
 	 */
-#if !UE_ASSETREGISTRY_INDIRECT_ASSETDATA_POINTERS
-	UE_DEPRECATED(5.5, "Use EnumerateAssetsByPackageName or CopyAssetsByPackageName instead.")
-	TArrayView<FAssetData const* const> GetAssetsByPackageName(const FName PackageName) const;
-#endif
 	void EnumerateAssetsByPackageName(const FName PackageName,
 		TFunctionRef<bool(const FAssetData* AssetData)> Callback) const;
 	/** Gets the array of AssetData pointers for the package; does not copy the AssetDatas, just the pointers. */
@@ -393,15 +360,11 @@ public:
 		TFunctionRef<bool(const FAssetData* AssetData)> Callback) const;
 
 	/**
-	 * Gets the asset data for the specified asset class
+	 * Enumerates the asset datas for the specified asset class
 	 *
 	 * @param ClassPathName the class path name of the assets to look for
-	 * @return An array of AssetData*, empty if nothing found
+	 * Calls the callback once with each AssetData*, callback is not called if nothing found
 	 */
-#if !UE_ASSETREGISTRY_INDIRECT_ASSETDATA_POINTERS
-	UE_DEPRECATED(5.5, "Use EnumerateAssetsByClassPathName instead.")
-	const TArray<const FAssetData*>& GetAssetsByClassPathName(const FTopLevelAssetPath ClassPathName) const;
-#endif
 	void EnumerateAssetsByClassPathName(const FTopLevelAssetPath ClassPathName,
 		TFunctionRef<bool(const FAssetData* AssetData)> Callback) const;
 
@@ -416,21 +379,10 @@ public:
 	bool ContainsTag(FName TagName) const;
 
 	/**
-	 * Gets the asset data for the specified asset tag
+	 * Enumerates the asset datas with the specified asset tag
 	 *
 	 * @param TagName the tag name to search for
-	 * @return An array of AssetData*, empty if nothing found
-	 */
-#if !UE_ASSETREGISTRY_INDIRECT_ASSETDATA_POINTERS
-	UE_DEPRECATED(5.5, "Use EnumerateAssetsByTagName instead.")
-	const TArray<const FAssetData*>& GetAssetsByTagName(const FName TagName) const;
-#endif
-
-	/**
-	 * Enumerates the asset data for the specified asset tag
-	 *
-	 * @param TagName the tag name to search for
-	 * @param Callback the function called for each asset data
+	 * Callst he callback once with each AssetData*, callback is not called if nothing found
 	 */
 	ASSETREGISTRY_API void EnumerateAssetsByTagName(const FName TagName,
 		TFunctionRef<bool(const FAssetData* AssetData)> Callback) const;
@@ -444,18 +396,6 @@ public:
 	ASSETREGISTRY_API void EnumerateTagToAssetDatas(
 		TFunctionRef<bool(FName TagName, IAssetRegistry::FEnumerateAssetDatasFunc EnumerateAssets)> Callback) const;
 
-#if !UE_ASSETREGISTRY_INDIRECT_ASSETDATA_POINTERS
-	/** Returns const version of internal ObjectPath->AssetData map for fast iteration */
-	UE_DEPRECATED(5.5, "FAssetDataMap is a complicated implementation detail of FAssetRegistryState. Use the enumeration functions on FAssetRegistryState instead of using it directly.")
-	const FConstAssetDataMap& GetAssetDataMap() const;
-#endif
-
-#if !UE_ASSETREGISTRY_INDIRECT_ASSETDATA_POINTERS
-	/** Returns const version of internal Tag->AssetDatas map for fast iteration */
-	UE_DEPRECATED(5.5, "Use EnumerateTags or EnumerateTagToAssetDatas instead.")
-	const TMap<FName, const TArray<const FAssetData*>> GetTagToAssetDatasMap() const;
-#endif
-
 	/** Returns const version of internal PackageName->PackageData map for fast iteration */
 	const TMap<FName, const FAssetPackageData*>& GetAssetPackageDataMap() const;
 
@@ -874,18 +814,6 @@ inline FAssetData* FAssetRegistryState::GetMutableAssetByObjectPath(
 	return FoundAsset ? *FoundAsset : nullptr;
 }
 
-#if !UE_ASSETREGISTRY_INDIRECT_ASSETDATA_POINTERS
-inline TArrayView<FAssetData const* const> FAssetRegistryState::GetAssetsByPackageName(const FName PackageName) const
-{
-	if (const TArray<FAssetData*, TInlineAllocator<1>>*FoundAssetArray = CachedAssetsByPackageName.Find(PackageName))
-	{
-		return MakeArrayView(*FoundAssetArray);
-	}
-
-	return TArrayView<FAssetData* const>();
-}
-#endif
-
 inline void FAssetRegistryState::EnumerateAssetsByPackageName(const FName PackageName,
 	TFunctionRef<bool(const FAssetData* AssetData)> Callback) const
 {
@@ -990,21 +918,6 @@ inline void FAssetRegistryState::EnumerateMutableAssetsByPackagePath(FName LongP
 #endif
 }
 
-#if !UE_ASSETREGISTRY_INDIRECT_ASSETDATA_POINTERS
-inline const TArray<const FAssetData*>& FAssetRegistryState::GetAssetsByClassPathName(
-	const FTopLevelAssetPath ClassPathName) const
-{
-	static TArray<const FAssetData*> InvalidArray;
-	const TArray<FAssetData*>* FoundAssetArray = CachedAssetsByClass.Find(ClassPathName);
-	if (FoundAssetArray)
-	{
-		return reinterpret_cast<const TArray<const FAssetData*>&>(*FoundAssetArray);
-	}
-
-	return InvalidArray;
-}
-#endif
-
 inline void FAssetRegistryState::EnumerateAssetsByClassPathName(
 	const FTopLevelAssetPath ClassPathName, TFunctionRef<bool(const FAssetData* AssetData)> Callback) const
 {
@@ -1056,32 +969,6 @@ inline bool FAssetRegistryState::ContainsTag(FName TagName) const
 #endif
 }
 
-#if !UE_ASSETREGISTRY_INDIRECT_ASSETDATA_POINTERS
-inline const TArray<const FAssetData*>& FAssetRegistryState::GetAssetsByTagName(const FName TagName) const
-{
-	ensureMsgf(false, TEXT("GetAssetsByTagName has been deprecated. Please use EnumerateAssetsByTagName"));
-	static TArray<const FAssetData*> InvalidArray;
-	return InvalidArray;
-}
-#endif
-
-#if !UE_ASSETREGISTRY_INDIRECT_ASSETDATA_POINTERS
-inline const FAssetRegistryState::FConstAssetDataMap& FAssetRegistryState::GetAssetDataMap() const
-{
-	return reinterpret_cast<const FConstAssetDataMap&>(CachedAssets);
-}
-#endif
-
-#if !UE_ASSETREGISTRY_INDIRECT_ASSETDATA_POINTERS
-inline const TMap<FName, const TArray<const FAssetData*>> FAssetRegistryState::GetTagToAssetDatasMap() const
-{
-	ensureMsgf(false,
-		TEXT("GetTagToAssetDatasMap has been deprecated. Please use EnumerateTags or EnumerateTagToAssetDatas"));
-	static TMap<FName, const TArray<const FAssetData*>> InvalidMap;
-	return InvalidMap;
-}
-#endif
-
 inline const TMap<FName, const FAssetPackageData*>& FAssetRegistryState::GetAssetPackageDataMap() const
 {
 	return reinterpret_cast<const TMap<FName, const FAssetPackageData*>&>(CachedPackageData);
diff --git a/Engine/Source/Runtime/AssetRegistry/Public/AssetRegistry/IAssetRegistry.h b/Engine/Source/Runtime/AssetRegistry/Public/AssetRegistry/IAssetRegistry.h
--- a/Engine/Source/Runtime/AssetRegistry/Public/AssetRegistry/IAssetRegistry.h
+++ b/Engine/Source/Runtime/AssetRegistry/Public/AssetRegistry/IAssetRegistry.h
@@ -357,22 +357,6 @@ public:
 	virtual bool GetInMemoryAssets(const FARFilter& Filter, TArray<FAssetData>& OutAssetData, bool bSkipARFilteredAssets=true) const = 0;
 	virtual bool GetInMemoryAssets(const FARCompiledFilter& Filter, TArray<FAssetData>& OutAssetData, bool bSkipARFilteredAssets=true) const = 0;
 
-	/**
-	 * Enumerate asset data for all assets that match the filter.
-	 * Assets returned must satisfy every filter component if there is at least one element in the component's array.
-	 * Assets will satisfy a component if they match any of the elements in it.
-	 *
-	 * @param Filter filter to apply to the assets in the AssetRegistry
-	 * @param Callback function to call for each asset data enumerated
-	 * @param bSkipARFilteredAssets If true, skips Objects that return true for IsAsset but are not assets in the current platform.
-	 */
-	UE_DEPRECATED(5.5, "Use EnumerateAssets with InEnumerateFlags instead.")
-	virtual bool EnumerateAssets(const FARFilter& Filter, TFunctionRef<bool(const FAssetData&)> Callback,
-		bool bSkipARFilteredAssets) const = 0;
-	UE_DEPRECATED(5.5, "Use EnumerateAssets with InEnumerateFlags instead.")
-	virtual bool EnumerateAssets(const FARCompiledFilter& Filter, TFunctionRef<bool(const FAssetData&)> Callback,
-		bool bSkipARFilteredAssets) const = 0;
-
 	/**
 	 * Enumerate asset data for all assets that match the filter.
 	 * Assets returned must satisfy every filter component if there is at least one element in the component's array.
@@ -384,12 +368,10 @@ public:
 	 *        @see EEnumerateAssetsFlags.
 	 * @return False if filter is invalid, otherwise true.
 	 */
-	virtual bool EnumerateAssets(const FARFilter& Filter, TFunctionRef<bool(const FAssetData&)> Callback) const = 0;
-	virtual bool EnumerateAssets(const FARCompiledFilter& Filter, TFunctionRef<bool(const FAssetData&)> Callback) const = 0;
 	virtual bool EnumerateAssets(const FARFilter& Filter, TFunctionRef<bool(const FAssetData&)> Callback,
-		UE::AssetRegistry::EEnumerateAssetsFlags InEnumerateFlags) const = 0;
+		UE::AssetRegistry::EEnumerateAssetsFlags InEnumerateFlags = UE::AssetRegistry::EEnumerateAssetsFlags::None) const = 0;
 	virtual bool EnumerateAssets(const FARCompiledFilter& Filter, TFunctionRef<bool(const FAssetData&)> Callback,
-		UE::AssetRegistry::EEnumerateAssetsFlags InEnumerateFlags) const = 0;
+		UE::AssetRegistry::EEnumerateAssetsFlags InEnumerateFlags = UE::AssetRegistry::EEnumerateAssetsFlags::None) const = 0;
 
 	UE_DEPRECATED(5.6, "Asset path FNames have been deprecated, use Soft Object Path instead.")
 	UFUNCTION(BlueprintCallable, BlueprintPure=false, Category="AssetRegistry", meta=(DeprecatedFunction="Note", DeprecationMessage="Use other version that takes a SoftObjectPath"))
@@ -461,18 +443,6 @@ public:
 	UFUNCTION(BlueprintCallable, BlueprintPure=false, Category = "AssetRegistry")
 	virtual bool GetAllAssets(TArray<FAssetData>& OutAssetData, bool bIncludeOnlyOnDiskAssets = false) const = 0;
 
-	/**
-	 * Enumerate asset data for all assets in the registry.
-	 * This method may be slow, use a filter if possible to avoid iterating over the entire registry.
-	 *
-	 * @param Callback function to call for each asset data enumerated
-	 * @param bIncludeOnlyOnDiskAssets If true, use only DiskGatheredData, do not calculate from UObjects.
-	 *        @see IAssetRegistry class header for bIncludeOnlyOnDiskAssets.
-	 */
-	UE_DEPRECATED(5.5, "Use EnumerateAllAssets with InEnumerateFlags instead.")
-	virtual bool EnumerateAllAssets(TFunctionRef<bool(const FAssetData&)> Callback,
-		bool bIncludeOnlyOnDiskAssets) const = 0;
-
 	/**
 	 * Enumerate asset data for all assets currently cached in the asset registry. 
 	 * This method may be slow, use a filter if possible to avoid iterating over the entire registry.
@@ -481,9 +451,8 @@ public:
 	 * @param InEnumerateFlags flags to control enumeration and filtering.
 	 *        @see IAssetRegistry class header for EEnumerateAssetsFlags.
 	 */
-	virtual bool EnumerateAllAssets(TFunctionRef<bool(const FAssetData&)> Callback) const = 0;
 	virtual bool EnumerateAllAssets(TFunctionRef<bool(const FAssetData&)> Callback,
-		UE::AssetRegistry::EEnumerateAssetsFlags InEnumerateFlags) const = 0;
+		UE::AssetRegistry::EEnumerateAssetsFlags InEnumerateFlags = UE::AssetRegistry::EEnumerateAssetsFlags::None) const = 0;
 
 	/**
 	 * Gets the LongPackageName for all packages with the given PackageName.
@@ -903,8 +872,6 @@ public:
 	 * SavePackage.
 	 */
 	virtual void AssetUpdateTags(UObject* Object, EAssetRegistryTagsCaller Caller) = 0;
-	UE_DEPRECATED(5.4, "Call AssetUpdateTags with EAssetRegistryTagsCaller::Fast")
-	virtual void AssetFullyUpdateTags(UObject* Object) = 0;
 
 	/** Informs the asset registry that a Verse file has been created on disk */
 	virtual bool VerseCreated(const FString& FilePathOnDisk) = 0;
@@ -1115,14 +1082,6 @@ public:
 	/** Load FAssetData from the specified package filename */
 	virtual void LoadPackageRegistryData(const FString& PackageFilename, FLoadPackageRegistryData& InOutData) const = 0;
 
-	/**
-	 * Enumerate all pairs in State->TagToAssetDataMapAssetRegistry and call a callback on each pair.
-	 * To avoid copies, the callback is called from within the ReadLock.
-	 * DO NOT CALL AssetRegistry functions from the callback; doing so will create a deadlock.
-	 */
-	UE_DEPRECATED(5.5, "ReadLockEnumerateTagToAssetDatas with TArray has been deprecated. Use ReadLockEnumerateTagToAssetDatas with EnumerateAssets callback instead.")
-	virtual void ReadLockEnumerateTagToAssetDatas(TFunctionRef<void(FName TagName, const TArray<const FAssetData*>& Assets)> Callback) const = 0;
-
 	typedef TFunctionRef<bool(const FAssetData* AssetData)> FAssetDataFunc;
 	typedef TFunctionRef<bool(FAssetDataFunc AssetCallback)> FEnumerateAssetDatasFunc;
 
@@ -1258,9 +1217,6 @@ namespace UE::AssetRegistry
 		int64 AssetRegistryDependencyDataSize = 0;
 	};
 	ASSETREGISTRY_API bool ReadPackageDataDependencies(FReadPackageDataDependenciesArgs& Args);
-	UE_DEPRECATED(5.5, "Use version that takes FReadPackageDataArgs")
-	ASSETREGISTRY_API bool ReadPackageDataDependencies(FArchive& BinaryArchive, TBitArray<>& OutImportUsedInGame,
-		TBitArray<>& OutSoftPackageUsedInGame);
 
 	/**
 	 * Given a list of packages, gather the most important assets for each package.
diff --git a/Engine/Source/Runtime/CoreUObject/Public/Misc/AssetRegistryInterface.h b/Engine/Source/Runtime/CoreUObject/Public/Misc/AssetRegistryInterface.h
--- a/Engine/Source/Runtime/CoreUObject/Public/Misc/AssetRegistryInterface.h
+++ b/Engine/Source/Runtime/CoreUObject/Public/Misc/AssetRegistryInterface.h
@@ -479,7 +479,7 @@ public:
 	 * @return False if the AssetRegistry is not available or the filter is invalid, otherwise true.
 	 */
 	virtual bool EnumerateAssets(const FARFilter& Filter, TFunctionRef<bool(const FAssetData&)> Callback,
-		UE::AssetRegistry::EEnumerateAssetsFlags InEnumerateFlags) const = 0;
+		UE::AssetRegistry::EEnumerateAssetsFlags InEnumerateFlags = UE::AssetRegistry::EEnumerateAssetsFlags::None) const = 0;
 
 protected:
 
diff --git a/Engine/Source/Runtime/CoreUObject/Tests/ObjectPropertyTest.cpp b/Engine/Source/Runtime/CoreUObject/Tests/ObjectPropertyTest.cpp
--- a/Engine/Source/Runtime/CoreUObject/Tests/ObjectPropertyTest.cpp
+++ b/Engine/Source/Runtime/CoreUObject/Tests/ObjectPropertyTest.cpp
@@ -334,7 +334,7 @@ public:
 	}
 
 	virtual bool EnumerateAssets(const FARFilter& Filter, TFunctionRef<bool(const FAssetData&)> Callback,
-		UE::AssetRegistry::EEnumerateAssetsFlags InEnumerateFlags) const override
+		UE::AssetRegistry::EEnumerateAssetsFlags InEnumerateFlags = UE::AssetRegistry::EEnumerateAssetsFlags::None) const override
 	{
 		return true;
 	}
