[AssetRegistry]

- Detects and lowers the warning for when there is a version mismatch when reading in the DiscoveryCache

#rnx

diff --git a/Engine/Source/Runtime/AssetRegistry/Private/AssetDataGathererDiscoveryCache.cpp b/Engine/Source/Runtime/AssetRegistry/Private/AssetDataGathererDiscoveryCache.cpp
--- a/Engine/Source/Runtime/AssetRegistry/Private/AssetDataGathererDiscoveryCache.cpp
+++ b/Engine/Source/Runtime/AssetRegistry/Private/AssetDataGathererDiscoveryCache.cpp
@@ -259,6 +259,7 @@ bool FAssetDataDiscoveryCache::TryReadCacheFile()
 	FSharedBuffer RawBuffer;
 	bool bCacheValid = true;
 	bool bCacheCorrupt = false;
+	bool bCacheVersionMismatch = false;
 	FString Filename = GetCacheFileName();
 	{
 		TUniquePtr<FArchive> CompressedFileArchive(IFileManager::Get().CreateFileReader(*Filename));
@@ -288,7 +289,7 @@ bool FAssetDataDiscoveryCache::TryReadCacheFile()
 	if (bCacheValid)
 	{
 		FMemoryReaderView Ar(RawBuffer.GetView());
-		SerializeReadCacheFile(Ar);
+		SerializeReadCacheFile(Ar, bCacheVersionMismatch);
 		if (Ar.IsError())
 		{
 			bCacheValid = false;
@@ -298,6 +299,12 @@ bool FAssetDataDiscoveryCache::TryReadCacheFile()
 
 	if (!bCacheValid)
 	{
+		if (bCacheVersionMismatch)
+		{
+			UE_LOG(LogAssetRegistry, Display,
+				TEXT("Version mismatch for AssetDiscovery cache %s. AssetRegistry discovery of files will be uncached."),
+				*Filename);
+		}
 		if (bCacheCorrupt)
 		{
 			UE_LOG(LogAssetRegistry, Warning,
@@ -321,11 +328,12 @@ bool FAssetDataDiscoveryCache::TryReadCacheFile()
 	return true;
 }
 
-void FAssetDataDiscoveryCache::SerializeReadCacheFile(FArchive& Ar)
+void FAssetDataDiscoveryCache::SerializeReadCacheFile(FArchive& Ar, bool& bOutCacheVersionMismatch)
 {
 	FGuid Version;
 	Ar << Version;
-	if (Version != AssetDataGathererConstants::DiscoveryCacheVersion)
+	bOutCacheVersionMismatch = Version != AssetDataGathererConstants::DiscoveryCacheVersion;
+	if (bOutCacheVersionMismatch)
 	{
 		Ar.SetError();
 		return;
diff --git a/Engine/Source/Runtime/AssetRegistry/Private/AssetDataGathererDiscoveryCache.h b/Engine/Source/Runtime/AssetRegistry/Private/AssetDataGathererDiscoveryCache.h
--- a/Engine/Source/Runtime/AssetRegistry/Private/AssetDataGathererDiscoveryCache.h
+++ b/Engine/Source/Runtime/AssetRegistry/Private/AssetDataGathererDiscoveryCache.h
@@ -155,7 +155,7 @@ public:
 
 private:
 	bool TryReadCacheFile();
-	void SerializeReadCacheFile(FArchive& Ar);
+	void SerializeReadCacheFile(FArchive& Ar, bool& bOutCacheVersionMismatch);
 	void SerializeWriteCacheFile(FArchive& Ar);
 
 	bool FindOrAddIsReparsePoint(FStringView DirName, FCachedVolumeInfo& CachedVolume);
