AssetManager/AssetRegistry: Enable the new graph search code for CookRule propagation: when propagating CookRules we now do not traverse editor-only references, even if they are Build (aka ManagementPropagation) references.
#rnx

diff --git a/Engine/Source/Runtime/AssetRegistry/Private/AssetRegistry.cpp b/Engine/Source/Runtime/AssetRegistry/Private/AssetRegistry.cpp
--- a/Engine/Source/Runtime/AssetRegistry/Private/AssetRegistry.cpp
+++ b/Engine/Source/Runtime/AssetRegistry/Private/AssetRegistry.cpp
@@ -9625,23 +9625,14 @@ void FAssetRegistryImpl::SetManageReferences(FSetManageReferencesContext& Contex
 					EDependencyProperty DependencyProperties, bool bDuplicate)
 					{
 						// Skip dependencies that are EditorOnly and non-build.
-						// Propagate only through UsedInGame or build dependencies.
+						// Propagate only through UsedInGame or Build dependencies.
 						if (EnumHasAnyFlags(DependencyProperties, EDependencyProperty::Game | EDependencyProperty::Build))
 						{
-							// Different behavior in the graph searches is temporarily disabled in the first submit,
-							// So we can test the new graph search independently of the desired behavior change.
-							static constexpr bool bCookRulesPaintedOnlyWithGame = false;
-
+							// ChunkManagement propagates through either UsedGame or Build dependencies, but CookRules
+							// propagate only through Game dependencies.
 							EManageSearchColorBit EdgePaintColors;
-							if (bCookRulesPaintedOnlyWithGame)
-							{
-								EdgePaintColors = EnumHasAnyFlags(DependencyProperties, EDependencyProperty::Game) ?
+							EdgePaintColors = EnumHasAnyFlags(DependencyProperties, EDependencyProperty::Game) ?
 									EManageSearchColorBit::All : EManageSearchColorBit::Chunking;
-							}
-							else
-							{
-								EdgePaintColors = EManageSearchColorBit::All;
-							}
 							EManageSearchColorBit TargetPaintColors = VertexPaintColors & EdgePaintColors;
 							VisitNode(SourceNode, TargetNode, DependencyCategory, DependencyProperties, TargetPaintColors);
 						}
