[AssetRegistry]

- Adds scratch area for DependencyGathering temporary allocations to avoid frequent large temporary allocations
- Only increments the batch number if we deal with an asset rather than when we skip assets that are not ready for processing

#rnx

diff --git a/Engine/Source/Runtime/AssetRegistry/Private/AssetRegistry.cpp b/Engine/Source/Runtime/AssetRegistry/Private/AssetRegistry.cpp
--- a/Engine/Source/Runtime/AssetRegistry/Private/AssetRegistry.cpp
+++ b/Engine/Source/Runtime/AssetRegistry/Private/AssetRegistry.cpp
@@ -696,6 +696,43 @@ void FAsyncConsumer::Consume(UAssetRegistryImpl& UARI, UE::AssetRegistry::Impl::
 
 namespace UE::AssetRegistry
 {
+
+namespace DependsNode
+{
+typedef TMap<FDependsNode*, FDependsNode::FPackageFlagSet> FPackageDependencyMap;
+
+/** Uncompressed version of DependsNode data, used to avoid allocations and
+compression cost when constructing batches of DependsNodes in DependencyDataGathered. */
+struct FConstructData
+{
+	FDependsNode* SourceNode;
+	FPackageDependencyMap PackageDependencies;
+	FDependsNode::FDependsNodeList SearchableNamesNodes;
+};
+
+struct FGatherDependencyDataScratchpad
+{
+	TMap<FName, FConstructData> AssetNameToConstructDataMap;
+	TMap<FName, FName> CachedDepToRedirect;
+	TSet<FDependsNode*> ModifiedDependNodes;
+
+	void Reset()
+	{
+		AssetNameToConstructDataMap.Reset();
+		CachedDepToRedirect.Reset();
+		ModifiedDependNodes.Reset();
+	}
+
+	void Shrink()
+	{
+		AssetNameToConstructDataMap.Shrink();
+		CachedDepToRedirect.Shrink();
+		ModifiedDependNodes.Shrink();
+	}
+};
+
+}
+
 void FAssetRegistryImpl::ConditionalLoadPremadeAssetRegistry(UAssetRegistryImpl& UARI, Impl::FEventContext& EventContext, UE::AssetRegistry::FInterfaceWriteScopeLock& ScopeLock)
 {
 	AsyncConsumer.Wait(UARI, ScopeLock);
@@ -1880,6 +1917,8 @@ bool FAssetRegistryImpl::TryConstructGathererIfNeeded()
 	}
 
 	GlobalGatherer = MakeUnique<FAssetDataGatherer>(UAssetRegistryImpl::Get());
+	GatherDependencyDataScratchPad.Reset(new UE::AssetRegistry::DependsNode::FGatherDependencyDataScratchpad());
+
 	UpdateMaxSecondsPerFrame();
 
 	// Read script packages if all initial plugins have been loaded, otherwise do nothing; we wait for the callback.
@@ -4651,12 +4690,7 @@ UE::AssetRegistry::Impl::EGatherStatus UAssetRegistryImpl::TickOnBackgroundThrea
 					if (InterfaceLock.HasWaiters() || IsBackgroundProcessingPaused())
 					{
 #if WITH_EDITOR
-						// During EngineStartup many packages are loading and need to query the AssetRegistry; do not
-						// count them in the metric for backgroundtick interruptions.
-						if (IsEngineStartupModuleLoadingComplete())
-						{
-							++GuardedData.GetBackgroundTickInterruptionsCount();
-						}
+						++GuardedData.GetBackgroundTickInterruptionsCount();
 #endif
 						return true;
 					}
@@ -5531,6 +5565,8 @@ void FAssetRegistryImpl::Shrink()
 	PackagesNeedingDependencyCalculationOnGameThread.Shrink();
 #endif
 	State.Shrink();
+
+	GatherDependencyDataScratchPad->Shrink();
 }
 
 void FAssetRegistryImpl::OnInitialSearchCompleted(Impl::FEventContext& EventContext)
@@ -7080,20 +7116,6 @@ void FAssetRegistryImpl::PathDataGathered(Impl::FEventContext& EventContext, TRi
 	}
 }
 
-namespace DependsNode
-{
-typedef TMap<FDependsNode*, FDependsNode::FPackageFlagSet> FPackageDependencyMap;
-
-/** Uncompressed version of DependsNode data, used to avoid allocations and 
-compression cost when constructing batches of DependsNodes in DependencyDataGathered. */
-struct FConstructData
-{
-	FDependsNode* SourceNode;
-	FPackageDependencyMap PackageDependencies;
-	FDependsNode::FDependsNodeList SearchableNamesNodes;
-};
-}
-
 void FAssetRegistryImpl::DependencyDataGathered(TMultiMap<FName, FPackageDependencyData>& DependsResults,
 	TMultiMap<FName, FPackageDependencyData>& OutDeferredDependencyResults,
 	TSet<FName>* OutPackagesNeedingDependencyCalculation, Impl::FInterruptionContext& InOutInterruptionContext,
@@ -7124,13 +7146,13 @@ void FAssetRegistryImpl::DependencyDataGathered(TMultiMap<FName, FPackageDepende
 		Utils::InitializeMountPoints(MountPointsForVerifyAfterGather);
 	}
 
-	TMap<FName, FConstructData, TInlineSetAllocator<1>> AssetNameToConstructDataMap;
-	TMap<FName, FName, TInlineSetAllocator<1>> CachedDepToRedirect;
-	TSet<FDependsNode*, DefaultKeyFuncs<FDependsNode*>, TInlineSetAllocator<1>> ModifiedDependNodes;
-
+	TMap<FName, FConstructData>& AssetNameToConstructDataMap = GatherDependencyDataScratchPad->AssetNameToConstructDataMap;
+	TMap<FName, FName>& CachedDepToRedirect = GatherDependencyDataScratchPad->CachedDepToRedirect;
+	TSet<FDependsNode*>& ModifiedDependNodes = GatherDependencyDataScratchPad->ModifiedDependNodes;
 	AssetNameToConstructDataMap.Reserve(DependsResults.Num());
 	CachedDepToRedirect.Reserve(DependsResults.Num());
 	ModifiedDependNodes.Reserve(DependsResults.Num());
+	ON_SCOPE_EXIT{ GatherDependencyDataScratchPad->Reset(); };
 
 	State.CachedPackageData.Reserve(DependsResults.Num());
 	State.CachedDependsNodes.Reserve(DependsResults.Num());
@@ -7146,6 +7168,13 @@ void FAssetRegistryImpl::DependencyDataGathered(TMultiMap<FName, FPackageDepende
 		const bool bShouldSortReferencers = ShouldSortReferencers();
 		for (TMultiMap<FName, FPackageDependencyData>::TIterator Iter(DependsResults); Iter && !bInterrupted; ++Iter)
 		{
+			if (DeferredAssets.Contains(Iter.Key()) || DeferredAssetsForGameThread.Contains(Iter.Key()))
+			{
+				OutDeferredDependencyResults.Add(MoveTemp(*Iter));
+				Iter.RemoveCurrent();
+				continue;
+			}
+
 			ON_SCOPE_EXIT
 			{
 				// ShouldExitEarly calls FPlatformTime::Seconds which isn't super cheap
@@ -7157,13 +7186,6 @@ void FAssetRegistryImpl::DependencyDataGathered(TMultiMap<FName, FPackageDepende
 				}
 			};
 
-			if (DeferredAssets.Contains(Iter.Key()) || DeferredAssetsForGameThread.Contains(Iter.Key()))
-			{
-				OutDeferredDependencyResults.Add(MoveTemp(*Iter));
-				Iter.RemoveCurrent();
-				continue;
-			}
-
 			FPackageDependencyData Result = MoveTemp(Iter.Value());
 			Iter.RemoveCurrent();
 
diff --git a/Engine/Source/Runtime/AssetRegistry/Private/AssetRegistryImpl.h b/Engine/Source/Runtime/AssetRegistry/Private/AssetRegistryImpl.h
--- a/Engine/Source/Runtime/AssetRegistry/Private/AssetRegistryImpl.h
+++ b/Engine/Source/Runtime/AssetRegistry/Private/AssetRegistryImpl.h
@@ -47,6 +47,7 @@ namespace UE::AssetRegistry::Impl { struct FTickContext; }
 namespace UE::AssetRegistry::Impl { struct FClassInheritanceContext; }
 namespace UE::AssetRegistry::Private { class FRWLockWithPriority; }
 namespace UE::AssetRegistry::Private { using FInterfaceRWLock = FRWLockWithPriority; }
+namespace UE::AssetRegistry::DependsNode { struct FGatherDependencyDataScratchpad; }
 namespace UE::AssetRegistry { template <typename TScopeLockType> class TRWScopeLockWithPriority; }
 namespace UE::AssetRegistry { using FInterfaceReadScopeLock = TRWScopeLockWithPriority<UE::TReadScopeLock<Private::FRWLockWithPriority>>; }
 namespace UE::AssetRegistry { using FInterfaceWriteScopeLock = TRWScopeLockWithPriority<UE::TWriteScopeLock<Private::FRWLockWithPriority>>; }
@@ -841,6 +842,7 @@ private:
 	friend UE::AssetRegistry::Impl::FClassInheritanceContext;
 	friend UE::AssetRegistry::Impl::FTickContext;
 
+	TUniquePtr<UE::AssetRegistry::DependsNode::FGatherDependencyDataScratchpad> GatherDependencyDataScratchPad = nullptr;
 };
 
 namespace Impl
