[AssetRegistry]

Check that if journaling is enabled, only assert if journaling returns an invalid NextJournalEntry. Otherwise, invalid entries are expected.

#rnx

diff --git a/Engine/Source/Runtime/AssetRegistry/Private/AssetDataGathererDiscoveryCache.cpp b/Engine/Source/Runtime/AssetRegistry/Private/AssetDataGathererDiscoveryCache.cpp
--- a/Engine/Source/Runtime/AssetRegistry/Private/AssetDataGathererDiscoveryCache.cpp
+++ b/Engine/Source/Runtime/AssetRegistry/Private/AssetDataGathererDiscoveryCache.cpp
@@ -608,7 +608,7 @@ void FCachedVolumeInfo::InitializePlatformData()
 		if (NextJournalEntryToScan == FileJournalEntryHandleInvalid)
 		{
 			JournalId = JournalIdOnDisk;
-			check(NextJournalEntryOnDisk != FileJournalEntryHandleInvalid);
+			check(!bJournalAvailable || NextJournalEntryOnDisk != FileJournalEntryHandleInvalid);
 			NextJournalEntryToScan = NextJournalEntryOnDisk;
 		}
 	}
