[AssetRegistry]
- Make enumerating in-memory assets thread-safe and callable from the loader

diff --git a/Engine/Source/Runtime/AssetRegistry/Private/AssetRegistry.cpp b/Engine/Source/Runtime/AssetRegistry/Private/AssetRegistry.cpp
--- a/Engine/Source/Runtime/AssetRegistry/Private/AssetRegistry.cpp
+++ b/Engine/Source/Runtime/AssetRegistry/Private/AssetRegistry.cpp
@@ -1053,7 +1053,10 @@ void FAssetRegistryImpl::Initialize(Impl::FInitializeContext& Context)
 	// Collect all code generator classes (currently BlueprintCore-derived ones)
 	CollectCodeGeneratorClasses();
 #if WITH_ENGINE && WITH_EDITOR
-	Utils::PopulateSkipClasses(SkipUncookedClasses, SkipCookedClasses);
+	{
+		UE::TWriteScopeLock WriteLock(SkipClassesLock);
+		Utils::PopulateSkipClasses(SkipUncookedClasses, SkipCookedClasses);
+	}
 #endif
 
 	// Read default serialization options
@@ -1402,12 +1405,10 @@ void FAssetRegistryImpl::OnPluginLoadingComplete(bool bPhaseSuccessful)
 	// Reparse the skip classes the next time ShouldSkipAsset is called, since available classes
 	// for the search over all classes may have changed
 #if WITH_ENGINE && WITH_EDITOR
-	// If we ever need to update the Filtering list outside of the game thread, we will need to defer the update 
-	// of the Filtering namespace to the tick function; UE::AssetRegistry::Filtering can only be used in game thread
-	check(IsInGameThread());
-
-	Utils::PopulateSkipClasses(SkipUncookedClasses, SkipCookedClasses);
-	UE::AssetRegistry::FFiltering::SetSkipClasses(SkipUncookedClasses, SkipCookedClasses);
+	{
+		UE::TWriteScopeLock WriteLock(SkipClassesLock);
+		Utils::PopulateSkipClasses(SkipUncookedClasses, SkipCookedClasses);
+	}
 #endif
 }
 
@@ -2578,7 +2579,7 @@ void EnumerateMemoryAssetsHelper(const FARCompiledFilter& InFilter, TSet<FName>&
 	bool& bOutStopIteration, TFunctionRef<bool(const UObject* Object, FAssetData&& PartialAssetData)> Callback,
 	bool bSkipARFilteredAssets)
 {
-	checkf(IsInGameThread(), TEXT("Enumerating in-memory assets can only be done on the game thread; it uses non-threadsafe UE::AssetRegistry::Filtering globals."));
+	checkf(IsInGameThread() || IsInAsyncLoadingThread() || IsInParallelLoadingThread(), TEXT("Enumerating in-memory assets can only be done on the game thread or in the loader, there are too many GetAssetRegistryTags() still not thread-safe."));
 	bOutStopIteration = false;
 	UE_TRACK_REFERENCING_OPNAME_SCOPED(PackageAccessTrackingOps::NAME_ResetContext);
 
@@ -3021,7 +3022,7 @@ namespace Utils
 void EnumerateAllMemoryAssets(TSet<FName>& OutPackageNamesWithAssets, bool& bOutStopIteration,
 	TFunctionRef<bool(FAssetData&&)> Callback)
 {
-	checkf(IsInGameThread(), TEXT("Enumerating memory assets can only be done on the game thread; it uses non-threadsafe UE::AssetRegistry::Filtering globals."));
+	checkf(IsInGameThread() || IsInAsyncLoadingThread() || IsInParallelLoadingThread(), TEXT("Enumerating memory assets can only be done on the game thread or in the loader, there are too many GetAssetRegistryTags() still not thread-safe."));
 	bOutStopIteration = false;
 	for (FThreadSafeObjectIterator ObjIt; ObjIt; ++ObjIt)
 	{
@@ -9866,6 +9867,7 @@ const TSet<FName>& FAssetRegistryImpl::GetCachedEmptyPackages() const
 bool FAssetRegistryImpl::ShouldSkipAsset(FTopLevelAssetPath AssetClass, uint32 PackageFlags) const
 {
 #if WITH_ENGINE && WITH_EDITOR
+	UE::TReadScopeLock ReadLock(SkipClassesLock);
 	return Utils::ShouldSkipAsset(AssetClass, PackageFlags, SkipUncookedClasses, SkipCookedClasses);
 #else
 	return false;
@@ -9875,6 +9877,7 @@ bool FAssetRegistryImpl::ShouldSkipAsset(FTopLevelAssetPath AssetClass, uint32 P
 bool FAssetRegistryImpl::ShouldSkipAsset(const UObject* InAsset) const
 {
 #if WITH_ENGINE && WITH_EDITOR
+	UE::TReadScopeLock ReadLock(SkipClassesLock);
 	return Utils::ShouldSkipAsset(InAsset, SkipUncookedClasses, SkipCookedClasses);
 #else
 	return false;
diff --git a/Engine/Source/Runtime/AssetRegistry/Private/AssetRegistryImpl.h b/Engine/Source/Runtime/AssetRegistry/Private/AssetRegistryImpl.h
--- a/Engine/Source/Runtime/AssetRegistry/Private/AssetRegistryImpl.h
+++ b/Engine/Source/Runtime/AssetRegistry/Private/AssetRegistryImpl.h
@@ -832,6 +832,7 @@ private:
 	bool bGatherNeedsApplyObjectRedirects = false;
 #endif
 #if WITH_ENGINE && WITH_EDITOR
+	mutable FRWLock SkipClassesLock;
 	/** Class names that return true for IsAsset but which should not be treated as assets in uncooked packages */
 	TSet<FTopLevelAssetPath> SkipUncookedClasses;
 	/** Class names that return true for IsAsset but which should not be treated as assets in cooked packages */
diff --git a/Engine/Source/Runtime/CoreUObject/Private/Misc/AssetRegistryInterface.cpp b/Engine/Source/Runtime/CoreUObject/Private/Misc/AssetRegistryInterface.cpp
--- a/Engine/Source/Runtime/CoreUObject/Private/Misc/AssetRegistryInterface.cpp
+++ b/Engine/Source/Runtime/CoreUObject/Private/Misc/AssetRegistryInterface.cpp
@@ -28,6 +28,7 @@ namespace Private
 }
 
 #if WITH_ENGINE && WITH_EDITOR
+	FRWLock                  SkipClassesLock;
 	TSet<FTopLevelAssetPath> SkipUncookedClasses;
 	TSet<FTopLevelAssetPath> SkipCookedClasses;
 	bool bInitializedSkipClasses = false;
@@ -35,6 +36,7 @@ namespace Private
 	void FFiltering::SetSkipClasses(const TSet<FTopLevelAssetPath>& InSkipUncookedClasses, const TSet<FTopLevelAssetPath>& InSkipCookedClasses)
 	{
 		bInitializedSkipClasses = true;
+		UE::TWriteScopeLock WriteLock(SkipClassesLock);
 		SkipUncookedClasses = InSkipUncookedClasses;
 		SkipCookedClasses = InSkipCookedClasses;
 	}
@@ -48,6 +50,7 @@ namespace Private
 			// be replaced eventually, so leaving it for now.
 			if (GIsEditor && (!IsRunningCommandlet() || IsRunningCookCommandlet()))
 			{
+				UE::TWriteScopeLock WriteLock(SkipClassesLock);
 				Utils::PopulateSkipClasses(SkipUncookedClasses, SkipCookedClasses);
 			}
 
@@ -158,6 +161,7 @@ namespace Utils
 		// or when we support multiple assets.
 		InitializeShouldSkipAsset();
 
+		UE::TReadScopeLock ReadLock(SkipClassesLock);
 		return Utils::ShouldSkipAsset(AssetClass, PackageFlags, SkipUncookedClasses, SkipCookedClasses);
 #else
 		return false;
