[AssetRegistry] Allow overlapping KeepInDevOnly ClassFilter rules to exclude one another

diff --git a/Engine/Source/Runtime/AssetRegistry/Private/AssetRegistry.cpp b/Engine/Source/Runtime/AssetRegistry/Private/AssetRegistry.cpp
--- a/Engine/Source/Runtime/AssetRegistry/Private/AssetRegistry.cpp
+++ b/Engine/Source/Runtime/AssetRegistry/Private/AssetRegistry.cpp
@@ -2,6 +2,7 @@
 
 #include "AssetRegistry.h"
 
+#include "Algo/HeapSort.h"
 #include "Algo/Transform.h"
 #include "Algo/Unique.h"
 #include "AssetDataGatherer.h"
@@ -50,6 +51,9 @@
 #include "UObject/ConstructorHelpers.h"
 #include "UObject/CoreRedirects.h"
 #include "UObject/MetaData.h"
+#include "UObject/Object.h"
+#include "UObject/SoftObjectPath.h"
+#include "UObject/UObjectGlobals.h"
 #include "UObject/UObjectHash.h"
 #include "UObject/UObjectIterator.h"
 #include "UObject/UObjectThreadContext.h"
@@ -1060,8 +1064,9 @@ void FAssetRegistryImpl::Initialize(Impl::FInitializeContext& Context)
 #endif
 
 	// Read default serialization options
-	Utils::InitializeSerializationOptionsFromIni(SerializationOptions, nullptr, UE::AssetRegistry::ESerializationTarget::ForGame);
-	Utils::InitializeSerializationOptionsFromIni(DevelopmentSerializationOptions, nullptr, UE::AssetRegistry::ESerializationTarget::ForDevelopment);
+	constexpr static bool bInitializeFilterTagOptions = false;
+	Utils::InitializeSerializationOptionsFromIni(SerializationOptions, nullptr, UE::AssetRegistry::ESerializationTarget::ForGame, bInitializeFilterTagOptions);
+	Utils::InitializeSerializationOptionsFromIni(DevelopmentSerializationOptions, nullptr, UE::AssetRegistry::ESerializationTarget::ForDevelopment, bInitializeFilterTagOptions);
 
 	bool bStartedAsyncGather = false;
 	if (ShouldSearchAllAssetsAtStart())
@@ -1281,6 +1286,10 @@ void UAssetRegistryImpl::InitializeEvents(UE::AssetRegistry::Impl::FInitializeCo
 	{
 		FCoreDelegates::OnAllModuleLoadingPhasesComplete.AddUObject(this, &UAssetRegistryImpl::OnInitialPluginLoadingComplete);
 	}
+	else
+	{
+		OnInitialPluginLoadingComplete();
+	}
 }
 
 UAssetRegistryImpl::UAssetRegistryImpl(FVTableHelper& Helper)
@@ -1384,6 +1393,10 @@ void UAssetRegistryImpl::OnInitialPluginLoadingComplete()
 	{
 		LLM_SCOPE(ELLMTag::AssetRegistry);
 		UE::AssetRegistry::FInterfaceWriteScopeLock InterfaceScopeLock(InterfaceLock);
+		// Read default serialization options
+		constexpr static bool bInitializeFilterTagOptions = true;
+		UE::AssetRegistry::Utils::InitializeSerializationOptionsFromIni(GuardedData.SerializationOptions, nullptr, UE::AssetRegistry::ESerializationTarget::ForGame, bInitializeFilterTagOptions);
+		UE::AssetRegistry::Utils::InitializeSerializationOptionsFromIni(GuardedData.DevelopmentSerializationOptions, nullptr, UE::AssetRegistry::ESerializationTarget::ForDevelopment, bInitializeFilterTagOptions);
 		GuardedData.OnPluginLoadingComplete(true);
 	}
 
@@ -1451,7 +1464,8 @@ void UAssetRegistryImpl::InitializeSerializationOptions(FAssetRegistrySerializat
 	}
 	else
 	{
-		UE::AssetRegistry::Utils::InitializeSerializationOptionsFromIni(Options, TargetPlatform, Target);
+		constexpr static bool bInitializeFilterTagOptions = true;
+		UE::AssetRegistry::Utils::InitializeSerializationOptionsFromIni(Options, TargetPlatform, Target, bInitializeFilterTagOptions);
 	}
 }
 
@@ -1488,7 +1502,8 @@ static TSet<FName> MakeNameSet(const TArray<FString>& Strings)
 void InitializeSerializationOptionsFromIni(
 	FAssetRegistrySerializationOptions& Options,
 	const ITargetPlatform* TargetPlatform,
-	UE::AssetRegistry::ESerializationTarget Target)
+	UE::AssetRegistry::ESerializationTarget Target,
+	bool bInitializeFilterTagOptions)
 {
 	FString PlatformIniName;
 	EBuildTargetType TargetType = EBuildTargetType::Game;
@@ -1544,14 +1559,24 @@ void InitializeSerializationOptionsFromIni(
 	}
 
 	EngineIni->GetBool(TEXT("AssetRegistry"), TEXT("bUseAssetRegistryTagsWhitelistInsteadOfBlacklist"), Options.bUseAssetRegistryTagsAllowListInsteadOfDenyList);
-	TArray<FString> FilterListItems;
-	if (Options.bUseAssetRegistryTagsAllowListInsteadOfDenyList)
-	{
-		EngineIni->GetArray(TEXT("AssetRegistry"), TEXT("CookedTagsWhitelist"), FilterListItems);
-	}
-	else
+
+	if (bInitializeFilterTagOptions)
 	{
-		EngineIni->GetArray(TEXT("AssetRegistry"), TEXT("CookedTagsBlacklist"), FilterListItems);
+		TArray<FString> FilterListItems;
+		if (Options.bUseAssetRegistryTagsAllowListInsteadOfDenyList)
+		{
+			EngineIni->GetArray(TEXT("AssetRegistry"), TEXT("CookedTagsWhitelist"), FilterListItems);
+		}
+		else
+		{
+			EngineIni->GetArray(TEXT("AssetRegistry"), TEXT("CookedTagsBlacklist"), FilterListItems);
+		}
+
+		const bool bKeepDevelopmentTags = bForDevelopment || FParse::Param(FCommandLine::Get(), TEXT("ARKeepDevTags"));
+		const bool bLoadClasses = true;
+		const bool bIsFilteringDevelopmentAR = bKeepDevelopmentTags;
+		TMap<FName, FParsedFilterTagRules> Rules = ParseFilterTags(FilterListItems, bKeepDevelopmentTags, Options.bUseAssetRegistryTagsAllowListInsteadOfDenyList);
+		Options.CookFilterlistTagsByClass = ResolveFilterTags(Rules, bLoadClasses);
 	}
 
 	{
@@ -1564,6 +1589,24 @@ void InitializeSerializationOptionsFromIni(
 		EngineIni->GetArray(TEXT("AssetRegistry"), TEXT("CookedTagsAsPathName"), AsPathName);
 		Options.CookTagsAsPath = MakeNameSet(AsPathName);
 	}
+}
+
+static int32 CalculateInheritanceDepth(const UClass* Class)
+{
+	int32 InheritanceDepth = 0;
+	if (Class)
+	{
+		for (auto Iter = Class->GetSuperStructIterator(); Iter; ++Iter)
+		{
+			++InheritanceDepth;
+		}
+	}
+	return InheritanceDepth;
+}
+
+TMap<FName, FParsedFilterTagRules> ParseFilterTags(const TArray<FString>& FilterListItems, bool bIsFilteringDevelopmentAR , bool bUseAllowListInsteadOfDenyList)
+{
+	TMap<FName, FParsedFilterTagRules> TagsOrder;
 
 	// Takes on the pattern "(Class=SomeClass,Tag=SomeTag)"
 	// Optional key KeepInDevOnly for tweaking a DevelopmentAssetRegistry (additive if allow list, subtractive if deny list)
@@ -1613,23 +1656,43 @@ void InitializeSerializationOptionsFromIni(
 			}
 		}
 
-		const bool bKeepDevelopmentTags = bForDevelopment || FParse::Param(FCommandLine::Get(), TEXT("ARKeepDevTags"));
-		const bool bPassesDevOnlyRule = !bKeepInDevOnly || Options.bUseAssetRegistryTagsAllowListInsteadOfDenyList == bKeepDevelopmentTags;
-		if (!ClassName.IsEmpty() && !TagName.IsEmpty() && bPassesDevOnlyRule)
+		bool bAddToList;
+		if (bUseAllowListInsteadOfDenyList)
+		{
+			bAddToList = !bKeepInDevOnly || bIsFilteringDevelopmentAR;
+		}
+		else
 		{
-			FName TagFName = FName(*TagName);
+			// "KeepInDevOnly" means "Keep the tag only when creating the DevelopmentAssetRegistry"
+			// That is complicated meaning when we are creating a DenyList, wherein adding an element to the list removes
+			// it from the filtered collection of tags.
+			// We add it to the list if either it applies to all ARs because KeepInDevOnly was not set, or if we are creating
+			// the GameAssetRegistry.
+			bAddToList = !bKeepInDevOnly || !bIsFilteringDevelopmentAR;
+		}
 
-			// Include subclasses if the class is in memory at this time (native classes only)
-			UClass* FilterlistClass = Cast<UClass>(StaticFindObject(UClass::StaticClass(), nullptr, *ClassName));
-			if (FilterlistClass)
+		if (!ClassName.IsEmpty() && !TagName.IsEmpty())
+		{
+			FParsedFilterTagRules& Rules = TagsOrder.FindOrAdd(FName(TagName.Len(), *TagName));
+			TArray<FFilterTagsRule>& OrderedList = Rules.OrderedList;
+			const UClass* FilterlistClass = nullptr;
+			// Skip known wildcard
+			if (ClassName != TEXTVIEW("*"))
 			{
-				Options.CookFilterlistTagsByClass.FindOrAdd(FilterlistClass->GetClassPathName()).Add(TagFName);
+				FilterlistClass = Cast<UClass>(StaticFindObject(UClass::StaticClass(), nullptr, *ClassName));
+			}
 
-				TArray<UClass*> DerivedClasses;
-				GetDerivedClasses(FilterlistClass, DerivedClasses);
-				for (UClass* DerivedClass : DerivedClasses)
+			if (FilterlistClass)
+			{
+				const int32 InheritanceDepth = CalculateInheritanceDepth(FilterlistClass);
+				// Essentially I want to do OrderedList.HeapPush(FFilterTagsRule{...}, &GetInheritanceDepth);
+				// but there is no TArray::HeapPush which takes in an Item and a projection functor.
+				// Also there is no equivalent Emplace semantics to construct object inplace without extra copies/moves.
+				// So here is an inline implementation of this:
+				OrderedList.Emplace(FilterlistClass->GetClassPathName(), FilterlistClass, InheritanceDepth, bAddToList);
+				if (!Rules.bContainsAnyUnresolvedRule)
 				{
-					Options.CookFilterlistTagsByClass.FindOrAdd(DerivedClass->GetClassPathName()).Add(TagFName);
+					AlgoImpl::HeapSiftUp(OrderedList.GetData(), 0, OrderedList.Num() - 1, [](const FFilterTagsRule& Info) { return Info.InheritanceDepth; }, TLess<int32>());
 				}
 			}
 			else
@@ -1648,13 +1711,146 @@ void InitializeSerializationOptionsFromIni(
 				{
 					ClassPathName = FTopLevelAssetPath(ClassName);
 				}
-				// Class is not in memory yet. Just add an explicit filter.
-				// Automatically adding subclasses of non-native classes is not supported.
-				// In these cases, using Class=* is usually sufficient				
-				Options.CookFilterlistTagsByClass.FindOrAdd(ClassPathName).Add(TagFName);
+
+				if (!ClassPathName.IsValid())
+				{
+					UE_LOG(LogAssetRegistry, Warning, 
+						TEXT("There is a rule to %s tag \"%s\" for class \"%s\", which is not a valid class path. Ignoring."),
+						bAddToList ? TEXT("included") : TEXT("excluded"),
+						*TagName,
+						*ClassName
+					);
+					continue;
+				}
+
+				Rules.bContainsAnyUnresolvedRule = true;
+				OrderedList.Emplace(ClassPathName, nullptr, -1, bAddToList);
 			}
 		}
 	}
+
+	return TagsOrder;
+}
+
+static void ResolveRules(FStringView TagName, TArray<FFilterTagsRule>& OrderedList, bool bLoadClasses)
+{
+	auto RuleIt = OrderedList.CreateIterator();
+	if (RuleIt->ClassPath == UE::AssetRegistry::WildcardPathName)
+	{
+		// Move cursor to next element and delete the rest
+		for (++RuleIt; RuleIt; ++RuleIt)
+		{
+			const bool bIsDuplicate = (OrderedList[0].bInclude == RuleIt->bInclude);
+			UE_CLOG(!bIsDuplicate, LogAssetRegistry, Warning, 
+				TEXT("There is a wildcard rule to %s tag \"%.*s\", ")
+				TEXT("but there is another rule to %s that tag for class \"%s\". ")
+				TEXT("Only wildcard rule is accepted. Ignoring another rule."),
+				OrderedList[0].bInclude ? TEXT("include") : TEXT("exclude"),
+				TagName.Len(), TagName.GetData(),
+				RuleIt->bInclude ? TEXT("include") : TEXT("exclude"),
+				*RuleIt->ClassPath.ToString()
+			);
+			RuleIt.RemoveCurrentSwap();
+		}
+		return;
+	}
+
+	for (; RuleIt; ++RuleIt)
+	{
+		RuleIt->ClassPtr = bLoadClasses
+			? StaticLoadClass(UObject::StaticClass(), nullptr, RuleIt->ClassPath.ToString())
+			: Cast<UClass>(StaticFindObject(UClass::StaticClass(), RuleIt->ClassPath));
+		if (RuleIt->ClassPtr)
+		{
+			RuleIt->InheritanceDepth = CalculateInheritanceDepth(RuleIt->ClassPtr);
+		}
+		else
+		{
+			// Silently ignore same unresolved rules
+			const bool bIsDuplicate = (RuleIt.GetIndex() != 0) && (OrderedList[0].bInclude == RuleIt->bInclude);
+			UE_CLOG(!bIsDuplicate, LogAssetRegistry, Warning,
+				TEXT("There is a rule to %s tag \"%.*s\" for class \"%s\". Class is unknown. Ignoring."),
+				RuleIt->bInclude ? TEXT("include") : TEXT("exclude"),
+				TagName.Len(), TagName.GetData(),
+				*RuleIt->ClassPath.ToString()
+			);
+			RuleIt.RemoveCurrentSwap();
+		}
+	}
+	Algo::HeapSortBy(OrderedList, [](const FFilterTagsRule& Info) { return Info.InheritanceDepth; }, TLess<int32>());
+}
+
+TMap<FTopLevelAssetPath, TSet<FName>> ResolveFilterTags(TMap<FName, FParsedFilterTagRules>& InOutRules, bool bLoadClasses)
+{
+	TMap<FTopLevelAssetPath, TSet<FName>> Result;
+	// keep temp allocation around while we traverse the ordered list
+	TSet<const UClass*> ClassesToAdd;
+	// Populate CookFilterlistTagsByClass by walking the ordered list
+	for (TTuple<FName, FParsedFilterTagRules>& Pair : InOutRules)
+	{
+		const FName Tag = Pair.Key;
+		TStringBuilder<512> TagNameBuilder;
+		Pair.Key.ToString(TagNameBuilder);
+		const FStringView TagName = TagNameBuilder.ToView();
+		FParsedFilterTagRules& TagRules = Pair.Value;
+		TArray<FFilterTagsRule>& OrderedList = TagRules.OrderedList;
+		if (TagRules.bContainsAnyUnresolvedRule)
+		{
+			const bool bIsWildcard = OrderedList.IsValidIndex(0) && (OrderedList[0].ClassPath == UE::AssetRegistry::WildcardPathName);
+			UE_CLOG(!bIsWildcard, LogAssetRegistry, Warning,
+				TEXT("Tag \"%.*s\" has unresolved rules because some classes were not found during INI parsing. Will try to resolve again."),
+				TagName.Len(), TagName.GetData()
+			);
+			ResolveRules(TagName, OrderedList, bLoadClasses);
+			TagRules.bContainsAnyUnresolvedRule = false;
+		}
+
+		ClassesToAdd.Reset();
+		// The list is sorted in heap order of Inheritance depth,
+		// so we can just walk through it and add/remove as needed
+		TArray<UClass*> Temp;
+		for (const FFilterTagsRule& Info : OrderedList)
+		{
+			// Add Wildcard or unknown class (we guarantee to only have one of those)
+			if (!Info.ClassPtr)
+			{
+				UE_CLOG(OrderedList.Num() > 1, LogAssetRegistry, Warning, 
+					TEXT("There is more than one rule for tag \"%.*s\" with unknown class, ignoring."),
+					TagName.Len(), TagName.GetData()
+				);
+				if (Info.bInclude)
+				{
+					Result.FindOrAdd(Info.ClassPath).Add(Tag);
+				}
+				break;
+			}
+
+			Temp.Reset();
+			GetDerivedClasses(Info.ClassPtr, Temp, true);
+			if (Info.bInclude)
+			{
+				ClassesToAdd.Reserve(ClassesToAdd.Num() + Temp.Num() + 1);
+				ClassesToAdd.Add(Info.ClassPtr);
+				ClassesToAdd.Append(Temp);
+			}
+			else
+			{
+				for (UClass* Derived : Temp)
+				{
+					ClassesToAdd.Remove(Derived);
+				}
+				ClassesToAdd.Remove(Info.ClassPtr);
+			}
+		}
+
+		// After the final list of classes is known - add them to the CookFilterlistTagsByClass
+		for (const UClass* Class : ClassesToAdd)
+		{
+			Result.FindOrAdd(Class->GetClassPathName()).Add(Tag);
+		}
+	}
+
+	return Result;
 }
 
 }
@@ -1759,10 +1955,6 @@ void FAssetRegistryImpl::RefreshNativeClasses()
 	// and serialization options
 	CollectCodeGeneratorClasses();
 	TempCachedInheritanceBuffer.bDirty = true;
-
-	// Read default serialization options
-	Utils::InitializeSerializationOptionsFromIni(SerializationOptions, nullptr, UE::AssetRegistry::ESerializationTarget::ForGame);
-	Utils::InitializeSerializationOptionsFromIni(DevelopmentSerializationOptions, nullptr, UE::AssetRegistry::ESerializationTarget::ForDevelopment);
 }
 
 }
diff --git a/Engine/Source/Runtime/AssetRegistry/Private/AssetRegistryImpl.h b/Engine/Source/Runtime/AssetRegistry/Private/AssetRegistryImpl.h
--- a/Engine/Source/Runtime/AssetRegistry/Private/AssetRegistryImpl.h
+++ b/Engine/Source/Runtime/AssetRegistry/Private/AssetRegistryImpl.h
@@ -1034,7 +1034,33 @@ void RunAssetsThroughFilter(TArray<FAssetData>& AssetDataList, const FARCompiled
 void InitializeSerializationOptionsFromIni(
 	FAssetRegistrySerializationOptions& Options,
 	const ITargetPlatform* TargetPlatform = nullptr,
-	ESerializationTarget Target = ESerializationTarget::ForGame);
+	ESerializationTarget Target = ESerializationTarget::ForGame,
+	bool bInitializeFilterTagOptions = true);
+
+struct FFilterTagsRule
+{
+	FTopLevelAssetPath ClassPath;
+	const UClass* ClassPtr = nullptr;
+	int32 InheritanceDepth = 0;
+	bool bInclude = false;
+};
+
+struct FParsedFilterTagRules
+{
+	TArray<FFilterTagsRule> OrderedList;
+	bool bContainsAnyUnresolvedRule = false;
+};
+/**
+ * Parses the list of CookedTagsBlacklist or CookedTagsWhitelist gathered from the INI.
+ * Returns a mapping from a tag to a list of rules. May contain unresolved rules.
+ */
+TMap<FName, FParsedFilterTagRules> ParseFilterTags(const TArray<FString>& FilterListItems, bool bIsFilteringDevelopmentAR, bool bUseAllowListInsteadOfDenyList);
+
+/**
+ * Resolves and expands the list of rules to list of all derived loaded classes that satisfy those rules.
+ * Then transposes rules to a map with a class as a key and a list of tags as value.
+ */
+TMap<FTopLevelAssetPath, TSet<FName>> ResolveFilterTags(TMap<FName, FParsedFilterTagRules>& InOutRules, bool bLoadClasses);
 
 /** Gets the current availability of an asset, primarily for streaming install purposes. */
 EAssetAvailability::Type GetAssetAvailability(const FAssetData& AssetData);
