[AssetRegistry]

- Extends the save gather cache timer to be 5min instead of 30s. In the warm gather case, AR initialization takes longer than 30s which guarantees on the "fast" path you will request the cache be saved at least twice (since we save when complete) which is wasteful. In the cold cases, caching can take > 5min, so it doesn't feel too wasteful to cache our progress at a slower cadence since saving the cache itself can be expensive (10s+) which can add up to minutes of extra time during cold gathering.
- The AssetDataGatherer may now request the gather thread to run the AR tick if gatherer has collected some results. This is to help avoid AR gather work from ending up on the game thread's AR tick when it's not necessary (i.e. not an explicit Wait/SynchronousScan)
- Adds LongPackageName parameter to the PackageReader OpenPackageFile overload that allows a non-owned loader. Failing to pass in the package name when passing in a loader explicitly will cause package reader failures otherwise
- Adds a FNonBufferingReadOnlyArchive which, on windows, opens an archive that enables using FILE_FLAG_NO_BUFFERING which will skip the OS CacheManager. When reading in parallel the cache manager can fight itself trying to allocate cache pages. Instead for package reading we know we will sequentially read a (typically) small mount of data which we can write directly to our own re-uable buffer within the archive bypassing this contention. A follow-up change will make the archive async read aware. This loader is only used when no platform file wrapping is involved since the archive will request more data than might be available in the file which isn't an error on non-buffering platform APIs but likely breaks the platform file wrappers assumptions.
- Changes how we call ReadAssetFiles when gathering to use a TaskConcurrencyLimiter to only allow 32 task threads to read files.
- ReadAssetFile tasks are now fully asyncronous in that we won't block the gather thread by using a ParallelFor allowing the gather thread to processes reads when they are complete while enqueing more work for task threads to pick up as they come in.

Booting into an empty map with Standby cache emptied,
Before:
- Cold Gather: AssetRegistryGather time 828.4724s: AssetDataDiscovery 17.6921s, AssetDataGather 693.0968s, StoreResults 117.6835s. Wall time 960.0700s.
- Warm Gather: AssetRegistryGather time 216.3906s: AssetDataDiscovery 12.9178s, AssetDataGather 19.4677s, StoreResults 184.0051s. Wall time 225.4120s.

After:
- Cold Gather: AssetRegistryGather time 378.0449s: AssetDataDiscovery 12.6879s, AssetDataGather 56.8019s, StoreResults 308.5551s. Wall time 421.6330s.
- Warm Gather: AssetRegistryGather time 131.7834s: AssetDataDiscovery 11.5907s, AssetDataGather 17.9331s, StoreResults 102.2596s. Wall time 136.3650s

#rnx

diff --git a/Engine/Source/Runtime/AssetRegistry/Internal/AssetRegistry/PackageReader.h b/Engine/Source/Runtime/AssetRegistry/Internal/AssetRegistry/PackageReader.h
--- a/Engine/Source/Runtime/AssetRegistry/Internal/AssetRegistry/PackageReader.h
+++ b/Engine/Source/Runtime/AssetRegistry/Internal/AssetRegistry/PackageReader.h
@@ -54,7 +54,8 @@ public:
 	ASSETREGISTRY_API bool OpenPackageFile(FStringView PackageFilename, EOpenPackageResult* OutErrorCode = nullptr);
 	ASSETREGISTRY_API bool OpenPackageFile(FStringView LongPackageName, FStringView PackageFilename,
 		EOpenPackageResult* OutErrorCode = nullptr);
-	ASSETREGISTRY_API bool OpenPackageFile(FArchive* Loader, EOpenPackageResult* OutErrorCode = nullptr);
+	ASSETREGISTRY_API bool OpenPackageFile(FArchive* Loader, FStringView LongPackageName = FStringView(), 
+		EOpenPackageResult* OutErrorCode = nullptr);
 	ASSETREGISTRY_API bool OpenPackageFile(TUniquePtr<FArchive> Loader, EOpenPackageResult* OutErrorCode = nullptr);
 	ASSETREGISTRY_API bool OpenPackageFile(EOpenPackageResult& OutErrorCode);
 
diff --git a/Engine/Source/Runtime/AssetRegistry/Private/AssetDataGatherer.cpp b/Engine/Source/Runtime/AssetRegistry/Private/AssetDataGatherer.cpp
--- a/Engine/Source/Runtime/AssetRegistry/Private/AssetDataGatherer.cpp
+++ b/Engine/Source/Runtime/AssetRegistry/Private/AssetDataGatherer.cpp
@@ -39,6 +39,7 @@
 #include "String/Find.h"
 #include "String/LexFromString.h"
 #include "Tasks/Task.h"
+#include "Tasks/TaskConcurrencyLimiter.h"
 #include "TelemetryRouter.h"
 #include "Templates/UniquePtr.h"
 #include "Templates/UnrealTemplate.h"
@@ -1812,6 +1813,7 @@ void FAssetDataDiscovery::TickInternal(bool bTickAll)
 							// Some setting has been applied to the ScanDir that requires a new scan
 							// Consume the invalidated flag and ignore the results of our scan
 							Data.ScanDir->SetScanInFlightInvalidated(false);
+							UE_LOG(LogAssetRegistry, Display, TEXT("ScanDir %s has been invalidated (Current State: %s). Will not scan."), *Data.ScanDir->GetLocalAbsPath(), (Data.ScanDir->IsValid() ? TEXT("valid") : TEXT("invalid")));
 						}
 						else
 						{
@@ -2304,7 +2306,6 @@ void FAssetDataDiscovery::SetIsIdle(bool bInIsIdle, double& TickStartTime)
 	}
 }
 
-
 void FAssetDataDiscovery::GetAndTrimSearchResults(bool& bOutIsComplete, TArray<FString>& OutDiscoveredPaths, FFilesToSearch& OutFilesToSearch, int32& OutNumPathsToSearch)
 {
 	FGathererScopeLock ResultsScopeLock(&ResultsLock);
@@ -2314,7 +2315,7 @@ void FAssetDataDiscovery::GetAndTrimSearchResults(bool& bOutIsComplete, TArray<F
 
 	for (FDirectoryResult& DirectoryResult : DiscoveredFiles)
 	{
-		OutFilesToSearch.AddDirectory(MoveTemp(DirectoryResult.DirAbsPath), MoveTemp(DirectoryResult.Files));
+		OutFilesToSearch.AddFiles(MoveTemp(DirectoryResult.Files));
 	}
 	DiscoveredFiles.Reset();
 	for (FGatheredPathData& FileResult : DiscoveredSingleFiles)
@@ -3774,6 +3775,22 @@ FPreloader GPreloader;
 
 } // namespace UE::AssetDataGather::Private
 
+static inline int32 GetNumberOfReadTaskThreadsToAllow()
+{
+	using namespace UE::AssetDataGather::Private;
+	GGatherSettings.Initialize();
+
+	// Cap how wide we allow read tasks to go. We don't want too much parallelism 
+	// since ultimately there is stress on the OS for fetching things like File 
+	// Control Blocks on Windows that can't be done in parallel. Experimentally, 
+	// half the available physical cores has performed well.
+	const int32 MaxReadThreads = FPlatformMisc::NumberOfCores();
+	const int32 DefaultNumReadThreads = MaxReadThreads / 2;
+	const int32 NumReadThreads = FGatherSettings::GARGatherThreads <= 0 ? DefaultNumReadThreads : FGatherSettings::GARGatherThreads;
+	const int32 MinReadThreads = 1; // The TaskConcurrencyLimiter doesn't allow 0 worker threads
+	return FMath::Min(MaxReadThreads, FMath::Max(MinReadThreads, NumReadThreads));
+}
+
 FAssetDataGatherer::FAssetDataGatherer(UAssetRegistryImpl& InRegistryImpl)
 	: AssetRegistry(InRegistryImpl)
 	, Thread(nullptr)
@@ -3788,13 +3805,13 @@ FAssetDataGatherer::FAssetDataGatherer(UAssetRegistryImpl& InRegistryImpl)
 	, bHasLoadedCache(false)
 	, bDiscoveryIsComplete(false)
 	, bIsComplete(false)
+	, bRequestAssetRegistryTick(false)
 	, bIsIdle(false)
 	, bFirstTickAfterIdle(true)
 	, bFinishedInitialDiscovery(false)
 	, bIsInitialSearchCompleted(false)
 	, bIsAdditionalMountSearchInProgress(false)
 	, bGatherOnGameThreadOnly(false)
-	, WaitBatchCount(-1)
 	, LastCacheSaveNumUncachedAssetFiles(0)
 	, CacheInUseCount(0)
 	, bIsSavingAsyncCache(false)
@@ -3802,13 +3819,7 @@ FAssetDataGatherer::FAssetDataGatherer(UAssetRegistryImpl& InRegistryImpl)
 {
 	using namespace UE::AssetDataGather::Private;
 
-	int32 NumGatherThreads = FMath::Max(1, FTaskGraphInterface::Get().GetNumWorkerThreads());
 	GGatherSettings.Initialize();
-	if (FGatherSettings::GARGatherThreads > 0)
-	{
-		NumGatherThreads = FMath::Min(NumGatherThreads, FGatherSettings::GARGatherThreads);
-	}
-	TickInternalBatchSize = NumGatherThreads * FGatherSettings::SingleThreadFilesPerBatch;
 
 	constexpr bool bEditorExecutable = WITH_EDITOR;
 	bGatherAssetPackageData = bEditorExecutable || GGatherSettings.IsForceDependsGathering();
@@ -3834,6 +3845,7 @@ FAssetDataGatherer::FAssetDataGatherer(UAssetRegistryImpl& InRegistryImpl)
 	}
 
 	FilesToSearch = MakeUnique<UE::AssetDataGather::Private::FFilesToSearch>();
+	FileReadScheduler = MakeUnique<FFileReadScheduler>(GetNumberOfReadTaskThreadsToAllow(), bGatherAssetPackageData, bGatherDependsData);
 }
 
 FAssetDataGatherer::~FAssetDataGatherer()
@@ -3847,7 +3859,7 @@ FAssetDataGatherer::~FAssetDataGatherer()
 		delete AssetData;
 	}
 	NewCachedAssetData.Empty();
-	for (TPair<int32, FDiskCachedAssetData*> BlockData : DiskCachedAssetBlocks)
+	for (TPair<int32, FDiskCachedAssetData*>& BlockData : DiskCachedAssetBlocks)
 	{
 		delete[] BlockData.Get<1>();
 	}
@@ -3901,6 +3913,7 @@ void FAssetDataGatherer::OnInitialSearchCompleted()
 
 	RequestAsyncCacheSave();
 	bIsInitialSearchCompleted.store(true, std::memory_order_relaxed);
+	FileReadScheduler->FreeReadTaskLoaders();
 }
 
 void FAssetDataGatherer::OnAdditionalMountSearchCompleted()
@@ -3912,6 +3925,7 @@ void FAssetDataGatherer::OnAdditionalMountSearchCompleted()
 
 	RequestAsyncCacheSave();
 	bIsAdditionalMountSearchInProgress.store(false, std::memory_order_relaxed);
+	FileReadScheduler->FreeReadTaskLoaders();
 }
 
 void FAssetDataGatherer::StartAsync()
@@ -3941,24 +3955,19 @@ uint32 FAssetDataGatherer::Run()
 		ETickResult TickResult = InnerTickLoop(false /* bInSynchronousTick */, true /* bContributeToCacheSave */,
 			-1. /* EndTimeSeconds */);
 
-		for (;;)
+		if (bRequestAssetRegistryTick)
 		{
-			bool bLocalIdle = false;
-			{
-				FGathererScopeLock ResultsScopeLock(&ResultsLock); // bIsIdle requires the lock
-				if (IsStopped || bSaveAsyncCacheTriggered || (!IsGatheringPaused && !bIsIdle))
-				{
-					break;
-				}
-				bLocalIdle = bIsIdle;
-			}
+			bRequestAssetRegistryTick = false;
+			TryTickOnBackgroundThread();
+		}
 
+		bool bLocalIdle = false;
+		while (ShouldBackgroundGatherThreadPauseGathering(bLocalIdle))
+		{
 			UE::AssetRegistry::Impl::EGatherStatus Status = UE::AssetRegistry::Impl::EGatherStatus::Complete;
-			if (bLocalIdle && !IsGatherOnGameThreadOnly()
-				&& (IsProcessingPaused.load(std::memory_order_relaxed) == 0))
+			if (bLocalIdle)
 			{
-				IAssetRegistry& Registry = IAssetRegistry::GetChecked();
-				Status = Cast<UAssetRegistryImpl>(&Registry)->TickOnBackgroundThread();
+				Status = TryTickOnBackgroundThread();
 			}
 
 			// TODO: Need IsGatheringPaused to be a condition variable so we avoid sleeping while waiting for it and then taking a long time to wake after it is unset.
@@ -4075,6 +4084,25 @@ FAssetDataGatherer::ETickResult FAssetDataGatherer::InnerTickLoop(bool bInSynchr
 	return Result;
 }
 
+bool FAssetDataGatherer::ShouldBackgroundGatherThreadPauseGathering(bool& bOutGathererIsIdle)
+{
+	FGathererScopeLock ResultsScopeLock(&ResultsLock); // bIsIdle requires the lock
+	bOutGathererIsIdle = bIsIdle;
+	return !IsStopped && !bSaveAsyncCacheTriggered && (IsGatheringPaused || bIsIdle);
+}
+
+UE::AssetRegistry::Impl::EGatherStatus FAssetDataGatherer::TryTickOnBackgroundThread()
+{
+	UE::AssetRegistry::Impl::EGatherStatus Status = UE::AssetRegistry::Impl::EGatherStatus::Complete;
+	if (!IsGatherOnGameThreadOnly() && (IsProcessingPaused.load(std::memory_order_relaxed) == 0))
+	{
+		IAssetRegistry& Registry = IAssetRegistry::GetChecked();
+		Status = Cast<UAssetRegistryImpl>(&Registry)->TickOnBackgroundThread();
+	}
+
+	return Status;
+}
+
 void FAssetDataGatherer::SaveCacheFile(const TArray<TPair<FName,FDiskCachedAssetData*>>& AssetsToSave)
 {
 	using namespace UE::AssetDataGather::Private;
@@ -4223,9 +4251,11 @@ void FAssetDataGatherer::EnsureCompletion()
 FAssetDataGatherer::ETickResult FAssetDataGatherer::TickInternal(double& TickStartTime, bool bPollDiscovery)
 {
 	using namespace UE::AssetDataGather::Private;
+	using namespace UE::Tasks;
 
 	typedef TInlineAllocator<FGatherSettings::ExpectedMaxBatchSize> FBatchInlineAllocator;
 
+	CHECK_IS_LOCKED_CURRENT_THREAD(TickLock);
 	TArray<FGatheredPathData, FBatchInlineAllocator> LocalFilesToSearch;
 	TArray<TUniquePtr<FAssetData>, FBatchInlineAllocator> LocalAssetResults;
 	TArray<TUniquePtr<FAssetData>, FBatchInlineAllocator> LocalAssetResultsForGameThread;
@@ -4234,10 +4264,11 @@ FAssetDataGatherer::ETickResult FAssetDataGatherer::TickInternal(double& TickSta
 	TArray<FString, FBatchInlineAllocator> LocalCookedPackageNamesWithoutAssetDataResults;
 	TArray<FName, FBatchInlineAllocator> LocalVerseResults;
 	TArray<FString, FBatchInlineAllocator> LocalBlockedResults;
+	TArray<FString> LocalPriorityDirectories;
+	FWaitBatchDirectorySet LocalWaitBatchDirectories;
 	bool bLoadCache = false;
 	bool bLocalIsCacheWriteEnabled = false;
 	double LocalLastCacheWriteTime = 0.0;
-	bool bWaitBatchCountDecremented = false;
 	ETickResult TickResult = ETickResult::KeepTicking;
 
 	{
@@ -4259,14 +4290,11 @@ FAssetDataGatherer::ETickResult FAssetDataGatherer::TickInternal(double& TickSta
 			FilesToSearch->RetryLaterRetryFiles();
 		}
 
-		// Take a batch off of the work list. If we're waiting only on the first WaitBatchCount results don't take more than that
-		int32 NumToProcess = FMath::Min<int32>(TickInternalBatchSize - LocalFilesToSearch.Num(), FilesToSearch->GetNumAvailable());
-		// If no work is available mark idle and exit
-		if (NumToProcess == 0)
+		if (FilesToSearch->GetNumAvailable() == 0 && NumUncachedAssetFilesOutstanding == 0)
 		{
-			if (WaitBatchCount != -1)
+			if (WaitBatchDirectories.IsSet())
 			{
-				WaitBatchCount = -1; // WaitBatchCount was set equal to or higher than FilesToSearch->GetNumAvailable(), mark it completed
+				WaitBatchDirectories.Reset(); // WaitBatchDirectories was set after we've already gathered all results, mark it completed
 				TickResult = ETickResult::Interrupt;
 			}
 
@@ -4287,31 +4315,23 @@ FAssetDataGatherer::ETickResult FAssetDataGatherer::TickInternal(double& TickSta
 			return TickResult;
 		}
 
-		if (WaitBatchCount >= 0)
+		if (WaitBatchDirectories.IsSet())
 		{
-			if (WaitBatchCount == 0)
+			if (WaitBatchDirectories->IsEmpty())
 			{
 				// We've finished executing the caller's requested batchcount and we have fulfilled our check-for-idle
 				// contract above (and found we are not yet idle), so exit now without doing any further work.
-				WaitBatchCount = -1;
+				WaitBatchDirectories.Reset();
 				return ETickResult::Interrupt;
 			}
 
 			// Otherwise we still have some work to do for the caller's requested batchcount, so do work up to
 			// that batchcount.
-			bWaitBatchCountDecremented = true;
-			NumToProcess = FMath::Min(NumToProcess, WaitBatchCount);
-			WaitBatchCount -= NumToProcess;
-			// WaitBatchCount might now be 0, but we do not set TickResult == ETickResult::Interrupt until the next
-			// call to TickInternal, because we need to execute the code above to check for whether we are idle; we
-			// have a contract that we set Idle if possible before breaking out of the InnerTickLoop when processing
-			// a WaitBatchCount. This is important for the case of ScanPathsSynchronous on an asset that needs
-			// LoadCalculatedDependencies. LoadCalculatedDependencies can not run until the Gatherer reports Idle.
+			LocalWaitBatchDirectories = *WaitBatchDirectories;
 		}
 		DependencyResults.Reserve(FilesToSearch->GetNumAvailable() + DependencyResults.Num());
-		check(NumToProcess > 0);
-		FilesToSearch->PopFront(LocalFilesToSearch, NumToProcess);
-		check(LocalFilesToSearch.Num() > 0);
+		FilesToSearch->PopFront(LocalFilesToSearch, FilesToSearch->GetNumAvailable());
+		FilesToSearch->PopPriorityDirectories(LocalPriorityDirectories);
 
 		if (bCacheReadEnabled && !bHasLoadedCache)
 		{
@@ -4351,31 +4371,12 @@ FAssetDataGatherer::ETickResult FAssetDataGatherer::TickInternal(double& TickSta
 		LocalLastCacheWriteTime = LastCacheWriteTime;
 	}
 
-	struct FReadContext
-	{
-		FName PackageName;
-		FName Extension;
-		FGatheredPathData& AssetFileData;
-		TArray<FAssetData*> AssetDataFromFile;
-		FPackageDependencyData DependencyData;
-		TArray<FString> CookedPackageNamesWithoutAssetData;
-		UE::AssetRegistry::Private::EOptionalBool HasMarkOfTheWeb = UE::AssetRegistry::Private::EOptionalBool::Unset;
-		bool bCanAttemptAssetRetry = false;
-		bool bResult = false;
-		bool bCanceled = false;
-
-		FReadContext(FName InPackageName, FName InExtension, FGatheredPathData& InAssetFileData)
-			: PackageName(InPackageName)
-			, Extension(InExtension)
-			, AssetFileData(InAssetFileData)
-		{
-		}
-	};
-
 	int32 NewCachedAssetFiles = 0;
 	int32 NewUncachedAssetFiles = 0;
+	int32 NewOutstandingReads = 0;
+	int32 NumReadContextsToProcess = 0;
+
 	// Try to read each file in the batch out of the cache, and accumulate a list for more expensive reading of all of the files that are not in the cache 
-	TArray<FReadContext> ReadContexts;
 	for (FGatheredPathData& AssetFileData : LocalFilesToSearch)
 	{
 		// If this a Verse-related file, just directly add its file name to the Verse results
@@ -4479,65 +4480,45 @@ FAssetDataGatherer::ETickResult FAssetDataGatherer::TickInternal(double& TickSta
 		else
 		{
 			// Not found in cache (or stale) - schedule to be read from disk
-			ReadContexts.Emplace(PackageName, FName(FPathViews::GetExtension(AssetFileData.LocalAbsPath)), AssetFileData);
+			FileReadScheduler->AddFile(PackageName, MoveTemp(AssetFileData));
+			NewOutstandingReads++;
 		}
 	}
 
-	// For all the files not found in the cache, read them from their package files on disk; the file reads are done in parallel
-	static const EParallelForFlags ParallelFlags = []()
-		{
-			EParallelForFlags ReturnFlags = EParallelForFlags::Unbalanced;
-			if (!IsRunningCommandlet())
-			{
-				ReturnFlags |= EParallelForFlags::BackgroundPriority;
-			}
-			return ReturnFlags;
-		}();
-	// We want to restrict the number of threads, but ParallelFor only provides an API for restricting MinBatchSize
-	// NumberOfThreads == ParallelForNum/MinBatchSize == TickInternalBatchSize/SingleThreadFilesPerBatch
-	//                 == Min(WorkerThreads, GARGatherThreads)*SingleThreadFilesPerBatch / SingleThreadFilesPerBatch
-	//                 <= GARGatherThreads
-	const int32 MinBatchSize = FGatherSettings::SingleThreadFilesPerBatch;
-
-	ParallelFor(TEXT("AssetDataGatherReadAssetFile"), ReadContexts.Num(), MinBatchSize,
-		[this, &ReadContexts](int32 Index)
-		{
-			FReadContext& ReadContext = ReadContexts[Index];
-			if (!bSynchronousTick && IsGatheringPaused)
-			{
-				ReadContext.bCanceled = true;
-				return;
-			}
-			UE_SCOPED_IO_ACTIVITY(*WriteToString<512>(TEXT("Loading Asset"), ReadContext.PackageName.ToString()));
-			ReadContext.bResult = ReadAssetFile(ReadContext.AssetFileData.LongPackageName,
-				ReadContext.AssetFileData.LocalAbsPath, ReadContext.AssetDataFromFile, ReadContext.DependencyData,
-				ReadContext.CookedPackageNamesWithoutAssetData, ReadContext.bCanAttemptAssetRetry);
-			if (ReadContext.bResult && bBlockPackagesWithMarkOfTheWeb)
-			{
-				ReadContext.HasMarkOfTheWeb = BoolToOptionalBool(
-					IPlatformFile::GetPlatformPhysical().HasMarkOfTheWeb(ReadContext.AssetFileData.LocalAbsPath));
-			}
-		},
-		ParallelFlags
-	);
+	// Wait on any files being explicitly waited on before scheduling other reads.
+	// Cache if we are in a wait batch now since WaitOnDirectoriesAndUpdateWaitList will update
+	// our wait list, and we don't want to perform a gather, even if we waited on nothing since 
+	// the game thread is waiting for the tick to return.
+	bool bInWaitBatch = !LocalWaitBatchDirectories.IsEmpty();
+	FileReadScheduler->WaitOnDirectoriesAndUpdateWaitList(LocalWaitBatchDirectories);
+
+	// Call Schedule to push any recently added uncached files to the read queue, and also give it a list of any new directory priorities we have
+	FileReadScheduler->Schedule(LocalPriorityDirectories);
+
+	// Collect any completed reads for processing on this thread
+	TArray<TUniquePtr<FReadContext>> CompletedReadContexts;
+	FileReadScheduler->CollectCompletedReads(CompletedReadContexts);
+
+	// If we have finished any async reads outside of a wait batch, signal that we want to 
+	// kick the gatherer to help avoid gather work from running on the game thread
+	bRequestAssetRegistryTick = !bInWaitBatch && !CompletedReadContexts.IsEmpty();
 
 	// Accumulate the results
-	bool bHasCancelation = false;
-	for (FReadContext& ReadContext : ReadContexts)
+	for (TUniquePtr<FReadContext>& ReadContext : CompletedReadContexts)
 	{
-		if (ReadContext.bCanceled)
+		if (ReadContext->bCanceled)
 		{
-			bHasCancelation = true;
+			ReadContext->DirectoryReadTaskData->bHasCanceledRead = true;
 		}
-		else if (ReadContext.bResult)
+		else if (ReadContext->bResult)
 		{
 			// Do not add the results from a cooked package into the map of data we keep to write out the new version of the cache file
 			bool bCachePackage = bLocalIsCacheWriteEnabled
-				&& ReadContext.CookedPackageNamesWithoutAssetData.Num() == 0
-				&& ensure(ReadContext.AssetFileData.Type == EGatherableFileType::PackageFile);
+				&& ReadContext->CookedPackageNamesWithoutAssetData.Num() == 0
+				&& ensure(ReadContext->AssetFileData.Type == EGatherableFileType::PackageFile);
 			if (bCachePackage)
 			{
-				for (const FAssetData* AssetData : ReadContext.AssetDataFromFile)
+				for (const FAssetData* AssetData : ReadContext->AssetDataFromFile)
 				{
 					if (!!(AssetData->PackageFlags & PKG_FilterEditorOnly))
 					{
@@ -4551,33 +4532,34 @@ FAssetDataGatherer::ETickResult FAssetDataGatherer::TickInternal(double& TickSta
 			if (bCachePackage)
 			{
 				// Update the cache
-				FDiskCachedAssetData* NewData = new FDiskCachedAssetData(ReadContext.AssetFileData.PackageTimestamp, GatherStartTime, ReadContext.Extension);
-				NewData->AssetDataList.Reserve(ReadContext.AssetDataFromFile.Num());
-				for (const FAssetData* BackgroundAssetData : ReadContext.AssetDataFromFile)
+				FDiskCachedAssetData* NewData = new FDiskCachedAssetData(ReadContext->AssetFileData.PackageTimestamp, 
+					GatherStartTime, ReadContext->Extension);
+				NewData->AssetDataList.Reserve(ReadContext->AssetDataFromFile.Num());
+				for (const FAssetData* BackgroundAssetData : ReadContext->AssetDataFromFile)
 				{
 					NewData->AssetDataList.Add(*BackgroundAssetData);
 				}
 
-				NewData->DependencyData = ReadContext.DependencyData;
-				NewData->HasMarkOfTheWeb = ReadContext.HasMarkOfTheWeb;
+				NewData->DependencyData = ReadContext->DependencyData;
+				NewData->HasMarkOfTheWeb = ReadContext->HasMarkOfTheWeb;
 
 				NewCachedAssetData.Add(NewData);
-				AddToCache(ReadContext.PackageName, NewData);
+				AddToCache(ReadContext->PackageName, NewData);
 			}
 
 			bool bPackageBlocked = bBlockPackagesWithMarkOfTheWeb
-				&& ReadContext.HasMarkOfTheWeb == UE::AssetRegistry::Private::EOptionalBool::True;
+				&& ReadContext->HasMarkOfTheWeb == UE::AssetRegistry::Private::EOptionalBool::True;
 			if (bPackageBlocked)
 			{
 				// If this is a blocked package, add its file path to the Blocked results, and not the published result.
 				// It still will have been cached above, so we don't have to read its data next time.
-				LocalBlockedResults.Add(ReadContext.AssetFileData.LocalAbsPath);
+				LocalBlockedResults.Add(ReadContext->AssetFileData.LocalAbsPath);
 			}
 			else
 			{
 				++NewUncachedAssetFiles;
 				// Add the results from a cooked package into our results on cooked package
-				LocalCookedPackageNamesWithoutAssetDataResults.Append(MoveTemp(ReadContext.CookedPackageNamesWithoutAssetData));
+				LocalCookedPackageNamesWithoutAssetDataResults.Append(MoveTemp(ReadContext->CookedPackageNamesWithoutAssetData));
 
 				bool MustBeHandledByGameThread = false;
 				// 
@@ -4603,17 +4585,11 @@ FAssetDataGatherer::ETickResult FAssetDataGatherer::TickInternal(double& TickSta
 					MustBeHandledByGameThread ? LocalDependencyResultsForGameThread : LocalDependencyResults;
 
 				// Add the results from the package into our output results
-				TargetAssetResults.Append(MoveTemp(ReadContext.AssetDataFromFile));
-				FPackageDependencyData& DependencyData = TargetDependencyResults.Add_GetRef(MoveTemp(ReadContext.DependencyData));
-				DependencyData.PackageData.SetIsReadOnly(ReadContext.AssetFileData.bIsReadOnly);
+				TargetAssetResults.Append(MoveTemp(ReadContext->AssetDataFromFile));
+				FPackageDependencyData& DependencyData = TargetDependencyResults.Add_GetRef(MoveTemp(ReadContext->DependencyData));
+				DependencyData.PackageData.SetIsReadOnly(ReadContext->AssetFileData.bIsReadOnly);
 			}
 		}
-		else if (ReadContext.bCanAttemptAssetRetry)
-		{
-			// If the read temporarily failed, return it to the worklist, pushed to the end
-			FGathererScopeLock ResultsScopeLock(&ResultsLock);
-			FilesToSearch->AddFileForLaterRetry(MoveTemp(ReadContext.AssetFileData));
-		}
 	}
 
 	{
@@ -4630,25 +4606,45 @@ FAssetDataGatherer::ETickResult FAssetDataGatherer::TickInternal(double& TickSta
 
 		NumUncachedAssetFiles += NewUncachedAssetFiles;
 		NumCachedAssetFiles += NewCachedAssetFiles;
+		NumUncachedAssetFilesOutstanding += NewOutstandingReads - CompletedReadContexts.Num();
 
-		if (bHasCancelation)
+		// Pop all ready contexts so they can be destroyed and we can 
+		// cleanup our DirectoryReadTaskData as necessary
+		while (!CompletedReadContexts.IsEmpty())
 		{
-			// If we skipped reading files due to a pause request, push the canceled files back onto the FilesToSearch
-			for (int Index = ReadContexts.Num() - 1; Index >= 0; --Index) // AddToFront in reverse order so that the elements are readded in the same order they were popped
+			TUniquePtr<FReadContext> ReadContext = CompletedReadContexts.Pop(EAllowShrinking::No);
+
+			// This is the last read for the directory so we can now remove it from the Scheduler
+			if (ReadContext->DirectoryReadTaskData->OutstandingReads == 1)
 			{
-				FReadContext& ReadContext = ReadContexts[Index];
-				if (ReadContext.bCanceled)
+				FString Directory = FPaths::GetPath(ReadContext->AssetFileData.LocalAbsPath);
+				check(!Directory.IsEmpty());
+
+				// Only remove directories from the wait batch when we know they were completed fully (i.e. 
+				// didn't contain a canceled read that will need the directory to be rescheduled)
+				if (bInWaitBatch && !ReadContext->DirectoryReadTaskData->bHasCanceledRead)
 				{
-					FilesToSearch->AddFileAgainAfterTimeout(MoveTemp(ReadContext.AssetFileData));
-					if (bWaitBatchCountDecremented)
-					{
-						// Restore WaitBatchCount to 0 if we marked it done by setting it to -1
-						WaitBatchCount = FMath::Max(0, WaitBatchCount);
-						// Then add back on the work that we thought we were going to do but didn't get to
-						++WaitBatchCount;
-					}
+					LocalWaitBatchDirectories.Remove(Directory);
 				}
+				FileReadScheduler->RemoveScheduledDirectory(MoveTemp(Directory));
 			}
+
+			if (ReadContext->bCanceled)
+			{
+				FilesToSearch->AddFileAgainAfterTimeout(MoveTemp(ReadContext->AssetFileData));
+			}
+			else if(ReadContext->bCanAttemptAssetRetry)
+			{
+				// If the read temporarily failed, return it to the worklist, pushed to the end
+				FilesToSearch->AddFileForLaterRetry(MoveTemp(ReadContext->AssetFileData));
+			}
+		}
+
+		if (bInWaitBatch)
+		{
+			// Update the wait batch, as we will have added and/or removed directories
+			// based on possibly child directories outstanding and read completions
+			WaitBatchDirectories = MoveTemp(LocalWaitBatchDirectories);
 		}
 
 		const int32 NumAssetsReadSinceLastCacheWrite = NumUncachedAssetFiles - LastCacheSaveNumUncachedAssetFiles;
@@ -4796,7 +4792,7 @@ void FAssetDataGatherer::GetAndTrimSearchResults(UE::AssetDataGather::FResults&
 	OutContext.SearchTimes.Append(MoveTemp(SearchTimes));
 	SearchTimes.Reset();
 
-	OutContext.NumFilesToSearch = FilesToSearch->Num();
+	OutContext.NumFilesToSearch = FilesToSearch->Num() + NumUncachedAssetFilesOutstanding;
 	OutContext.NumPathsToSearch = NumPathsToSearchAtLastSyncPoint;
 	OutContext.bIsDiscoveringFiles = !bDiscoveryIsComplete;
 
@@ -4905,7 +4901,7 @@ void FAssetDataGatherer::ClearCache()
 			delete AssetData;
 		}
 		NewCachedAssetData.Empty();
-		for (TPair<int32, FDiskCachedAssetData*> BlockData : DiskCachedAssetBlocks)
+		for (TPair<int32, FDiskCachedAssetData*>& BlockData : DiskCachedAssetBlocks)
 		{
 			delete[] BlockData.Get<1>();
 		}
@@ -4913,7 +4909,6 @@ void FAssetDataGatherer::ClearCache()
 	}
 }
 
-
 void FAssetDataGatherer::ScanPathsSynchronous(const TArray<FString>& InLocalPaths, bool bForceRescan,
 	bool bIgnoreDenyListScanFilters)
 {
@@ -4949,11 +4944,35 @@ void FAssetDataGatherer::WaitOnPathsInternal(TArrayView<UE::AssetDataGather::Pri
 			FGathererScopeLock ResultsScopeLock(&ResultsLock);
 			IngestDiscoveryResults();
 
-			int32 NumDiscoveredPaths;
-			SortPathsByPriority(QueryPaths, UE::AssetDataGather::Private::EPriority::Blocking, NumDiscoveredPaths);
-			// Set WaitBatchCount to valid (non-negative) but possibly to 0. If it is 0 we still want to call TickInternal
+			// Set WaitBatchCount to valid (set) but possibly to empty. If it is empty we still want to call TickInternal
 			// to SetIsIdle(false) if necessary.
-			WaitBatchCount = NumDiscoveredPaths;
+			if (!WaitBatchDirectories.IsSet())
+			{
+				WaitBatchDirectories.Emplace();
+			}
+
+			for (UE::AssetDataGather::Private::FPathExistence& PathExistence : QueryPaths)
+			{
+				using namespace UE::AssetDataGather::Private;
+				FString WaitDir = PathExistence.GetLocalAbsPath();
+				const bool bIsDirectory = PathExistence.GetType() == UE::AssetDataGather::Private::FPathExistence::Directory;
+				if (!bIsDirectory)
+				{
+					FStringView DirectoryPath = FPathViews::GetPath(WaitDir);
+					WaitDir.LeftInline(DirectoryPath.Len());
+				}
+
+				FWaitBatchDirectory* WaitBatchDirectory = WaitBatchDirectories->Find(WaitDir);
+				if (WaitBatchDirectory)
+				{
+					// Ensure the directory is to be scanned recursively
+					WaitBatchDirectory->bIsRecursive = bIsDirectory;
+				}
+				else
+				{
+					WaitBatchDirectories->Add(FWaitBatchDirectory(MoveTemp(WaitDir), bIsDirectory /*bIsRecursive*/));
+				}
+			}
 		}
 	}
 
@@ -4963,7 +4982,7 @@ void FAssetDataGatherer::WaitOnPathsInternal(TArrayView<UE::AssetDataGather::Pri
 	{
 		InnerTickLoop(true /* bInSynchronousTick */, true /* bContributeToCacheSave */, -1. /* EndTimeSeconds */);
 		FGathererScopeLock ResultsScopeLock(&ResultsLock); // WaitBatchCount requires the lock
-		if (WaitBatchCount < 0)
+		if (!WaitBatchDirectories.IsSet())
 		{
 			break;
 		}
@@ -5509,7 +5528,8 @@ TArray<FCachePayload> LoadCacheFiles(TConstArrayView<FString> InCacheFilenames)
 
 	return MoveTemp(Results);
 }
-}
+
+} // namespace UE::AssetDataGather::Private
 
 SIZE_T FAssetDataGatherer::GetAllocatedSize() const
 {
@@ -5586,6 +5606,7 @@ void FAssetDataGatherer::Shrink()
 	BlockedResults.Shrink();
 	SearchTimes.Shrink();
 	DiscoveredPaths.Shrink();
+	FileReadScheduler->Shrink();
 }
 
 void FAssetDataGatherer::AddMountPoint(FStringView LocalPath, FStringView LongPackageName)
@@ -5821,378 +5842,477 @@ FStringView FAssetDataGatherer::NormalizeLongPackageName(FStringView LongPackage
 namespace UE::AssetDataGather::Private
 {
 
-void FFilesToSearch::AddPriorityFile(FGatheredPathData&& FilePath)
+FFileReadScheduler::FFileReadScheduler(int32 InNumReadThreads, bool InGatherAssetPackageData, bool InGatherDependsData)
+	: ReadConcurrencyLimiter(InNumReadThreads, UE::Tasks::ETaskPriority::BackgroundLow)
+	, NumReadTaskThreads(InNumReadThreads)
+	, bGatherAssetPackageData(InGatherAssetPackageData)
+	, bGatherDependsData(InGatherDependsData)
 {
-	++AvailableFilesNum;
-	BlockingFiles.Add(MoveTemp(FilePath));
+	ScheduledDirectories = MakeUnique<TDirectoryTree<TSharedPtr<FDirectoryReadTaskData>>>();
 }
 
-void FFilesToSearch::AddDirectory(FString&& DirAbsPath, TArray<FGatheredPathData>&& FilePaths)
+FFileReadScheduler::~FFileReadScheduler()
 {
-	if (FilePaths.Num() == 0)
-	{
-		return;
-	}
-	check(!DirAbsPath.IsEmpty());
+	// Ensure all read tasks are complete
+	ReadConcurrencyLimiter.Wait();
 
-	FTreeNode& Node = Root.FindOrAddNode(DirAbsPath);
-	AvailableFilesNum += FilePaths.Num();
-	Node.AddFiles(MoveTemp(FilePaths));
+	// Free all allocated read buffers
+	FreeReadTaskLoaders();
 }
 
-void FFilesToSearch::AddFileAgainAfterTimeout(FGatheredPathData&& FilePath)
+void FFileReadScheduler::FreeReadTaskLoaders()
 {
-	++AvailableFilesNum;
-	BlockingFiles.AddFront(MoveTemp(FilePath));
+	TArray<FNonBufferingReadOnlyArchive*> Loaders;
+	ReadTaskLoaderFreeList.PopAll(Loaders);
+	for (FNonBufferingReadOnlyArchive* Loader : Loaders)
+	{
+		delete Loader;
+	}
 }
 
-void FFilesToSearch::AddFileForLaterRetry(FGatheredPathData&& FilePath)
+void FFileReadScheduler::ProcessReadContexts(TArray<TUniquePtr<UE::AssetDataGather::Private::FReadContext>>&& ReadContexts, bool bIsTimeSliced)
 {
-	LaterRetryFiles.Add(FilePath);
-}
+	using namespace UE::AssetDataGather::Private;
+	using namespace UE::Tasks;
 
-void FFilesToSearch::RetryLaterRetryFiles()
-{
-	while (!LaterRetryFiles.IsEmpty())
+	TRACE_CPUPROFILER_EVENT_SCOPE(FAssetDataGatherer::ReadDirectoryAssetFiles);
+
+	// We only use a non-buffering loader if we fully control how archive loading will be done from disk.
+	bool bIsPhysicalPlatformFile = &IPlatformFile::GetPlatformPhysical() == &FPlatformFileManager::Get().GetPlatformFile();
+	double StartTime = FPlatformTime::Seconds();
+	const double TimeLimit = 1.0;
+	int32 MinBatchSize = 32;
+	int32 Iteration = 0;
+	
+	FNonBufferingReadOnlyArchive* Loader = nullptr;
+	if (bIsPhysicalPlatformFile)
 	{
-		FGatheredPathData FilePath = LaterRetryFiles.PopFrontValue();
-		FTreeNode& Node = Root.FindOrAddNode(FPathViews::GetPath(FilePath.LocalAbsPath));
-		++AvailableFilesNum;
-		Node.AddFile(MoveTemp(FilePath));
+		Loader = ReadTaskLoaderFreeList.Pop();
+		if (!Loader)
+		{
+			const int64 BufferSize = 1024 * 1024;
+			Loader = new FNonBufferingReadOnlyArchive(BufferSize);
+		}
+	}
+
+	for (int32 i = 0; i < ReadContexts.Num(); ++i)
+	{
+		if (bIsTimeSliced && Iteration++ > MinBatchSize)
+		{
+			if ((FPlatformTime::Seconds() - StartTime) >= TimeLimit)
+			{
+				for (int32 j = i; j < ReadContexts.Num(); ++j)
+				{
+					TUniquePtr<FReadContext>& ReadContext = ReadContexts[j];
+					ReadContext->bCanceled = true;
+					ReadyReadContexts.Enqueue(MoveTemp(ReadContext));
+				}
+				break;
+			}
+			Iteration = 0;
+		}
+
+		TUniquePtr<FReadContext>& ReadContext = ReadContexts[i];
+		FPackageReader::EOpenPackageResult OpenPackageResult;
+		FPackageReader PackageReader;
+		auto ProcessPackageFile = [this](FPackageReader& PackageReader, FPackageReader::EOpenPackageResult OpenPackageResult, TUniquePtr<FReadContext>& ReadContext)
+			{
+				if (OpenPackageResult != FPackageReader::EOpenPackageResult::Success)
+				{
+					// If we're missing a custom version, we might be able to load this package later once the module containing that version is loaded...
+					//   -	Attempting a retry is only useful when engine startup module is not yet complete and therefore more plugins are expected.
+					const bool bAllowRetry = !IsEngineStartupModuleLoadingComplete();
+					if (OpenPackageResult == FPackageReader::EOpenPackageResult::CustomVersionMissing)
+					{
+						ReadContext->bCanAttemptAssetRetry = bAllowRetry;
+						if (!bAllowRetry)
+						{
+							UE_LOG(LogAssetRegistry, Display, TEXT("Package %s uses an unknown custom version and cannot be loaded for the AssetRegistry"), *ReadContext->AssetFileData.LocalAbsPath);
+						}
+					}
+					else
+					{
+						ReadContext->bCanAttemptAssetRetry = false;
+					}
+					ReadContext->bResult = false;
+				}
+				else
+				{
+					ReadContext->bResult = FAssetDataGatherer::ReadAssetFile(PackageReader, ReadContext->AssetDataFromFile,
+						ReadContext->DependencyData, ReadContext->CookedPackageNamesWithoutAssetData,
+						(bGatherAssetPackageData ? FPackageReader::EReadOptions::PackageData : FPackageReader::EReadOptions::None) |
+						(bGatherDependsData ? FPackageReader::EReadOptions::Dependencies : FPackageReader::EReadOptions::None));
+				}
+
+				if (ReadContext->bResult && bBlockPackagesWithMarkOfTheWeb)
+				{
+					ReadContext->HasMarkOfTheWeb = BoolToOptionalBool(
+						IPlatformFile::GetPlatformPhysical().HasMarkOfTheWeb(ReadContext->AssetFileData.LocalAbsPath));
+				}
+			};
+
+		if (!bIsPhysicalPlatformFile)
+		{		
+			PackageReader.OpenPackageFile(ReadContext->AssetFileData.LongPackageName, 
+				ReadContext->AssetFileData.LocalAbsPath, &OpenPackageResult);
+			ProcessPackageFile(PackageReader, OpenPackageResult, ReadContext);
+		}
+		else if(Loader->OpenFile(*ReadContext->AssetFileData.LocalAbsPath))
+		{
+			PackageReader.OpenPackageFile(Loader, ReadContext->AssetFileData.LongPackageName, &OpenPackageResult);
+			ProcessPackageFile(PackageReader, OpenPackageResult, ReadContext);
+			Loader->Close();
+		}
+		else
+		{
+			UE_LOG(LogAssetRegistry, Warning, TEXT("Package %s could not be opened during gathering. Package will be ignored."),
+				*ReadContext->AssetFileData.LocalAbsPath);
+			ReadContext->bResult = false;
+		}
+
+		ReadyReadContexts.Enqueue(MoveTemp(ReadContext));
+	}
+
+	if (bIsPhysicalPlatformFile)
+	{
+		ReadTaskLoaderFreeList.Push(Loader);
 	}
 }
 
-template <typename AllocatorType>
-void FFilesToSearch::PopFront(TArray<FGatheredPathData, AllocatorType>& Out, int32 NumToPop)
+void FFileReadScheduler::AddFile(const FName PackageName, FGatheredPathData&& AssetFileData)
 {
-	int32 InitialNumToPop = NumToPop;
-	while (NumToPop > 0 && !BlockingFiles.IsEmpty())
+	FString DirectoryPath = FPaths::GetPath(AssetFileData.LocalAbsPath);
+	check(!AssetFileData.LocalAbsPath.IsEmpty());
+	check(!DirectoryPath.IsEmpty());
+
+	bool bAlreadyExisted = false;
+	TSharedPtr<FDirectoryReadTaskData>& DirectoryReadTaskData = ScheduledDirectories->FindOrAdd(DirectoryPath, &bAlreadyExisted);
+	if (!bAlreadyExisted)
 	{
-		Out.Add(BlockingFiles.PopFrontValue());
-		--NumToPop;
+		check(!DirectoryReadTaskData.IsValid());
+		DirectoryReadTaskData = MakeShared<FDirectoryReadTaskData>();
 	}
-	Root.PopFiles(Out, NumToPop);
-	AvailableFilesNum += NumToPop - InitialNumToPop;
-	check(AvailableFilesNum >= 0);
+
+	TArray<TUniquePtr<FReadContext>>& ReadContextsToProcess = DirectoryReadsToSchedule.FindOrAdd(DirectoryPath);
+	ReadContextsToProcess.Emplace(new FReadContext(PackageName,
+		FName(FPathViews::GetExtension(AssetFileData.LocalAbsPath)), MoveTemp(AssetFileData), DirectoryReadTaskData));
 }
 
-void FFilesToSearch::PrioritizePath(FPathExistence& QueryPath, EPriority Priority)
+void FFileReadScheduler::ScheduleDirectory(UE::Tasks::FTaskConcurrencyLimiter& TaskLimiter, const FString& DirectoryPath, 
+	TArray<TUniquePtr<FReadContext>>&& ReadContexts, bool bUnblockImmediately)
 {
-	// We may need to prioritize a LaterRetryFile that is now loadable, so add them all into the Root
-	RetryLaterRetryFiles();
-
-	if (Priority > EPriority::Blocking)
+	using namespace UE::Tasks;
+	TSharedPtr<FDirectoryReadTaskData>* DirectoryReadTaskData = ScheduledDirectories->Find(DirectoryPath);
+	if (!DirectoryReadTaskData)
 	{
-		// TODO: Implement another tree that is searched first for the High Priority 
-		// We cannot add the High Priority files to the BlockingFiles array, because
-		// then blocking on BlockingFiles to be empty could be slow. We cannot add them
-		// as a separate simple array, because we would have to search that (sometimes large)
-		// array linearly when looking for files to accomodate a blocking priority request
 		return;
 	}
+	check(DirectoryReadTaskData->IsValid());
 
-	// We want to look for the querypath without knowing whether it is a file or directory,
-	// because to know whether it is a file or directory we would have to do some IO.
-	// So handle both possibilities.
-	const FString& LocalAbsPath = QueryPath.GetLocalAbsPath();
-	FStringView RemainingRelPath;
-	FTreeNode& TreeNode = Root.FindClosestParentNode(LocalAbsPath, RemainingRelPath);
-	if (RemainingRelPath.IsEmpty())
+	FTask NewReadTask;
+	FTask ExistingTask = (*DirectoryReadTaskData)->LastScheduledReadTask;
+	if (!ReadContexts.IsEmpty())
 	{
-		// The queryPath matches an existing directory in our tree; prioritize that directory by pushing all of its
-		// files (and subdirectory's files) into BlockingFiles
-		TreeNode.PopAllFiles(BlockingFiles);
+		// Schedule a new task that depends on our event being triggered (so the task can be limited) 
+		// but also depends on the directory's existing scheduled task (this is a case where files 
+		// in the directory come in after we have already started loading files of that directory 
+		// previously added to the gatherer for reading)
+		NewReadTask = Launch(TEXT("ReadDirectoryAssetFiles"),
+			[this, ReadContexts = MoveTemp(ReadContexts)]() mutable
+			{
+				ProcessReadContexts(MoveTemp(ReadContexts), true /*bIsTimeSliced*/);
+			}, UE::Tasks::Prerequisites(ExistingTask, (*DirectoryReadTaskData)->UnblockRead),
+				bUnblockImmediately ? LowLevelTasks::ETaskPriority::BackgroundNormal : LowLevelTasks::ETaskPriority::BackgroundLow);
+		(*DirectoryReadTaskData)->LastScheduledReadTask = NewReadTask;
 	}
-	else if (UE::String::FindFirstOfAnyChar(RemainingRelPath, { '/', '\\' }) != INDEX_NONE)
+
+	if (bUnblockImmediately)
 	{
-		// We found a parent directory of a parent directory. Neither QueryPath itself, nor its parent if it is a file,
-		// is in our tree, so there is nothing we need to prioritize for it
+		// Kick off the task immediately, allowing it to run in parallel to
+		// tasks in the task limiter but at a higher priority
+		(*DirectoryReadTaskData)->UnblockRead.Trigger();
 	}
-	else
+	else if(NewReadTask.IsValid()) // If we scheduled a new job
 	{
-		// If QueryPath is a directory, it does not exist in our tree and we don't need to prioritize anything for it.
-		// But it might be a file, and its parent directory exists so look for a matching file in the parent directory.
-		int32 BeforeSize = BlockingFiles.Num();
-		TreeNode.PopMatchingDirectFiles(BlockingFiles, LocalAbsPath);
+		// This task is made such that we can run our work under the limiter's concurrency limit
+		// but still let someone explicitly wait on the task outside of the limiter (possibly surpassing the limiter limit)
+		// Note, we pass in the read task to wait on since if we have a chain of tasks, DirectoryReadTaskData->ReadTask
+		// could change during the execution of our task.
+		TaskLimiter.Push(TEXT("TriggerReadDirectoryAssetFiles"),
+			[DirectoryReadTaskData = *DirectoryReadTaskData, ReadTask = NewReadTask](int32)
+			{
+				DirectoryReadTaskData->UnblockRead.Trigger();
+				ReadTask.Wait();
+			});
 	}
 }
 
-int32 FFilesToSearch::NumBlockingFiles() const
+void FFileReadScheduler::WaitOnDirectoriesAndUpdateWaitList(FWaitBatchDirectorySet& WaitBatchDirectories)
 {
-	return BlockingFiles.Num();
-}
+	using namespace UE::Tasks;
 
-void FFilesToSearch::Shrink()
-{
-	// TODO: Make TRingBuffer::Shrink
-	TRingBuffer<UE::AssetDataGather::Private::FGatheredPathData> Buffer;
-	Buffer.Reserve(BlockingFiles.Num());
-	for (UE::AssetDataGather::Private::FGatheredPathData& File : BlockingFiles)
+	if (WaitBatchDirectories.IsEmpty())
 	{
-		Buffer.Add(MoveTemp(File));
+		return;
 	}
-	Swap(Buffer, BlockingFiles);
 
-	Buffer.Empty(LaterRetryFiles.Num());
-	for (UE::AssetDataGather::Private::FGatheredPathData& File : LaterRetryFiles)
+	TRACE_CPUPROFILER_EVENT_SCOPE(FFileReadScheduler::WaitOnDirectoriesAndUpdateWaitList)
+	// Collect all child directories we might need to wait on
+	FWaitBatchDirectorySet NewWaitBatchDirectories(WaitBatchDirectories);
+	for (FWaitBatchDirectory& WaitBatchDirectory : WaitBatchDirectories)
 	{
-		Buffer.Add(MoveTemp(File));
+		if (!WaitBatchDirectory.bIsRecursive)
+		{
+			continue;
+		}
+
+		TArray<FString> ChildDirectories;
+		if (ScheduledDirectories->TryGetChildren(WaitBatchDirectory.Path, ChildDirectories, 
+			EDirectoryTreeGetFlags::ImpliedParent | EDirectoryTreeGetFlags::Recursive) && !ChildDirectories.IsEmpty())
+		{
+			TStringBuilder<512> AbsChildPath;
+			for (FString& ChildDir : ChildDirectories)
+			{
+				AbsChildPath.Reset();
+				AbsChildPath.Append(FPaths::Combine(WaitBatchDirectory.Path, ChildDir));
+				NewWaitBatchDirectories.Add(FWaitBatchDirectory(AbsChildPath, false /*bIsRecursive*/));
+			}
+		}
 	}
-	Swap(Buffer, LaterRetryFiles);
 
-	Root.Shrink();
-}
+	FWaitBatchDirectorySet EmptyWaitBatchDirectories;
+	FTaskConcurrencyLimiter WaitBatchLimiter(NumReadTaskThreads, LowLevelTasks::ETaskPriority::BackgroundHigh);
+	for (FWaitBatchDirectory& WaitBatchDirectory : NewWaitBatchDirectories)
+	{
+		// Schedule any recently added files for the wait directory
+		TArray<TUniquePtr<FReadContext>> UnscheduledReadContexts;
+		if (DirectoryReadsToSchedule.RemoveAndCopyValue(WaitBatchDirectory.Path, UnscheduledReadContexts) && !UnscheduledReadContexts.IsEmpty())
+		{
+			ScheduleDirectory(WaitBatchLimiter, WaitBatchDirectory.Path, MoveTemp(UnscheduledReadContexts), false /*bUnblockImmediately*/);
+		}
+		// If we have existing scheduled directories, find and wait for them
+		else if (TSharedPtr<FDirectoryReadTaskData>* DirectoryReadTaskData = ScheduledDirectories->Find(WaitBatchDirectory.Path))
+		{
+			check(DirectoryReadTaskData->IsValid());
 
-int32 FFilesToSearch::Num() const
-{
-	return AvailableFilesNum + LaterRetryFiles.Num();
-}
+			// Push a task to unblock our read task. We use a limiter so that we don't trigger all waiting directory reads as once
+			WaitBatchLimiter.Push(TEXT("WaitBatchTriggerReadDirectoryAssetFiles"),
+				[DirectoryReadTaskData = *DirectoryReadTaskData, ReadTask = (*DirectoryReadTaskData)->LastScheduledReadTask](int32)
+				{
+					DirectoryReadTaskData->UnblockRead.Trigger();
+					ReadTask.Wait();
+				});
+		}
+		else
+		{
+			// If there was no directory scheduled, remove this directory from the wait batch;
+			// we likely processed all files for this directory already
+			EmptyWaitBatchDirectories.Add(FWaitBatchDirectory(WaitBatchDirectory.Path, false /*bIsRecursive*/));
+		}
+	}
+	WaitBatchLimiter.Wait();
 
-int32 FFilesToSearch::GetNumAvailable() const
-{
-	return AvailableFilesNum;
+	// Update the passed in wait batch since it may have grown based on child directories, but also remove the empty directories
+	WaitBatchDirectories = NewWaitBatchDirectories.Difference(EmptyWaitBatchDirectories);
 }
 
-SIZE_T FFilesToSearch::GetAllocatedSize() const
+void FFileReadScheduler::Schedule(const TArray<FString>& PriorityDirectories)
 {
-	SIZE_T Size = 0;
-	Size += BlockingFiles.GetAllocatedSize();
-	for (const FGatheredPathData& PathData : BlockingFiles)
+	for (const FString& PriorityDirectory : PriorityDirectories)
 	{
-		Size += PathData.GetAllocatedSize();
+		TArray<TUniquePtr<FReadContext>> ReadContexts;
+		if (DirectoryReadsToSchedule.RemoveAndCopyValue(PriorityDirectory, ReadContexts))
+		{
+			ScheduleDirectory(ReadConcurrencyLimiter, PriorityDirectory, MoveTemp(ReadContexts), true /*bUnblockImmediately*/);
+		}
+		else if (ScheduledDirectories->Find(PriorityDirectory))
+		{
+			ScheduleDirectory(ReadConcurrencyLimiter, PriorityDirectory, {}, true /* bUnblockImmediately */);
+		}
 	}
-	Size += Root.GetAllocatedSize();
-	Size += LaterRetryFiles.GetAllocatedSize();
-	for (const FGatheredPathData& PathData : LaterRetryFiles)
+
+	// If there are any remaining directories, schedule them for background processing
+	for (TMap<FString, TArray<TUniquePtr<FReadContext>>>::TIterator It = DirectoryReadsToSchedule.CreateIterator(); It; ++It)
 	{
-		Size += PathData.GetAllocatedSize();
+		const FString& Directory = It->Key;
+		TArray<TUniquePtr<FReadContext>> ReadContexts = MoveTemp(It->Value);
+		ScheduleDirectory(ReadConcurrencyLimiter, Directory, MoveTemp(ReadContexts), false /*bUnblockImmediately*/);
+		It.RemoveCurrent();
 	}
-	return Size;
+
+	// Because of single ownership of FReadContexts, we must be certain we scheduled all directories with ReadContexts by this point
+	check(DirectoryReadsToSchedule.IsEmpty());
 }
 
-FFilesToSearch::FTreeNode::FTreeNode(FStringView InRelPath)
-	:RelPath(InRelPath)
+void FFileReadScheduler::CollectCompletedReads(TArray<TUniquePtr<FReadContext>>& OutCompletedReadContexts)
 {
+	TUniquePtr<FReadContext> Context;
+	while (ReadyReadContexts.Dequeue(Context))
+	{
+		OutCompletedReadContexts.Add(MoveTemp(Context));
+	}
 }
 
-FStringView FFilesToSearch::FTreeNode::GetRelPath() const
+void FFileReadScheduler::RemoveScheduledDirectory(const FString& Directory)
 {
-	return RelPath;
+	ScheduledDirectories->Remove(Directory);
 }
 
-FFilesToSearch::FTreeNode& FFilesToSearch::FTreeNode::FindOrAddNode(FStringView InRelPath)
+void FFileReadScheduler::Shrink()
 {
-	if (InRelPath.IsEmpty())
-	{
-		return *this;
-	}
-	FStringView FirstComponent;
-	FStringView RemainingPath;
-	FPathViews::SplitFirstComponent(InRelPath, FirstComponent, RemainingPath);
-	FTreeNode& SubDir = FindOrAddSubDir(FirstComponent);
-	return SubDir.FindOrAddNode(RemainingPath);
+	ScheduledDirectories->Shrink();
 }
 
-FFilesToSearch::FTreeNode* FFilesToSearch::FTreeNode::FindNode(FStringView InRelPath)
+void FFilesToSearch::AddPriorityFile(FGatheredPathData&& FilePath)
 {
-	if (InRelPath.IsEmpty())
-	{
-		return this;
-	}
-	FStringView FirstComponent;
-	FStringView RemainingPath;
-	FPathViews::SplitFirstComponent(InRelPath, FirstComponent, RemainingPath);
-	FTreeNode* SubDir = FindSubDir(FirstComponent);
-	if (!SubDir)
-	{
-		return nullptr;
-	}
-
-	return SubDir->FindNode(RemainingPath);
+	++AvailableFilesNum;
+	BlockingFiles.Add(MoveTemp(FilePath));
 }
 
-FFilesToSearch::FTreeNode& FFilesToSearch::FTreeNode::FindClosestParentNode(FStringView InRelPath, FStringView& OutRemainingRelPath)
+void FFilesToSearch::AddFiles(TArray<FGatheredPathData>&& FilePaths)
 {
-	if (InRelPath.IsEmpty())
-	{
-		OutRemainingRelPath = FStringView();
-		return *this;
-	}
-	FStringView FirstComponent;
-	FStringView RemainingPath;
-	FPathViews::SplitFirstComponent(InRelPath, FirstComponent, RemainingPath);
-	FTreeNode* SubDir = FindSubDir(FirstComponent);
-	if (!SubDir)
+	if (FilePaths.Num() == 0)
 	{
-		OutRemainingRelPath = InRelPath;
-		return *this;
+		return;
 	}
+	
+	AvailableFilesNum += FilePaths.Num();
+	NormalPriorityFiles.MoveAppendRange(FilePaths.GetData(), FilePaths.Num());
+}
 
-	return SubDir->FindClosestParentNode(RemainingPath, OutRemainingRelPath);
+void FFilesToSearch::AddFileAgainAfterTimeout(FGatheredPathData&& FilePath)
+{
+	++AvailableFilesNum;
+	BlockingFiles.AddFront(MoveTemp(FilePath));
 }
 
-FFilesToSearch::FTreeNode& FFilesToSearch::FTreeNode::FindOrAddSubDir(FStringView SubDirBaseName)
+void FFilesToSearch::AddFileForLaterRetry(FGatheredPathData&& FilePath)
 {
-	int32 Index = FindLowerBoundSubDir(SubDirBaseName);
-	if (Index == SubDirs.Num() || !FPathViews::Equals(SubDirs[Index]->GetRelPath(), SubDirBaseName))
-	{
-		return *SubDirs.EmplaceAt_GetRef(Index, MakeUnique<FTreeNode>(SubDirBaseName));
-	}
-	else
+	LaterRetryFiles.Add(FilePath);
+}
+
+void FFilesToSearch::RetryLaterRetryFiles()
+{
+	while (!LaterRetryFiles.IsEmpty())
 	{
-		return *SubDirs[Index];
+		FGatheredPathData FilePath = LaterRetryFiles.PopFrontValue();
+		++AvailableFilesNum;
+		NormalPriorityFiles.Add(MoveTemp(FilePath));
 	}
 }
 
-FFilesToSearch::FTreeNode* FFilesToSearch::FTreeNode::FindSubDir(FStringView SubDirBaseName)
+template <typename AllocatorType>
+void FFilesToSearch::PopFront(TArray<FGatheredPathData, AllocatorType>& Out, int32 NumToPop)
 {
-	int32 Index = FindLowerBoundSubDir(SubDirBaseName);
-	if (Index == SubDirs.Num() || !FPathViews::Equals(SubDirs[Index]->GetRelPath(), SubDirBaseName))
+	int32 InitialNumToPop = NumToPop;
+	while (NumToPop > 0 && !BlockingFiles.IsEmpty())
 	{
-		return nullptr;
+		Out.Add(BlockingFiles.PopFrontValue());
+		--NumToPop;
 	}
-	else
+	while (NumToPop > 0 && !NormalPriorityFiles.IsEmpty())
 	{
-		return SubDirs[Index].Get();
+		Out.Add(NormalPriorityFiles.PopFrontValue());
+		--NumToPop;
 	}
+	AvailableFilesNum += NumToPop - InitialNumToPop;
+	check(AvailableFilesNum >= 0);
 }
 
-int32 FFilesToSearch::FTreeNode::FindLowerBoundSubDir(FStringView SubDirBaseName)
+void FFilesToSearch::PrioritizePath(FPathExistence& QueryPath, EPriority Priority)
 {
-	return Algo::LowerBound(SubDirs, SubDirBaseName,
-		[](const TUniquePtr<FTreeNode>& SubDir, FStringView BaseName)
-		{
-			return FPathViews::Less(SubDir->RelPath, BaseName);
-		}
-	);
-}
+	// We may need to prioritize a LaterRetryFile that is now loadable, so add them all into the Root
+	RetryLaterRetryFiles();
 
-void FFilesToSearch::FTreeNode::AddFiles(TArray<FGatheredPathData>&& FilePaths)
-{
-	if (Files.Num() == 0)
+	if (Priority > EPriority::Blocking)
 	{
-		Files = MoveTemp(FilePaths);
+		// TODO: Implement another tree that is searched first for the High Priority 
+		// We cannot add the High Priority files to the BlockingFiles array, because
+		// then blocking on BlockingFiles to be empty could be slow. We cannot add them
+		// as a separate simple array, because we would have to search that (sometimes large)
+		// array linearly when looking for files to accomodate a blocking priority request
+		return;
 	}
-	else
+
+	FString PriorityDirectory = QueryPath.GetLocalAbsPath();
+	const bool bIsDirectory = QueryPath.GetType() == UE::AssetDataGather::Private::FPathExistence::Directory;
+	if (!bIsDirectory)
 	{
-		Files.Append(MoveTemp(FilePaths));
+		FStringView DirectoryPath = FPathViews::GetPath(PriorityDirectory);
+		PriorityDirectory.LeftInline(DirectoryPath.Len());
 	}
+	PriorityDirectories.Add(PriorityDirectory);
 }
 
-void FFilesToSearch::FTreeNode::AddFile(FGatheredPathData&& FilePath)
+void FFilesToSearch::PopPriorityDirectories(TArray<FString>& OutPriorityDirectories)
 {
-	Files.Add(MoveTemp(FilePath));
+	while (!PriorityDirectories.IsEmpty())
+	{
+		OutPriorityDirectories.Add(PriorityDirectories.PopFrontValue());
+	}
 }
 
-template <typename RangeType>
-void FFilesToSearch::FTreeNode::PopFiles(RangeType& Out, int32& NumToPop)
+int32 FFilesToSearch::NumBlockingFiles() const
 {
-	while (NumToPop > 0 && !Files.IsEmpty())
-	{
-		Out.Add(Files.Pop(EAllowShrinking::No));
-		--NumToPop;
-	}
-	while (NumToPop > 0 && SubDirs.Num() != 0)
-	{
-		TUniquePtr<FTreeNode>& SubDir = SubDirs[SubDirs.Num() - 1];
-		SubDir->PopFiles(Out, NumToPop);
-		if (SubDir->IsEmpty())
-		{
-			SubDirs.RemoveAt(SubDirs.Num() - 1);
-		}
-	}
+	return BlockingFiles.Num();
 }
 
-template <typename RangeType>
-void FFilesToSearch::FTreeNode::PopAllFiles(RangeType& Out)
+void FFilesToSearch::Shrink()
 {
-	while (!Files.IsEmpty())
-	{
-		Out.Add(Files.Pop(EAllowShrinking::No));
-	}
-	for (int32 Index = SubDirs.Num() - 1; Index >= 0; --Index) // Match the order of PopFiles
+	// TODO: Make TRingBuffer::Shrink
+	TRingBuffer<UE::AssetDataGather::Private::FGatheredPathData> Buffer;
+
+	Buffer.Empty(NormalPriorityFiles.Num());
+	for (UE::AssetDataGather::Private::FGatheredPathData& File : NormalPriorityFiles)
 	{
-		TUniquePtr<FTreeNode>& SubDir = SubDirs[Index];
-		SubDir->PopAllFiles(Out);
-		SubDir.Reset();
+		Buffer.Add(MoveTemp(File));
 	}
-	SubDirs.Empty();
-}
+	Swap(Buffer, NormalPriorityFiles);
 
-template <typename RangeType>
-void FFilesToSearch::FTreeNode::PopMatchingDirectFiles(RangeType& Out, FStringView FileAbsPathExtOptional)
-{
-	// TODO: Make this more performant by sorting the list of Files.
-	// To prevent shifting costs, when popping the file we will leave a placeholder behind with an ignore flag set.
-	FStringView FileAbsPathNoExt = FPathViews::GetBaseFilenameWithPath(FileAbsPathExtOptional);
-	int32 NumFiles = Files.Num();
-	for (int32 Index = 0; Index < NumFiles; )
+	Buffer.Empty(BlockingFiles.Num());
+	for (UE::AssetDataGather::Private::FGatheredPathData& File : BlockingFiles)
 	{
-		FGatheredPathData& PathData = Files[Index];
-		if (FPathViews::Equals(FPathViews::GetBaseFilenameWithPath(PathData.LocalAbsPath), FileAbsPathNoExt))
-		{
-			Out.Add(MoveTemp(PathData));
-			Files.RemoveAt(Index);
-			--NumFiles;
-		}
-		else
-		{
-			++Index;
-		}
+		Buffer.Add(MoveTemp(File));
 	}
-}
+	Swap(Buffer, BlockingFiles);
 
-void FFilesToSearch::FTreeNode::PruneEmptyChild(FStringView SubDirBaseName)
-{
-	int32 Index = FindLowerBoundSubDir(SubDirBaseName);
-	if (!(Index == SubDirs.Num() || !FPathViews::Equals(SubDirs[Index]->GetRelPath(), SubDirBaseName)))
+	Buffer.Empty(LaterRetryFiles.Num());
+	for (UE::AssetDataGather::Private::FGatheredPathData& File : LaterRetryFiles)
 	{
-		if (SubDirs[Index]->IsEmpty())
-		{
-			SubDirs.RemoveAt(Index);
-		}
+		Buffer.Add(MoveTemp(File));
 	}
+	Swap(Buffer, LaterRetryFiles);
 }
 
-bool FFilesToSearch::FTreeNode::IsEmpty() const
+int32 FFilesToSearch::Num() const
 {
-	return Files.Num() == 0 && SubDirs.Num() == 0;
+	return AvailableFilesNum + LaterRetryFiles.Num();
 }
 
-void FFilesToSearch::FTreeNode::Shrink()
+int32 FFilesToSearch::GetNumAvailable() const
 {
-	Files.Shrink();
-	SubDirs.Shrink();
-	for (TUniquePtr<FTreeNode>& SubDir : SubDirs)
-	{
-		SubDir->Shrink();
-	}
+	return AvailableFilesNum;
 }
 
-SIZE_T FFilesToSearch::FTreeNode::GetAllocatedSize() const
+SIZE_T FFilesToSearch::GetAllocatedSize() const
 {
 	SIZE_T Size = 0;
-	Size = Files.GetAllocatedSize();
-	for (const FGatheredPathData& File : Files)
+	Size += NormalPriorityFiles.GetAllocatedSize();
+	for (const FGatheredPathData& PathData : NormalPriorityFiles)
 	{
-		Size += File.GetAllocatedSize();
+		Size += PathData.GetAllocatedSize();
 	}
-	Size += SubDirs.GetAllocatedSize() + SubDirs.Num()*sizeof(FTreeNode);
-	for (const TUniquePtr<FTreeNode>& SubDir : SubDirs)
+	Size += BlockingFiles.GetAllocatedSize();
+	for (const FGatheredPathData& PathData : BlockingFiles)
 	{
-		Size += SubDir->GetAllocatedSize();
+		Size += PathData.GetAllocatedSize();
 	}
-	return Size;
-}
-
-int32 FFilesToSearch::FTreeNode::NumFiles() const
-{
-	int32 Num = Files.Num();
-	for (const TUniquePtr<FTreeNode>& SubDir : SubDirs)
+	Size += LaterRetryFiles.GetAllocatedSize();
+	for (const FGatheredPathData& PathData : LaterRetryFiles)
 	{
-		Num += SubDir->NumFiles();
+		Size += PathData.GetAllocatedSize();
 	}
-	return Num;
+	return Size;
 }
 
 bool IsVerseFile(const EGatherableFileType FileType)
diff --git a/Engine/Source/Runtime/AssetRegistry/Private/AssetDataGatherer.h b/Engine/Source/Runtime/AssetRegistry/Private/AssetDataGatherer.h
--- a/Engine/Source/Runtime/AssetRegistry/Private/AssetDataGatherer.h
+++ b/Engine/Source/Runtime/AssetRegistry/Private/AssetDataGatherer.h
@@ -2,9 +2,13 @@
 
 #pragma once
 
+#include "AssetRegistry/PackageReader.h"
 #include "Containers/Array.h"
 #include "Containers/ArrayView.h"
+#include "Containers/DirectoryTree.h"
+#include "Containers/LockFreeList.h"
 #include "Containers/Map.h"
+#include "Containers/MpscQueue.h"
 #include "Containers/RingBuffer.h"
 #include "Containers/StringFwd.h"
 #include "Containers/StringView.h"
@@ -17,31 +21,69 @@
 #include "HAL/Runnable.h"
 #include "Misc/DateTime.h"
 #include "Misc/Optional.h"
+#include "NonBufferingReadOnlyArchive.h"
 #include "PackageDependencyData.h"
-#include "AssetRegistry/PackageReader.h"
+#include "Tasks/TaskConcurrencyLimiter.h"
 #include "Templates/UniquePtr.h"
 #include "UObject/NameTypes.h"
 
 #include <atomic>
 
 class FArchive;
-struct FAssetData;
-class FDiskCachedAssetData;
 class FAssetRegistryReader;
 class FAssetRegistryWriter; // Not defined if !ALLOW_NAME_BATCH_SAVING
-namespace UE::AssetRegistry { class FAssetRegistryImpl; }
+class FDiskCachedAssetData;
+struct FAssetData;
+namespace UE::AssetDataGather::Private { class FAssetDataDiscovery; }
+namespace UE::AssetDataGather::Private { class FFileReadScheduler; }
+namespace UE::AssetDataGather::Private { class FFilesToSearch; }
+namespace UE::AssetDataGather::Private { struct FCachePayload; }
+namespace UE::AssetDataGather::Private { struct FDirectoryReadTaskData; }
+namespace UE::AssetDataGather::Private { struct FGatheredPathData; }
+namespace UE::AssetDataGather::Private { struct FPathExistence; }
+namespace UE::AssetDataGather::Private { struct FReadContext; }
+namespace UE::AssetDataGather::Private { struct FSetPathProperties; }
+namespace UE::AssetDataGather::Private { enum class EPriority : uint8; }
+namespace UE::AssetDataGather::Private 
+{ 
+struct FWaitBatchDirectory
+{
+	FWaitBatchDirectory(const FStringView InPath, const bool InIsRecursive)
+		: Path(InPath)
+		, bIsRecursive(InIsRecursive)
+	{
+	}
+
+	friend int32 GetTypeHash(const FWaitBatchDirectory& WaitBatchDirectory)
+	{
+		return GetTypeHash(WaitBatchDirectory.Path);
+	}
+
+	FString Path;
+	bool bIsRecursive;
+};
 
-namespace UE::AssetDataGather::Private
+struct FWaitBatchDirectorySetFuncs : BaseKeyFuncs<FWaitBatchDirectory, FString>
 {
-class FAssetDataDiscovery;
-class FFilesToSearch;
-struct FCachePayload;
-struct FGatheredPathData;
-struct FPathExistence;
-struct FSetPathProperties;
-enum class EPriority : uint8;
+	static FORCEINLINE const FString& GetSetKey(const FWaitBatchDirectory& Element)
+	{
+		return Element.Path;
+	}
+	static FORCEINLINE bool Matches(const FString& A, const FString& B)
+	{
+		return A.Equals(B, ESearchCase::CaseSensitive);
+	}
+	static FORCEINLINE uint32 GetKeyHash(const FString& Key)
+	{
+		return GetTypeHash(Key);
+	}
+};
+
+typedef TSet<FWaitBatchDirectory, FWaitBatchDirectorySetFuncs> FWaitBatchDirectorySet;
 }
 
+namespace UE::AssetRegistry { class FAssetRegistryImpl; }
+namespace UE::AssetRegistry::Impl { enum class EGatherStatus : uint8; }
 #if DO_CHECK
 typedef FCriticalSectionQueryable FGathererCriticalSection;
 typedef FScopeLockQueryable FGathererScopeLock;
@@ -295,6 +337,15 @@ private:
 	 * thread or inline from synchronous functions on other threads.
 	 */
 	ETickResult TickInternal(double& TickStartTime, bool bPollDiscovery);
+	/**
+	* Conditionally Ticks the AssetRegistry from the background thread returning the result of gathering. Return the new status of the
+	* AssetRegistry's gather. If the background thread should not currently perfom ticking, EGatherStatus::Complete is returned.
+	*/
+	UE::AssetRegistry::Impl::EGatherStatus TryTickOnBackgroundThread();
+
+	/** Called outside the TickLock or ResultLock; it takes the ResultLock internally. */
+	bool ShouldBackgroundGatherThreadPauseGathering(bool& bOutGathererIsIdle);
+
 	/** Add any new package files from the background directory scan to our work list **/
 	void IngestDiscoveryResults();
 
@@ -410,7 +461,6 @@ private:
 	 * If null, results will only be added when Wait functions are called. Constant during threading.
 	 */
 	FRunnableThread* Thread;
-	int32 TickInternalBatchSize = 1;
 
 	/**
 	 * True if async gathering is enabled, false if e.g. singlethreaded or disabled by commandline.
@@ -488,12 +538,16 @@ private:
 	int32 NumCachedAssetFiles = 0;
 	/** The total number of files in the search results that were not in the cache and were read by parsing the file. */
 	int32 NumUncachedAssetFiles = 0;
+	/** The total number of files in the search results that were not in the cache and are currently being read off disk. */
+	int32 NumUncachedAssetFilesOutstanding = 0;
 	/** Track whether the cache has been loaded. */
 	bool bHasLoadedCache;
 	/** Track whether the Discovery subsystem has gone idle and we have read all filenames from it. */
 	bool bDiscoveryIsComplete;
 	/** Track whether this Gather has gone idle and a caller has read all search data from it. */
 	bool bIsComplete;
+	/** Track whether the Gatherers background thread should interrupt reading from disk to tick the AssetRegistry. */
+	bool bRequestAssetRegistryTick;
 	/** Track whether this Gatherer has gone idle, either it has no more work or it's blocked on external events. */
 	bool bIsIdle;
 	/** Track the first tick after idle to set up e.g. timing data. */
@@ -521,8 +575,6 @@ private:
 	TMap<FName, FDiskCachedAssetData*> DiskCachedAssetDataMap;
 	/** Map of PackageName to cached discovered assets that will be written to disk at shutdown. */
 	TMap<FName, FDiskCachedAssetData*> NewCachedAssetDataMap;
-	/** Used to block on gather results. If non-negative, tick should end when WaitBatchCount files have been processed. */
-	int32 WaitBatchCount;
 	/** How many uncached asset files had been discovered at the last async cache save */
 	int32 LastCacheSaveNumUncachedAssetFiles;
 	/**
@@ -539,6 +591,10 @@ private:
 	bool bIsSavingAsyncCache;
 	/** Packages can be marked for retry up until bInitialPluginsLoaded is set. After it is set, we retry them once. */
 	bool bFlushedRetryFiles;
+	/** Directories to explicitly wait for all files under it to be gatherered the next time we tick. */
+	TOptional<UE::AssetDataGather::Private::FWaitBatchDirectorySet> WaitBatchDirectories;
+	/** Used to manage the async task for reading gathered files. */
+	TUniquePtr<UE::AssetDataGather::Private::FFileReadScheduler> FileReadScheduler;
 };
 
 
diff --git a/Engine/Source/Runtime/AssetRegistry/Private/AssetDataGathererPrivate.h b/Engine/Source/Runtime/AssetRegistry/Private/AssetDataGathererPrivate.h
--- a/Engine/Source/Runtime/AssetRegistry/Private/AssetDataGathererPrivate.h
+++ b/Engine/Source/Runtime/AssetRegistry/Private/AssetDataGathererPrivate.h
@@ -6,6 +6,7 @@
 
 #include "AssetDataGathererDiscoveryCache.h"
 #include "Containers/Set.h"
+#include "Containers/StringFwd.h"
 #include "Misc/Optional.h"
 #include "Misc/StringBuilder.h"
 #include "Templates/RefCounting.h"
@@ -132,14 +133,155 @@ struct FGatheredPathData
 	SIZE_T GetAllocatedSize() const;
 };
 
+/** Holds scheduling context reading assets during gathering defined in a ReadContext*/
+struct FDirectoryReadTaskData
+{
+	FDirectoryReadTaskData()
+		: UnblockRead(TEXT("UnblockingDirectoryReadTaskEvent"))
+		, LastScheduledReadTask()
+	{
+	}
+
+	~FDirectoryReadTaskData()
+	{
+		// We can verify that we properly signaled our task, but we cannot verify that the ReadTask is complete,
+		// because it is possible that we insert the last outstanding read from the task and then process and remove
+		// our DirectoryReadTaskData before the task can complete. However we can verify there is no outstanding work.
+		check(UnblockRead.IsCompleted());
+		check(OutstandingReads == 0);
+	}
+
+	UE::Tasks::FTaskEvent UnblockRead;
+	UE::Tasks::FTask LastScheduledReadTask;
+	int32 OutstandingReads = 0;
+	/** Indicates a read for this directory was cancelled, implying that the directory will be rescheduled in a later gather. */
+	bool bHasCanceledRead = false;
+};
+
+/** Holds data gathered from a package file on disk that is ingested into the cache and the AssetRegistry. */
+struct FReadContext
+{
+	FName PackageName;
+	FName Extension;
+	FGatheredPathData AssetFileData;
+	TArray<FAssetData*> AssetDataFromFile;
+	FPackageDependencyData DependencyData;
+	TArray<FString> CookedPackageNamesWithoutAssetData;
+	TSharedPtr<FDirectoryReadTaskData> DirectoryReadTaskData;
+	UE::AssetRegistry::Private::EOptionalBool HasMarkOfTheWeb = UE::AssetRegistry::Private::EOptionalBool::Unset;
+	bool bCanAttemptAssetRetry = false;
+	bool bResult = false;
+	bool bCanceled = false;
+
+	FReadContext(FName InPackageName, FName InExtension, FGatheredPathData&& InAssetFileData,
+		TSharedPtr<FDirectoryReadTaskData> InDirectoryReadTaskData)
+		: PackageName(InPackageName)
+		, Extension(InExtension)
+		, AssetFileData(MoveTemp(InAssetFileData))
+		, DirectoryReadTaskData(MoveTemp(InDirectoryReadTaskData))
+	{
+		check(!AssetFileData.LocalAbsPath.IsEmpty());
+		check(DirectoryReadTaskData.IsValid());
+		check(DirectoryReadTaskData->OutstandingReads >= 0);
+		DirectoryReadTaskData->OutstandingReads++;
+	}
+
+	~FReadContext()
+	{
+		check(DirectoryReadTaskData.IsValid());
+		check(DirectoryReadTaskData->OutstandingReads > 0);
+		DirectoryReadTaskData->OutstandingReads--;
+	}
+};
+
+/** Manages the scheduling of FReadContexts */
+class FFileReadScheduler
+{
+public:
+	FFileReadScheduler(int32 InNumReadThreads, bool InGatherAssetPackageData, bool InGatherDependsData);
+	~FFileReadScheduler();
+
+	/** Returns true if AssetPackageData should be gathered. Constant during threading. */
+	bool IsAssetPackageDataGatherEnabled() const;
+	/** Returns true if dependency data should be gathered. Constant during threading. */
+	bool IsDependsDataGatherEnabled() const;
+	/** Add file to be scheduled for reading asynchronously. */
+	void AddFile(const FName PackageName, FGatheredPathData&& AssetData);
+	/** 
+	* Schedule immediately and waits for all reads for the directories listed in WaitBatchDirectories. The 
+	* list will be updated with all child directories waited on if the passed in directory was listed as recursive.
+	*/
+	void WaitOnDirectoriesAndUpdateWaitList(FWaitBatchDirectorySet& WaitBatchDirectories);
+	/** 
+	* Schedules all added files since the last call to Schedule to be picked up by the task system such that up to NumReadTaskThreads 
+	* normal priority tasks may be in-flight at once. Existing or newly added directories in the given list of PriorityDirectories 
+	* will be scheduled immediately and not bound by the schedulers NumReadTaskThreads limit.
+	*/
+	void Schedule(const TArray<FString>& PriorityDirectories);
+	/** 
+	* Populates the provided array with all completed reads (completed means read context is no longer
+	* being processed. One must still check for success/failure/cancellation).
+	*/
+	void CollectCompletedReads(TArray<TUniquePtr<FReadContext>>& OutReadyContexts);
+	/** 
+	* Removes a scheduled directory from the scheduler. When all added files for a directory have been scheduled and completed, the 
+	* directory should be removed to avoid overhead should a wait batch directory attempt to wait on the empty, still scheduled, directory.
+	*/
+	void RemoveScheduledDirectory(const FString& Directory);
+	/**
+	* Reads and stores gather results for asset files in the provided array. If time slicing is used
+	* this method may exit before all read contexts are processed, and if so, will mark unprocessed
+	* read contexts as cancelled so they may be rescheduled in another gatherer tick.
+	*/
+	void ProcessReadContexts(TArray<TUniquePtr<UE::AssetDataGather::Private::FReadContext>>&& ReadContexts, bool bIsTimeSliced);
+	/**
+	 * Releases all allocated read task loaders created during gathering.
+	 * Must be called while gathering is not in progress.
+	 */
+	void FreeReadTaskLoaders();
+	/** Frees unnecessary allocations the scheduler may have accrued. */
+	void Shrink();
+private:
+	/** 
+	* Schedules the specified directory and read contexts for reading inside of the provided Task Limiter. 
+	* If bUnblockImmediately is true, the limiter is not used.
+	*/
+	void ScheduleDirectory(UE::Tasks::FTaskConcurrencyLimiter& TaskLimiter, const FString& DirectoryPath, 
+		TArray<TUniquePtr<FReadContext>>&& ReadContexts, bool bUnblockImmediately);
+
+	/** Map of directories to the files within that directory (non-recursive) needing scheduling. */
+	TMap<FString, TArray<TUniquePtr<FReadContext>>> DirectoryReadsToSchedule;
+	/** List of FReadContexts that have been populated (or cancelled) from worker threads requiring processing. */
+	TMpscQueue<TUniquePtr<UE::AssetDataGather::Private::FReadContext>> ReadyReadContexts;
+	/** Freelist of loaders created on task threads during gathering.*/
+	TLockFreePointerListFIFO<FNonBufferingReadOnlyArchive, 0> ReadTaskLoaderFreeList;
+	/** Used to control how parallel background asset file reading can be. */
+	UE::Tasks::FTaskConcurrencyLimiter ReadConcurrencyLimiter;
+	/**
+	* Tree of directories being gathered asynchronously. The FDirectoryReadTaskData is used to trigger and track the
+	* FTask performing file reads for that directory. We scheduled file reads by the directory they reside in to avoid
+	* contention caused by the OS filesystem.
+	*/
+	TUniquePtr<TDirectoryTree<TSharedPtr<UE::AssetDataGather::Private::FDirectoryReadTaskData>>> ScheduledDirectories;
+	/**
+	* Number of parallel read tasks allowed while gathering (wait batches are allowed to run in parallel to these tasks
+	* allowing for > NumReadTaskThreads to possibly be in flight at once.)
+	*/
+	int32 NumReadTaskThreads;
+	/** True if AssetPackageData should be gathered. Constant during threading. */
+	bool bGatherAssetPackageData;
+	/** True if dependency data should be gathered. Constant during threading. */
+	bool bGatherDependsData;
+};
+
 /** A container to efficiently receive files from discovery and search the container to raise priority of a path */
 class FFilesToSearch
 {
 public:
 	/** Add the given FilePath as a single file with blocking priority. */
 	void AddPriorityFile(FGatheredPathData&& FilePath);
-	/** Add the given directory and its discovered direct Files. */
-	void AddDirectory(FString&& DirAbsPath, TArray<FGatheredPathData>&& FilePaths);
+	/** Add the given FilePaths with normal priority. */
+	void AddFiles(TArray<FGatheredPathData>&& FilePaths);
 	/** Re-add a file from PopFront that we didn't have time to process. */
 	void AddFileAgainAfterTimeout(FGatheredPathData&& FilePath);
 	/**
@@ -159,6 +301,9 @@ public:
 	 * Path might be a directory name or a file BaseNameWithPath, or a file BaseNameWithPath.Extension.
 	*/
 	void PrioritizePath(FPathExistence& QueryPath, EPriority Priority);
+
+	void PopPriorityDirectories(TArray<FString>& OutPriorityDirectories);
+
 	/** How many files with blocking priority are in the container. */
 	int32 NumBlockingFiles() const;
 
@@ -171,69 +316,8 @@ public:
 	/** How much memory is used by the container, not counting sizeof(*this). */
 	SIZE_T GetAllocatedSize() const;
 private:
-	/** A directory-search-tree structure; each node has a list of direct files and subdirectories. */
-	struct FTreeNode
-	{
-		FTreeNode() = default;
-		explicit FTreeNode(FStringView InRelPath);
-
-		/** The name of the directory relative to its parent. */
-		FStringView GetRelPath() const;
-
-		/**
-		 * Recursively search the tree to find the given relative directory name,
-		 * adding nodes for the directory and its parents if required.
-		 */
-		FTreeNode& FindOrAddNode(FStringView InRelPath);
-		/**
-		 * Recursively search the tree to find the given relative directory name;
-		 * return nullptr if the node or any parent is not found.
-		 */
-		FTreeNode* FindNode(FStringView InRelPath);
-		/**
-		 * Recursively search the tree to find the closest parent of the given relative directory name;
-		 * returns *this if the first component of the relativepath is not in any child node.
-		 */
-		FTreeNode& FindClosestParentNode(FStringView InRelPath, FStringView& OutRemainingRelPath);
-
-		/** Add the given FilePaths as direct file children of the Node's directory. */
-		void AddFiles(TArray<FGatheredPathData>&& FilePaths);
-		/** Add the given FilePath as a direct file child of the Node's directory. */
-		void AddFile(FGatheredPathData&& FilePath);
-
-		/** Pop the NumToPop files out of this node and its children and decrement NumToPop by how many were popped. */
-		template <typename RangeType>
-		void PopFiles(RangeType& Out, int32& NumToPop);
-		/** Pop all files out of this node and its children. */
-		template <typename RangeType>
-		void PopAllFiles(RangeType& Out);
-		/** Pop files out of *this but not children, if they match the given full path or BaseNameWithPath. */
-		template <typename RangeType>
-		void PopMatchingDirectFiles(RangeType& Out, FStringView FileAbsPathExtOptional);
-
-		/** Delete the given child node if it is empty. */
-		void PruneEmptyChild(FStringView SubDirBaseName);
-		/** Return true if this node has no files or subdirs. */
-		bool IsEmpty() const;
-		/** Reduce memory used in buffers. */
-		void Shrink();
-		/** How much memory is used by *this, not counting sizeof(*this). */
-		SIZE_T GetAllocatedSize() const;
-		/** Number of files in *this and its child nodes. */
-		int32 NumFiles() const;
-	private:
-		/** Search direct SubDirs for the given BaseName and add a node for it if not existing. */
-		FTreeNode& FindOrAddSubDir(FStringView SubDirBaseName);
-		/** Search direct SubDirs for the given BaseName and return nullptr if not existing. */
-		FTreeNode* FindSubDir(FStringView SubDirBaseName);
-		/** Find the index of the subdir with the given Relative path. */
-		int32 FindLowerBoundSubDir(FStringView SubDirBaseName);
-
-		FString RelPath;
-		TArray<FGatheredPathData> Files; // Not sorted
-		TArray<TUniquePtr<FTreeNode>> SubDirs; // Sorted
-	};
-	FTreeNode Root;
+	TRingBuffer<FString> PriorityDirectories;
+	TRingBuffer<UE::AssetDataGather::Private::FGatheredPathData> NormalPriorityFiles;
 	TRingBuffer<UE::AssetDataGather::Private::FGatheredPathData> BlockingFiles;
 	TRingBuffer<UE::AssetDataGather::Private::FGatheredPathData> LaterRetryFiles;
 	int32 AvailableFilesNum = 0;
@@ -866,9 +950,8 @@ enum class EEditorGameScanMode : uint8
 struct FGatherSettings
 {
 public:
-	static constexpr int32 SingleThreadFilesPerBatch = 3;
 	static constexpr int32 ExpectedMaxBatchSize = 100;
-	static constexpr int32 MinSecondsToElapseBeforeCacheWrite = 30;
+	static constexpr int32 MinSecondsToElapseBeforeCacheWrite = 300;
 	static constexpr int32 MinAssetReadsBeforeCacheWrite = 100000;
 	static constexpr int32 CacheShardAssetCount = 250000;
 	static constexpr uint32 CacheSerializationMagic = 0x29FE48B0; // Versioning and integrity checking
@@ -893,8 +976,9 @@ public:
 	 * The maximum number of threads used in the ParallelFor around ReadAssetFile in FAssetDataGatherer::TickInternal.
 	 * Very high numbers of threads does not reduce wall time of the gather because the threads become IO limited.
 	 * Capping the number of threads allows those threads to be used in other async operations in Engine startup.
+	 * A value of <= 0 will let the system decide how many threads to use.
 	 */
-	static inline int32 GARGatherThreads = 60;
+	static inline int32 GARGatherThreads = 0;
 
 	/**
 	 * The maximum number of the gatherer's cache files that can be loaded in parallel in LoadCacheFiles.
@@ -1019,5 +1103,12 @@ inline const FString& FGatherSettings::GetAssetRegistryCacheRootFolder() const
 {
 	return AssetRegistryCacheRootFolder;
 }
-
+inline bool FFileReadScheduler::IsAssetPackageDataGatherEnabled() const
+{
+	return bGatherAssetPackageData;
+}
+inline bool FFileReadScheduler::IsDependsDataGatherEnabled() const
+{
+	bGatherDependsData;
+}
 } // namespace UE::AssetDataGather::Private
\ No newline at end of file
diff --git a/Engine/Source/Runtime/AssetRegistry/Private/NonBufferingReadOnlyArchive.cpp b/Engine/Source/Runtime/AssetRegistry/Private/NonBufferingReadOnlyArchive.cpp
new file mode 100644
--- /dev/null
+++ b/Engine/Source/Runtime/AssetRegistry/Private/NonBufferingReadOnlyArchive.cpp
@@ -0,0 +1,212 @@
+// Copyright Epic Games, Inc. All Rights Reserved.
+
+#include "NonBufferingReadOnlyArchive.h"
+
+#include "AssetRegistryPrivate.h"
+#include "HAL/FileManager.h"
+#include "HAL/PlatformFileManager.h"
+#include "Logging/MessageLog.h"
+#include "Logging/StructuredLog.h"
+#include "Misc/PathViews.h"
+
+FNonBufferingReadOnlyArchive::FNonBufferingReadOnlyArchive(int64 InBufferSize, uint32 InFlags)
+	: Filename()
+	, Size(0)
+	, Pos(0)
+	, BufferBase(0)
+	, Handle()
+	, Flags(InFlags)
+{
+	// Todo: The buffer must be aligned to at least the size of the physical sector of the disk used.
+	// We don't have a generic API for this so use the page size as that will cover most common cases
+	BufferReadAlignment = FPlatformMemory::GetConstants().BinnedPageSize;
+	check(FMath::IsPowerOfTwo(BufferReadAlignment));
+	BufferSize = FMath::Max(InBufferSize, BufferReadAlignment);
+	check(FMath::IsPowerOfTwo(BufferSize));
+	Buffer = (uint8*) FPlatformMemory::BinnedAllocFromOS(BufferSize);
+	check(Buffer);
+}
+
+FNonBufferingReadOnlyArchive::~FNonBufferingReadOnlyArchive()
+{
+	Close();
+	FPlatformMemory::BinnedFreeToOS(Buffer, BufferSize);
+}
+
+void FNonBufferingReadOnlyArchive::Seek(int64 InPos)
+{
+	checkf(InPos >= 0, TEXT("Attempted to seek to a negative location (%lld/%lld), file: %s. The file is most likely corrupt."), InPos, Size, *Filename);
+	checkf(InPos <= Size, TEXT("Attempted to seek past the end of file (%lld/%lld), file: %s. The file is most likely corrupt."), InPos, Size, *Filename);
+
+	// Note, we don't do any platform specific seeking when asked. We instead defer until a read attempt is made where we will
+	// be forced to read in a new buffer's amount of data.
+	Pos = InPos;
+}
+
+void FNonBufferingReadOnlyArchive::ReadLowLevel(uint8* Dest, int64 CountToRead, int64& OutBytesRead)
+{
+	int64 StartPos = Handle->Tell();
+	if (Handle->Read(Dest, CountToRead))
+	{
+		OutBytesRead = Handle->Tell() - StartPos;
+	}
+	else
+	{
+		OutBytesRead = 0;
+	}
+}
+
+bool FNonBufferingReadOnlyArchive::SeekLowLevel(int64 InPos)
+{
+	return Handle->Seek(InPos);
+}
+
+void FNonBufferingReadOnlyArchive::OpenFileLowLevel(const TCHAR* InFilename)
+{
+	Close();
+
+	FFileOpenResult FileResult = FPlatformFileManager::Get().GetPlatformFile().OpenReadNoBuffering(InFilename, IPlatformFile::EOpenReadFlags::NoBuffering);
+	if (!FileResult.HasError())
+	{
+		Handle = FileResult.StealValue();
+	}
+	else
+	{
+		UE_CLOG(!IsSilent(), LogAssetRegistry, Warning, TEXT("OpenFile failed: File:%s, Error:%s"),
+			InFilename, *FileResult.StealError().GetMessage());
+	}
+}
+
+void FNonBufferingReadOnlyArchive::CloseLowLevel()
+{
+	Handle.Reset();
+}
+
+bool FNonBufferingReadOnlyArchive::Close()
+{
+	CloseLowLevel();
+	return !IsError();
+}
+
+bool FNonBufferingReadOnlyArchive::Precache(int64 PrecacheOffset, int64 PrecacheSize)
+{
+	// Archives not based on async I/O should always return true, so we return true whether or not the precache was successful.
+	// Returning false would imply that the caller should continue calling until it returns true.
+	if (PrecacheSize < 0)
+	{
+		return true;
+	}
+	InternalPrecache(PrecacheOffset, PrecacheSize);
+	return true;
+}
+
+bool FNonBufferingReadOnlyArchive::OpenFile(const TCHAR* InFilename)
+{
+	OpenFileLowLevel(InFilename);
+	if (!Handle.IsValid())
+	{
+		return false;
+	}
+	Size = Handle->Size();
+	Filename = InFilename;
+	Pos = 0;
+
+	// Until we read in data we need a value that will force us to precache
+	BufferBase = -(BufferSize + 1);
+
+	this->SetIsLoading(true);
+	this->SetIsPersistent(true);
+	return true;
+}
+
+bool FNonBufferingReadOnlyArchive::InternalPrecache(int64 PrecacheOffset, int64 PrecacheSize)
+{
+	check(Buffer);
+
+	// Only precache at current position and avoid work if precaching same offset twice.
+	if (Pos != PrecacheOffset)
+	{
+		// We are refusing to precache, but return true if at least one byte after the requested PrecacheOffset is in our existing buffer.
+		return BufferBase <= PrecacheOffset && PrecacheOffset < BufferBase + BufferSize;
+	}
+
+	bool bPosWithinBuffer = BufferBase <= Pos && Pos < BufferBase + BufferSize;
+	int64 ReadCount = BufferSize;
+	if (bPosWithinBuffer)
+	{
+		// At least one byte after the requested PrecacheOffset is in our existing buffer so do nothing and return true
+		return true;
+	}
+
+	// We must always read at a sector aligned address so move the position backwards as needed
+	BufferBase = AlignDown(Pos, BufferReadAlignment);
+	if (BufferBase + ReadCount < Pos)
+	{
+		UE_CLOG(!IsSilent(), LogAssetRegistry, Warning, 
+			TEXT("Buffer is not large enough to satisfy the requested read. BufferBase=%lld BufferSize=%lld BufferReadAlignment=%lld ReadCount=%lld Pos=%lld"), BufferBase, BufferSize, BufferReadAlignment, ReadCount, Pos);
+		return false;
+	}
+
+	int64 Count = 0;
+	SeekLowLevel(BufferBase);
+	ReadLowLevel(Buffer, ReadCount, Count);
+
+	if (Count != ReadCount && Count != (Size - BufferBase))
+	{
+		uint32 LastError = FPlatformMisc::GetLastError();
+		TCHAR ErrorBuffer[1024];
+		UE_CLOG(!IsSilent(), LogAssetRegistry, Warning, TEXT("ReadFile failed: Count=%lld Pos=%lld Size=%lld BufferBase=%lld ReadCount=%lld BufferReadAlignment=%lld LastError=%u: %s"),
+			Count, Pos, Size, BufferBase, ReadCount, BufferReadAlignment, LastError, FPlatformMisc::GetSystemErrorMessage(ErrorBuffer, 1024, LastError));
+
+		return BufferBase + Count >= Pos;
+	}
+
+	return true;
+}
+
+void FNonBufferingReadOnlyArchive::Serialize(void* V, int64 Length)
+{
+	if (Pos + Length > Size)
+	{
+		SetError();
+		UE_CLOG(!IsSilent(), LogAssetRegistry, Error, TEXT("Requested read of %" INT64_FMT " bytes when %" INT64_FMT " bytes remain (file=%s, size=%" INT64_FMT ")"), Length, Size - Pos, *Filename, Size);
+		return;
+	}
+
+	bool bReadUncached = false;
+	const bool bIsOutsideBufferWindow = (Pos < BufferBase) || (Pos >= (BufferBase + BufferSize));
+	if (bIsOutsideBufferWindow)
+	{
+		bReadUncached = true;
+	}
+
+	while (Length > 0)
+	{
+		int64 Copy = FMath::Min(Length, BufferBase + BufferSize - Pos);
+		if (Copy <= 0 || bReadUncached)
+		{
+			if (!InternalPrecache(Pos, MAX_int32))
+			{
+				SetError();
+				UE_CLOG(!IsSilent(), LogAssetRegistry, Warning, TEXT("ReadFile failed during precaching for file %s"), *Filename);
+				return;
+			}
+			bReadUncached = false;
+			Copy = FMath::Min(Length, BufferBase + BufferSize - Pos);
+			if (Copy <= 0)
+			{
+				SetError();
+				UE_CLOG(!IsSilent(), LogAssetRegistry, Error, TEXT("ReadFile beyond EOF %lld+%lld/%lld for file %s"),
+					Pos, Length, Size, *Filename);
+			}
+			if (IsError())
+			{
+				return;
+			}
+		}
+		FMemory::Memcpy(V, Buffer + Pos - BufferBase, Copy);
+		Pos += Copy;
+		Length -= Copy;
+		V = (uint8*)V + Copy;
+	}
+}
diff --git a/Engine/Source/Runtime/AssetRegistry/Private/NonBufferingReadOnlyArchive.h b/Engine/Source/Runtime/AssetRegistry/Private/NonBufferingReadOnlyArchive.h
new file mode 100644
--- /dev/null
+++ b/Engine/Source/Runtime/AssetRegistry/Private/NonBufferingReadOnlyArchive.h
@@ -0,0 +1,81 @@
+// Copyright Epic Games, Inc. All Rights Reserved.
+
+#pragma once
+
+#include "Containers/StringFwd.h"
+#include "HAL/FileManager.h"
+#include "Serialization/ArchiveUObject.h"
+
+/*
+* This archive uses non-buffered IO to avoid overhead caused by reading into the OS cache manager 
+* rather than not reading into an archive owned buffer (it does that). It is intended to be 
+* used when reads will occur in parallel, contention would be caused at an OS level to allocate cache pages, 
+* and typically the whole file is not read sequentially but instead only a smaller portion of the file 
+* is needed (otherwise OS caching likely performs better since it can predict more accurately without waste).
+* This archive is intended to be re-usable by opening new files on top so the internal buffer can be re-used.
+*/
+class FNonBufferingReadOnlyArchive : public FArchive
+{
+public:
+	FNonBufferingReadOnlyArchive(int64 InBufferSize = 4096, uint32 InFlags = 0);
+	~FNonBufferingReadOnlyArchive();
+
+	// FArchive API
+	virtual void Seek(int64 InPos) override final;
+	virtual int64 Tell() override final
+	{
+		return Pos;
+	}
+	virtual int64 TotalSize() override final
+	{
+		return Size;
+	}
+	virtual bool Close() override final;
+	virtual void Serialize(void* V, int64 Length) override final;
+	virtual FString GetArchiveName() const override
+	{
+		return Filename;
+	}
+	virtual bool Precache(int64 PrecacheOffset, int64 PrecacheSize) override;
+	//End FArchive API
+
+	/** Open a new file to read with this archive. Any existing open file is closed and the archive is reset. */
+	bool OpenFile(const TCHAR* InFilename);
+
+protected:
+	/** Populates the internal buffer to read from during Serialize. */
+	bool InternalPrecache(int64 PrecacheOffset, int64 PrecacheSize);
+
+	/**
+	 * Platform specific seek
+	 * @param InPos - Offset from beginning of file to seek to
+	 * @return false on failure
+	**/
+	bool SeekLowLevel(int64 InPos);
+
+	/** Open the file handle */
+	void OpenFileLowLevel(const TCHAR* InFilename);
+
+	/** Close the file handle **/
+	void CloseLowLevel();
+
+	/** Platform specific read */
+	void ReadLowLevel(uint8* Dest, int64 CountToRead, int64& OutBytesRead);
+
+	/** Returns true if the archive should suppress logging in case of error */
+	bool IsSilent() const
+	{
+		return !!(Flags & EFileRead::FILEREAD_Silent);
+	}
+
+	FString Filename;
+	int64 Size;
+	int64 Pos;
+	int64 BufferBase;
+	TUniquePtr<IFileHandle> Handle;
+
+	uint8* Buffer;
+	int64 BufferReadAlignment;
+	int64 BufferSize;
+	uint32 Flags;
+};
\ No newline at end of file
diff --git a/Engine/Source/Runtime/AssetRegistry/Private/PackageReader.cpp b/Engine/Source/Runtime/AssetRegistry/Private/PackageReader.cpp
--- a/Engine/Source/Runtime/AssetRegistry/Private/PackageReader.cpp
+++ b/Engine/Source/Runtime/AssetRegistry/Private/PackageReader.cpp
@@ -6,6 +6,7 @@
 #include "AssetRegistryPrivate.h"
 #include "AssetRegistry/AssetData.h"
 #include "HAL/FileManager.h"
+#include "HAL/PlatformFileManager.h"
 #include "Internationalization/Internationalization.h"
 #include "Internationalization/GatherableTextData.h"
 #include "Logging/MessageLog.h"
@@ -114,12 +115,12 @@ bool FPackageReader::OpenPackageFile(FStringView InLongPackageName, FStringView
 	return OpenPackageFile(OutErrorCode ? *OutErrorCode : Tmp); 
 }
 
-bool FPackageReader::OpenPackageFile(FArchive* InLoader, EOpenPackageResult* OutErrorCode)
+bool FPackageReader::OpenPackageFile(FArchive* InLoader, FStringView InLongPackageName, EOpenPackageResult* OutErrorCode)
 {
 	check(!Loader);
 	Loader = InLoader;
 	bLoaderOwner = false;
-	LongPackageName.Empty();
+	LongPackageName = InLongPackageName;
 	PackageFilename = Loader->GetArchiveName();
 	EOpenPackageResult Tmp = EOpenPackageResult::Success;
 	return OpenPackageFile(OutErrorCode ? *OutErrorCode : Tmp); 
@@ -1895,4 +1896,4 @@ namespace UE::AssetRegistry
 	}
 }
 
-#undef LOCTEXT_NAMESPACE
\ No newline at end of file
+#undef LOCTEXT_NAMESPACE
diff --git a/Engine/Source/Runtime/Core/Private/Windows/WindowsPlatformFile.cpp b/Engine/Source/Runtime/Core/Private/Windows/WindowsPlatformFile.cpp
--- a/Engine/Source/Runtime/Core/Private/Windows/WindowsPlatformFile.cpp
+++ b/Engine/Source/Runtime/Core/Private/Windows/WindowsPlatformFile.cpp
@@ -1433,30 +1433,33 @@ IAsyncReadFileHandle* FWindowsPlatformFile::OpenAsyncRead(const TCHAR* Filename,
 
 FFileOpenResult FWindowsPlatformFile::OpenRead(const TCHAR* Filename, EOpenReadFlags Flags)
 {
-	const uint32  Access	= GENERIC_READ;
-	const uint32  WinFlags  =	FILE_SHARE_READ | GetFileShareDeleteFlagIfEnabled() |
-								(EnumHasAnyFlags(Flags, EOpenReadFlags::AllowWrite) ? FILE_SHARE_WRITE : 0) |
-								(EnumHasAnyFlags(Flags, EOpenReadFlags::AllowDelete) ? FILE_SHARE_DELETE : 0);
-	const uint32  Create	= OPEN_EXISTING;
+	const uint32 Access				= GENERIC_READ;
+	const uint32 ShareFlags			= FILE_SHARE_READ | GetFileShareDeleteFlagIfEnabled() |
+										(EnumHasAnyFlags(Flags, EOpenReadFlags::AllowWrite) ? FILE_SHARE_WRITE : 0) |
+										(EnumHasAnyFlags(Flags, EOpenReadFlags::AllowDelete) ? FILE_SHARE_DELETE : 0);
+	const uint32 FlagsAndAttributes = (!EnumHasAnyFlags(Flags, EOpenReadFlags::NoBuffering) ? FILE_ATTRIBUTE_NORMAL :
+										FILE_FLAG_NO_BUFFERING |	// Disables OS caching of file data
+										FILE_ATTRIBUTE_TEMPORARY);	// Disables OS caching of file metadata
+	const uint32 Create				= OPEN_EXISTING;
 
 #define USE_OVERLAPPED_IO (!IS_PROGRAM && !WITH_EDITOR)		// Use straightforward synchronous I/O in cooker/editor
 
 	TRACE_PLATFORMFILE_BEGIN_OPEN(Filename);
 #if USE_OVERLAPPED_IO
-	HANDLE Handle    = CreateFileW(*FNormalizedFilename(Filename), Access, WinFlags, NULL, Create, FILE_ATTRIBUTE_NORMAL | FILE_FLAG_OVERLAPPED, NULL);
+	HANDLE Handle    = CreateFileW(*FNormalizedFilename(Filename), Access, ShareFlags, NULL, Create, FlagsAndAttributes | FILE_FLAG_OVERLAPPED, NULL);
 	if (Handle != INVALID_HANDLE_VALUE)
 	{
 		TRACE_PLATFORMFILE_END_OPEN(Handle);
-		FAsyncBufferedFileReaderWindows* FileHandle = new FAsyncBufferedFileReaderWindows(Handle, Access, WinFlags, FILE_FLAG_OVERLAPPED);
+		FAsyncBufferedFileReaderWindows* FileHandle = new FAsyncBufferedFileReaderWindows(Handle, Access, ShareFlags, FlagsAndAttributes | FILE_FLAG_OVERLAPPED);
 		return MakeValue(FileHandle);
 	}
 #else
-	HANDLE Handle = CreateFileW(*FNormalizedFilename(Filename), Access, WinFlags, NULL, Create, FILE_ATTRIBUTE_NORMAL, NULL);
+	HANDLE Handle = CreateFileW(*FNormalizedFilename(Filename), Access, ShareFlags, NULL, Create, FlagsAndAttributes, NULL);
 	if (Handle != INVALID_HANDLE_VALUE)
 	{
 		TRACE_PLATFORMFILE_END_OPEN(Handle);
 
-		FFileHandleWindows* FileHandle = new FFileHandleWindows(Handle, Access, WinFlags, 0);
+		FFileHandleWindows* FileHandle = new FFileHandleWindows(Handle, Access, ShareFlags, FlagsAndAttributes);
 
 		// Some operations can fail during the handle initialization, so we
 		// double check that the handle is valid before returning it
@@ -1489,37 +1492,54 @@ IFileHandle* FWindowsPlatformFile::OpenRead(const TCHAR* Filename, bool bAllowWr
 	return Result.StealValue().Release();
 }
 
-IFileHandle* FWindowsPlatformFile::OpenReadNoBuffering(const TCHAR* Filename, bool bAllowWrite)
+FFileOpenResult FWindowsPlatformFile::OpenReadNoBuffering(const TCHAR* Filename, EOpenReadFlags Flags)
 {
-	uint32  Access = GENERIC_READ;
-	uint32  WinFlags = FILE_SHARE_READ | (bAllowWrite ? FILE_SHARE_WRITE : 0) | GetFileShareDeleteFlagIfEnabled();
-	uint32  Create = OPEN_EXISTING;
+	const uint32 Access				= GENERIC_READ;
+	const uint32 ShareFlags			= FILE_SHARE_READ | GetFileShareDeleteFlagIfEnabled() |
+										(EnumHasAnyFlags(Flags, EOpenReadFlags::AllowWrite) ? FILE_SHARE_WRITE : 0) |
+										(EnumHasAnyFlags(Flags, EOpenReadFlags::AllowDelete) ? FILE_SHARE_DELETE : 0);
+	const uint32 FlagsAndAttributes = FILE_FLAG_OVERLAPPED | 
+										(!EnumHasAnyFlags(Flags, EOpenReadFlags::NoBuffering) ? FILE_ATTRIBUTE_NORMAL :
+										FILE_FLAG_NO_BUFFERING |	// Disables OS caching of file data
+										FILE_ATTRIBUTE_TEMPORARY);	// Disables OS caching of file metadata
+	const uint32 Create				= OPEN_EXISTING;
+
 	TRACE_PLATFORMFILE_BEGIN_OPEN(Filename);
-	HANDLE Handle = CreateFileW(*FNormalizedFilename(Filename), Access, WinFlags, NULL, Create, FILE_ATTRIBUTE_NORMAL | FILE_FLAG_OVERLAPPED, NULL);
+	HANDLE Handle = CreateFileW(*FNormalizedFilename(Filename), Access, ShareFlags, NULL, Create, FlagsAndAttributes, NULL);
 	if (Handle != INVALID_HANDLE_VALUE)
 	{
 		TRACE_PLATFORMFILE_END_OPEN(Handle);
 
-		FFileHandleWindows* FileHandle = new FFileHandleWindows(Handle, Access, WinFlags, FILE_FLAG_OVERLAPPED);
+		FFileHandleWindows* FileHandle = new FFileHandleWindows(Handle, Access, ShareFlags, FlagsAndAttributes);
 
 		// Some operations can fail during the handle initialization, so we
 		// double check that the handle is valid before returning it
 		if (FileHandle->IsValid())
 		{
-			return FileHandle;
+			return MakeValue(FileHandle);
 		}
 		else
 		{
 			delete FileHandle;
-
-			return nullptr;
+			return MakeError(TEXT("Failed to create FFileHandleWindows"));
 		}
 	}
 	else
 	{
 		TRACE_PLATFORMFILE_FAIL_OPEN(Filename);
+		return MakeError(TEXT("Failed to open file for reading"), FPlatformMisc::GetLastError());
+	}
+}
+
+IFileHandle* FWindowsPlatformFile::OpenReadNoBuffering(const TCHAR* Filename, bool bAllowWrite)
+{
+	const EOpenReadFlags Flags = bAllowWrite ? EOpenReadFlags::AllowWrite : EOpenReadFlags::None;
+	FFileOpenResult Result = OpenReadNoBuffering(Filename, Flags);
+	if (Result.HasError())
+	{
 		return nullptr;
 	}
+	return Result.StealValue().Release();
 }
 
 IFileHandle* FWindowsPlatformFile::OpenWrite(const TCHAR* Filename, bool bAppend, bool bAllowRead)
diff --git a/Engine/Source/Runtime/Core/Private/Windows/WindowsPlatformFilePrivate.h b/Engine/Source/Runtime/Core/Private/Windows/WindowsPlatformFilePrivate.h
--- a/Engine/Source/Runtime/Core/Private/Windows/WindowsPlatformFilePrivate.h
+++ b/Engine/Source/Runtime/Core/Private/Windows/WindowsPlatformFilePrivate.h
@@ -66,6 +66,7 @@ public:
 
 	virtual FFileOpenResult OpenRead(const TCHAR* Filename, EOpenReadFlags Flags) override;
 	virtual IFileHandle* OpenRead(const TCHAR* Filename, bool bAllowWrite = false) override;
+	virtual FFileOpenResult OpenReadNoBuffering(const TCHAR* Filename, EOpenReadFlags Flags) override;
 	virtual IFileHandle* OpenReadNoBuffering(const TCHAR* Filename, bool bAllowWrite = false) override;
 	virtual IFileHandle* OpenWrite(const TCHAR* Filename, bool bAppend = false, bool bAllowRead = false) override;
 	virtual FOpenMappedResult OpenMappedEx(const TCHAR* Filename, EOpenReadFlags OpenOptions = EOpenReadFlags::None, int64 MaximumSize = 0) override;
diff --git a/Engine/Source/Runtime/Core/Public/GenericPlatform/GenericPlatformFile.h b/Engine/Source/Runtime/Core/Public/GenericPlatform/GenericPlatformFile.h
--- a/Engine/Source/Runtime/Core/Public/GenericPlatform/GenericPlatformFile.h
+++ b/Engine/Source/Runtime/Core/Public/GenericPlatform/GenericPlatformFile.h
@@ -495,7 +495,13 @@ public:
 		/** Allow other handles/processes to write to this file. This flag is needed to open files that are currently being written to as well */
 		AllowWrite	= 1 << 0,
 		/** Allow the file to be deleted or renamed while keeping this handle valid for reading. */
-		AllowDelete = 1 << 1
+		AllowDelete = 1 << 1,
+		/** 
+		* Disables buffering performed by the platform operating system, if any. Note, this typically requires read 
+		* offsets, and read sizes to be multiples of the volume's sector size, as well as the buffer being read into 
+		* being aligned to a multiple of the physical sector size. Avoid using if you cannot control these parameters.
+		*/
+		NoBuffering = 1 << 2
 	};
 
 	/**
@@ -519,6 +525,11 @@ public:
 	 */
 	virtual IFileHandle*	OpenRead(const TCHAR* Filename, bool bAllowWrite = false) = 0;
 
+	virtual FFileOpenResult OpenReadNoBuffering(const TCHAR* Filename, EOpenReadFlags Flags)
+	{
+		return OpenRead(Filename, Flags);
+	}
+
 	virtual IFileHandle* OpenReadNoBuffering(const TCHAR* Filename, bool bAllowWrite = false)
 	{
 		return OpenRead(Filename, bAllowWrite);
