[AssetRegistry]
Fix missing return breaking clang

#rnx

diff --git a/Engine/Source/Runtime/AssetRegistry/Private/AssetDataGathererPrivate.h b/Engine/Source/Runtime/AssetRegistry/Private/AssetDataGathererPrivate.h
--- a/Engine/Source/Runtime/AssetRegistry/Private/AssetDataGathererPrivate.h
+++ b/Engine/Source/Runtime/AssetRegistry/Private/AssetDataGathererPrivate.h
@@ -1109,6 +1109,6 @@ inline bool FFileReadScheduler::IsAssetPackageDataGatherEnabled() const
 }
 inline bool FFileReadScheduler::IsDependsDataGatherEnabled() const
 {
-	bGatherDependsData;
+	return bGatherDependsData;
 }
 } // namespace UE::AssetDataGather::Private
\ No newline at end of file
