[Backout] - CL43871991
#rnx
Original CL Desc
-----------------------------------------------------------------
[AssetRegistry] Allow overlapping KeepInDevOnly ClassFilter rules to exclude one another

diff --git a/Engine/Source/Runtime/AssetRegistry/Private/AssetRegistry.cpp b/Engine/Source/Runtime/AssetRegistry/Private/AssetRegistry.cpp
--- a/Engine/Source/Runtime/AssetRegistry/Private/AssetRegistry.cpp
+++ b/Engine/Source/Runtime/AssetRegistry/Private/AssetRegistry.cpp
@@ -2,7 +2,6 @@
 
 #include "AssetRegistry.h"
 
-#include "Algo/HeapSort.h"
 #include "Algo/Transform.h"
 #include "Algo/Unique.h"
 #include "AssetDataGatherer.h"
@@ -51,9 +50,6 @@
 #include "UObject/ConstructorHelpers.h"
 #include "UObject/CoreRedirects.h"
 #include "UObject/MetaData.h"
-#include "UObject/Object.h"
-#include "UObject/SoftObjectPath.h"
-#include "UObject/UObjectGlobals.h"
 #include "UObject/UObjectHash.h"
 #include "UObject/UObjectIterator.h"
 #include "UObject/UObjectThreadContext.h"
@@ -700,43 +696,6 @@ void FAsyncConsumer::Consume(UAssetRegistryImpl& UARI, UE::AssetRegistry::Impl::
 
 namespace UE::AssetRegistry
 {
-
-namespace DependsNode
-{
-typedef TMap<FDependsNode*, FDependsNode::FPackageFlagSet> FPackageDependencyMap;
-
-/** Uncompressed version of DependsNode data, used to avoid allocations and
-compression cost when constructing batches of DependsNodes in DependencyDataGathered. */
-struct FConstructData
-{
-	FDependsNode* SourceNode;
-	FPackageDependencyMap PackageDependencies;
-	FDependsNode::FDependsNodeList SearchableNamesNodes;
-};
-
-struct FGatherDependencyDataScratchpad
-{
-	TMap<FName, FConstructData> AssetNameToConstructDataMap;
-	TMap<FName, FName> CachedDepToRedirect;
-	TSet<FDependsNode*> ModifiedDependNodes;
-
-	void Reset()
-	{
-		AssetNameToConstructDataMap.Reset();
-		CachedDepToRedirect.Reset();
-		ModifiedDependNodes.Reset();
-	}
-
-	void Shrink()
-	{
-		AssetNameToConstructDataMap.Shrink();
-		CachedDepToRedirect.Shrink();
-		ModifiedDependNodes.Shrink();
-	}
-};
-
-}
-
 void FAssetRegistryImpl::ConditionalLoadPremadeAssetRegistry(UAssetRegistryImpl& UARI, Impl::FEventContext& EventContext, UE::AssetRegistry::FInterfaceWriteScopeLock& ScopeLock)
 {
 	AsyncConsumer.Wait(UARI, ScopeLock);
@@ -1057,16 +1016,12 @@ void FAssetRegistryImpl::Initialize(Impl::FInitializeContext& Context)
 	// Collect all code generator classes (currently BlueprintCore-derived ones)
 	CollectCodeGeneratorClasses();
 #if WITH_ENGINE && WITH_EDITOR
-	{
-		UE::TWriteScopeLock WriteLock(SkipClassesLock);
-		Utils::PopulateSkipClasses(SkipUncookedClasses, SkipCookedClasses);
-	}
+	Utils::PopulateSkipClasses(SkipUncookedClasses, SkipCookedClasses);
 #endif
 
 	// Read default serialization options
-	constexpr static bool bInitializeFilterTagOptions = false;
-	Utils::InitializeSerializationOptionsFromIni(SerializationOptions, nullptr, UE::AssetRegistry::ESerializationTarget::ForGame, bInitializeFilterTagOptions);
-	Utils::InitializeSerializationOptionsFromIni(DevelopmentSerializationOptions, nullptr, UE::AssetRegistry::ESerializationTarget::ForDevelopment, bInitializeFilterTagOptions);
+	Utils::InitializeSerializationOptionsFromIni(SerializationOptions, nullptr, UE::AssetRegistry::ESerializationTarget::ForGame);
+	Utils::InitializeSerializationOptionsFromIni(DevelopmentSerializationOptions, nullptr, UE::AssetRegistry::ESerializationTarget::ForDevelopment);
 
 	bool bStartedAsyncGather = false;
 	if (ShouldSearchAllAssetsAtStart())
@@ -1286,10 +1241,6 @@ void UAssetRegistryImpl::InitializeEvents(UE::AssetRegistry::Impl::FInitializeCo
 	{
 		FCoreDelegates::OnAllModuleLoadingPhasesComplete.AddUObject(this, &UAssetRegistryImpl::OnInitialPluginLoadingComplete);
 	}
-	else
-	{
-		OnInitialPluginLoadingComplete();
-	}
 }
 
 UAssetRegistryImpl::UAssetRegistryImpl(FVTableHelper& Helper)
@@ -1393,10 +1344,6 @@ void UAssetRegistryImpl::OnInitialPluginLoadingComplete()
 	{
 		LLM_SCOPE(ELLMTag::AssetRegistry);
 		UE::AssetRegistry::FInterfaceWriteScopeLock InterfaceScopeLock(InterfaceLock);
-		// Read default serialization options
-		constexpr static bool bInitializeFilterTagOptions = true;
-		UE::AssetRegistry::Utils::InitializeSerializationOptionsFromIni(GuardedData.SerializationOptions, nullptr, UE::AssetRegistry::ESerializationTarget::ForGame, bInitializeFilterTagOptions);
-		UE::AssetRegistry::Utils::InitializeSerializationOptionsFromIni(GuardedData.DevelopmentSerializationOptions, nullptr, UE::AssetRegistry::ESerializationTarget::ForDevelopment, bInitializeFilterTagOptions);
 		GuardedData.OnPluginLoadingComplete(true);
 	}
 
@@ -1418,10 +1365,12 @@ void FAssetRegistryImpl::OnPluginLoadingComplete(bool bPhaseSuccessful)
 	// Reparse the skip classes the next time ShouldSkipAsset is called, since available classes
 	// for the search over all classes may have changed
 #if WITH_ENGINE && WITH_EDITOR
-	{
-		UE::TWriteScopeLock WriteLock(SkipClassesLock);
-		Utils::PopulateSkipClasses(SkipUncookedClasses, SkipCookedClasses);
-	}
+	// If we ever need to update the Filtering list outside of the game thread, we will need to defer the update 
+	// of the Filtering namespace to the tick function; UE::AssetRegistry::Filtering can only be used in game thread
+	check(IsInGameThread());
+
+	Utils::PopulateSkipClasses(SkipUncookedClasses, SkipCookedClasses);
+	UE::AssetRegistry::FFiltering::SetSkipClasses(SkipUncookedClasses, SkipCookedClasses);
 #endif
 }
 
@@ -1464,8 +1413,7 @@ void UAssetRegistryImpl::InitializeSerializationOptions(FAssetRegistrySerializat
 	}
 	else
 	{
-		constexpr static bool bInitializeFilterTagOptions = true;
-		UE::AssetRegistry::Utils::InitializeSerializationOptionsFromIni(Options, TargetPlatform, Target, bInitializeFilterTagOptions);
+		UE::AssetRegistry::Utils::InitializeSerializationOptionsFromIni(Options, TargetPlatform, Target);
 	}
 }
 
@@ -1502,8 +1450,7 @@ static TSet<FName> MakeNameSet(const TArray<FString>& Strings)
 void InitializeSerializationOptionsFromIni(
 	FAssetRegistrySerializationOptions& Options,
 	const ITargetPlatform* TargetPlatform,
-	UE::AssetRegistry::ESerializationTarget Target,
-	bool bInitializeFilterTagOptions)
+	UE::AssetRegistry::ESerializationTarget Target)
 {
 	FString PlatformIniName;
 	EBuildTargetType TargetType = EBuildTargetType::Game;
@@ -1559,24 +1506,14 @@ void InitializeSerializationOptionsFromIni(
 	}
 
 	EngineIni->GetBool(TEXT("AssetRegistry"), TEXT("bUseAssetRegistryTagsWhitelistInsteadOfBlacklist"), Options.bUseAssetRegistryTagsAllowListInsteadOfDenyList);
-
-	if (bInitializeFilterTagOptions)
+	TArray<FString> FilterListItems;
+	if (Options.bUseAssetRegistryTagsAllowListInsteadOfDenyList)
 	{
-		TArray<FString> FilterListItems;
-		if (Options.bUseAssetRegistryTagsAllowListInsteadOfDenyList)
-		{
-			EngineIni->GetArray(TEXT("AssetRegistry"), TEXT("CookedTagsWhitelist"), FilterListItems);
-		}
-		else
-		{
-			EngineIni->GetArray(TEXT("AssetRegistry"), TEXT("CookedTagsBlacklist"), FilterListItems);
-		}
-
-		const bool bKeepDevelopmentTags = bForDevelopment || FParse::Param(FCommandLine::Get(), TEXT("ARKeepDevTags"));
-		const bool bLoadClasses = true;
-		const bool bIsFilteringDevelopmentAR = bKeepDevelopmentTags;
-		TMap<FName, FParsedFilterTagRules> Rules = ParseFilterTags(FilterListItems, bKeepDevelopmentTags, Options.bUseAssetRegistryTagsAllowListInsteadOfDenyList);
-		Options.CookFilterlistTagsByClass = ResolveFilterTags(Rules, bLoadClasses);
+		EngineIni->GetArray(TEXT("AssetRegistry"), TEXT("CookedTagsWhitelist"), FilterListItems);
+	}
+	else
+	{
+		EngineIni->GetArray(TEXT("AssetRegistry"), TEXT("CookedTagsBlacklist"), FilterListItems);
 	}
 
 	{
@@ -1589,24 +1526,6 @@ void InitializeSerializationOptionsFromIni(
 		EngineIni->GetArray(TEXT("AssetRegistry"), TEXT("CookedTagsAsPathName"), AsPathName);
 		Options.CookTagsAsPath = MakeNameSet(AsPathName);
 	}
-}
-
-static int32 CalculateInheritanceDepth(const UClass* Class)
-{
-	int32 InheritanceDepth = 0;
-	if (Class)
-	{
-		for (auto Iter = Class->GetSuperStructIterator(); Iter; ++Iter)
-		{
-			++InheritanceDepth;
-		}
-	}
-	return InheritanceDepth;
-}
-
-TMap<FName, FParsedFilterTagRules> ParseFilterTags(const TArray<FString>& FilterListItems, bool bIsFilteringDevelopmentAR , bool bUseAllowListInsteadOfDenyList)
-{
-	TMap<FName, FParsedFilterTagRules> TagsOrder;
 
 	// Takes on the pattern "(Class=SomeClass,Tag=SomeTag)"
 	// Optional key KeepInDevOnly for tweaking a DevelopmentAssetRegistry (additive if allow list, subtractive if deny list)
@@ -1656,43 +1575,23 @@ TMap<FName, FParsedFilterTagRules> ParseFilterTags(const TArray<FString>& Filter
 			}
 		}
 
-		bool bAddToList;
-		if (bUseAllowListInsteadOfDenyList)
-		{
-			bAddToList = !bKeepInDevOnly || bIsFilteringDevelopmentAR;
-		}
-		else
-		{
-			// "KeepInDevOnly" means "Keep the tag only when creating the DevelopmentAssetRegistry"
-			// That is complicated meaning when we are creating a DenyList, wherein adding an element to the list removes
-			// it from the filtered collection of tags.
-			// We add it to the list if either it applies to all ARs because KeepInDevOnly was not set, or if we are creating
-			// the GameAssetRegistry.
-			bAddToList = !bKeepInDevOnly || !bIsFilteringDevelopmentAR;
-		}
-
-		if (!ClassName.IsEmpty() && !TagName.IsEmpty())
+		const bool bKeepDevelopmentTags = bForDevelopment || FParse::Param(FCommandLine::Get(), TEXT("ARKeepDevTags"));
+		const bool bPassesDevOnlyRule = !bKeepInDevOnly || Options.bUseAssetRegistryTagsAllowListInsteadOfDenyList == bKeepDevelopmentTags;
+		if (!ClassName.IsEmpty() && !TagName.IsEmpty() && bPassesDevOnlyRule)
 		{
-			FParsedFilterTagRules& Rules = TagsOrder.FindOrAdd(FName(TagName.Len(), *TagName));
-			TArray<FFilterTagsRule>& OrderedList = Rules.OrderedList;
-			const UClass* FilterlistClass = nullptr;
-			// Skip known wildcard
-			if (ClassName != TEXTVIEW("*"))
-			{
-				FilterlistClass = Cast<UClass>(StaticFindObject(UClass::StaticClass(), nullptr, *ClassName));
-			}
+			FName TagFName = FName(*TagName);
 
+			// Include subclasses if the class is in memory at this time (native classes only)
+			UClass* FilterlistClass = Cast<UClass>(StaticFindObject(UClass::StaticClass(), nullptr, *ClassName));
 			if (FilterlistClass)
 			{
-				const int32 InheritanceDepth = CalculateInheritanceDepth(FilterlistClass);
-				// Essentially I want to do OrderedList.HeapPush(FFilterTagsRule{...}, &GetInheritanceDepth);
-				// but there is no TArray::HeapPush which takes in an Item and a projection functor.
-				// Also there is no equivalent Emplace semantics to construct object inplace without extra copies/moves.
-				// So here is an inline implementation of this:
-				OrderedList.Emplace(FilterlistClass->GetClassPathName(), FilterlistClass, InheritanceDepth, bAddToList);
-				if (!Rules.bContainsAnyUnresolvedRule)
+				Options.CookFilterlistTagsByClass.FindOrAdd(FilterlistClass->GetClassPathName()).Add(TagFName);
+
+				TArray<UClass*> DerivedClasses;
+				GetDerivedClasses(FilterlistClass, DerivedClasses);
+				for (UClass* DerivedClass : DerivedClasses)
 				{
-					AlgoImpl::HeapSiftUp(OrderedList.GetData(), 0, OrderedList.Num() - 1, [](const FFilterTagsRule& Info) { return Info.InheritanceDepth; }, TLess<int32>());
+					Options.CookFilterlistTagsByClass.FindOrAdd(DerivedClass->GetClassPathName()).Add(TagFName);
 				}
 			}
 			else
@@ -1711,146 +1610,13 @@ TMap<FName, FParsedFilterTagRules> ParseFilterTags(const TArray<FString>& Filter
 				{
 					ClassPathName = FTopLevelAssetPath(ClassName);
 				}
-
-				if (!ClassPathName.IsValid())
-				{
-					UE_LOG(LogAssetRegistry, Warning, 
-						TEXT("There is a rule to %s tag \"%s\" for class \"%s\", which is not a valid class path. Ignoring."),
-						bAddToList ? TEXT("included") : TEXT("excluded"),
-						*TagName,
-						*ClassName
-					);
-					continue;
-				}
-
-				Rules.bContainsAnyUnresolvedRule = true;
-				OrderedList.Emplace(ClassPathName, nullptr, -1, bAddToList);
-			}
-		}
-	}
-
-	return TagsOrder;
-}
-
-static void ResolveRules(FStringView TagName, TArray<FFilterTagsRule>& OrderedList, bool bLoadClasses)
-{
-	auto RuleIt = OrderedList.CreateIterator();
-	if (RuleIt->ClassPath == UE::AssetRegistry::WildcardPathName)
-	{
-		// Move cursor to next element and delete the rest
-		for (++RuleIt; RuleIt; ++RuleIt)
-		{
-			const bool bIsDuplicate = (OrderedList[0].bInclude == RuleIt->bInclude);
-			UE_CLOG(!bIsDuplicate, LogAssetRegistry, Warning, 
-				TEXT("There is a wildcard rule to %s tag \"%.*s\", ")
-				TEXT("but there is another rule to %s that tag for class \"%s\". ")
-				TEXT("Only wildcard rule is accepted. Ignoring another rule."),
-				OrderedList[0].bInclude ? TEXT("include") : TEXT("exclude"),
-				TagName.Len(), TagName.GetData(),
-				RuleIt->bInclude ? TEXT("include") : TEXT("exclude"),
-				*RuleIt->ClassPath.ToString()
-			);
-			RuleIt.RemoveCurrentSwap();
-		}
-		return;
-	}
-
-	for (; RuleIt; ++RuleIt)
-	{
-		RuleIt->ClassPtr = bLoadClasses
-			? StaticLoadClass(UObject::StaticClass(), nullptr, RuleIt->ClassPath.ToString())
-			: Cast<UClass>(StaticFindObject(UClass::StaticClass(), RuleIt->ClassPath));
-		if (RuleIt->ClassPtr)
-		{
-			RuleIt->InheritanceDepth = CalculateInheritanceDepth(RuleIt->ClassPtr);
-		}
-		else
-		{
-			// Silently ignore same unresolved rules
-			const bool bIsDuplicate = (RuleIt.GetIndex() != 0) && (OrderedList[0].bInclude == RuleIt->bInclude);
-			UE_CLOG(!bIsDuplicate, LogAssetRegistry, Warning,
-				TEXT("There is a rule to %s tag \"%.*s\" for class \"%s\". Class is unknown. Ignoring."),
-				RuleIt->bInclude ? TEXT("include") : TEXT("exclude"),
-				TagName.Len(), TagName.GetData(),
-				*RuleIt->ClassPath.ToString()
-			);
-			RuleIt.RemoveCurrentSwap();
-		}
-	}
-	Algo::HeapSortBy(OrderedList, [](const FFilterTagsRule& Info) { return Info.InheritanceDepth; }, TLess<int32>());
-}
-
-TMap<FTopLevelAssetPath, TSet<FName>> ResolveFilterTags(TMap<FName, FParsedFilterTagRules>& InOutRules, bool bLoadClasses)
-{
-	TMap<FTopLevelAssetPath, TSet<FName>> Result;
-	// keep temp allocation around while we traverse the ordered list
-	TSet<const UClass*> ClassesToAdd;
-	// Populate CookFilterlistTagsByClass by walking the ordered list
-	for (TTuple<FName, FParsedFilterTagRules>& Pair : InOutRules)
-	{
-		const FName Tag = Pair.Key;
-		TStringBuilder<512> TagNameBuilder;
-		Pair.Key.ToString(TagNameBuilder);
-		const FStringView TagName = TagNameBuilder.ToView();
-		FParsedFilterTagRules& TagRules = Pair.Value;
-		TArray<FFilterTagsRule>& OrderedList = TagRules.OrderedList;
-		if (TagRules.bContainsAnyUnresolvedRule)
-		{
-			const bool bIsWildcard = OrderedList.IsValidIndex(0) && (OrderedList[0].ClassPath == UE::AssetRegistry::WildcardPathName);
-			UE_CLOG(!bIsWildcard, LogAssetRegistry, Warning,
-				TEXT("Tag \"%.*s\" has unresolved rules because some classes were not found during INI parsing. Will try to resolve again."),
-				TagName.Len(), TagName.GetData()
-			);
-			ResolveRules(TagName, OrderedList, bLoadClasses);
-			TagRules.bContainsAnyUnresolvedRule = false;
-		}
-
-		ClassesToAdd.Reset();
-		// The list is sorted in heap order of Inheritance depth,
-		// so we can just walk through it and add/remove as needed
-		TArray<UClass*> Temp;
-		for (const FFilterTagsRule& Info : OrderedList)
-		{
-			// Add Wildcard or unknown class (we guarantee to only have one of those)
-			if (!Info.ClassPtr)
-			{
-				UE_CLOG(OrderedList.Num() > 1, LogAssetRegistry, Warning, 
-					TEXT("There is more than one rule for tag \"%.*s\" with unknown class, ignoring."),
-					TagName.Len(), TagName.GetData()
-				);
-				if (Info.bInclude)
-				{
-					Result.FindOrAdd(Info.ClassPath).Add(Tag);
-				}
-				break;
-			}
-
-			Temp.Reset();
-			GetDerivedClasses(Info.ClassPtr, Temp, true);
-			if (Info.bInclude)
-			{
-				ClassesToAdd.Reserve(ClassesToAdd.Num() + Temp.Num() + 1);
-				ClassesToAdd.Add(Info.ClassPtr);
-				ClassesToAdd.Append(Temp);
-			}
-			else
-			{
-				for (UClass* Derived : Temp)
-				{
-					ClassesToAdd.Remove(Derived);
-				}
-				ClassesToAdd.Remove(Info.ClassPtr);
+				// Class is not in memory yet. Just add an explicit filter.
+				// Automatically adding subclasses of non-native classes is not supported.
+				// In these cases, using Class=* is usually sufficient				
+				Options.CookFilterlistTagsByClass.FindOrAdd(ClassPathName).Add(TagFName);
 			}
 		}
-
-		// After the final list of classes is known - add them to the CookFilterlistTagsByClass
-		for (const UClass* Class : ClassesToAdd)
-		{
-			Result.FindOrAdd(Class->GetClassPathName()).Add(Tag);
-		}
 	}
-
-	return Result;
 }
 
 }
@@ -1955,6 +1721,10 @@ void FAssetRegistryImpl::RefreshNativeClasses()
 	// and serialization options
 	CollectCodeGeneratorClasses();
 	TempCachedInheritanceBuffer.bDirty = true;
+
+	// Read default serialization options
+	Utils::InitializeSerializationOptionsFromIni(SerializationOptions, nullptr, UE::AssetRegistry::ESerializationTarget::ForGame);
+	Utils::InitializeSerializationOptionsFromIni(DevelopmentSerializationOptions, nullptr, UE::AssetRegistry::ESerializationTarget::ForDevelopment);
 }
 
 }
@@ -2110,8 +1880,6 @@ bool FAssetRegistryImpl::TryConstructGathererIfNeeded()
 	}
 
 	GlobalGatherer = MakeUnique<FAssetDataGatherer>(UAssetRegistryImpl::Get());
-	GatherDependencyDataScratchPad.Reset(new UE::AssetRegistry::DependsNode::FGatherDependencyDataScratchpad());
-
 	UpdateMaxSecondsPerFrame();
 
 	// Read script packages if all initial plugins have been loaded, otherwise do nothing; we wait for the callback.
@@ -2681,6 +2449,33 @@ bool UAssetRegistryImpl::GetInMemoryAssets(const FARCompiledFilter& CompiledFilt
 	return true;
 }
 
+bool UAssetRegistryImpl::EnumerateAssets(const FARFilter& InFilter, TFunctionRef<bool(const FAssetData&)> Callback,
+	bool bSkipARFilteredAssets) const
+{
+	FARCompiledFilter CompiledFilter;
+	CompileFilter(InFilter, CompiledFilter);
+	return EnumerateAssets(CompiledFilter, Callback, bSkipARFilteredAssets);
+}
+
+bool UAssetRegistryImpl::EnumerateAssets(const FARCompiledFilter& InFilter, TFunctionRef<bool(const FAssetData&)> Callback,
+	bool bSkipARFilteredAssets) const
+{
+	const UE::AssetRegistry::EEnumerateAssetsFlags Flags = bSkipARFilteredAssets ? UE::AssetRegistry::EEnumerateAssetsFlags::None : UE::AssetRegistry::EEnumerateAssetsFlags::AllowUnfilteredArAssets;
+	return EnumerateAssets(InFilter, Callback, Flags);
+}
+
+bool UAssetRegistryImpl::EnumerateAssets(const FARFilter& InFilter, TFunctionRef<bool(const FAssetData&)> Callback) const
+{
+	FARCompiledFilter CompiledFilter;
+	CompileFilter(InFilter, CompiledFilter);
+	return EnumerateAssets(CompiledFilter, Callback, UE::AssetRegistry::EEnumerateAssetsFlags::None);
+}
+
+bool UAssetRegistryImpl::EnumerateAssets(const FARCompiledFilter& InFilter, TFunctionRef<bool(const FAssetData&)> Callback) const
+{
+	return EnumerateAssets(InFilter, Callback, UE::AssetRegistry::EEnumerateAssetsFlags::None);
+}
+
 bool UAssetRegistryImpl::EnumerateAssets(const FARFilter& InFilter, TFunctionRef<bool(const FAssetData&)> Callback,
 	UE::AssetRegistry::EEnumerateAssetsFlags InEnumerateFlags) const
 {
@@ -2771,7 +2566,7 @@ void EnumerateMemoryAssetsHelper(const FARCompiledFilter& InFilter, TSet<FName>&
 	bool& bOutStopIteration, TFunctionRef<bool(const UObject* Object, FAssetData&& PartialAssetData)> Callback,
 	bool bSkipARFilteredAssets)
 {
-	checkf(IsInGameThread() || IsInAsyncLoadingThread() || IsInParallelLoadingThread(), TEXT("Enumerating in-memory assets can only be done on the game thread or in the loader, there are too many GetAssetRegistryTags() still not thread-safe."));
+	checkf(IsInGameThread(), TEXT("Enumerating in-memory assets can only be done on the game thread; it uses non-threadsafe UE::AssetRegistry::Filtering globals."));
 	bOutStopIteration = false;
 	UE_TRACK_REFERENCING_OPNAME_SCOPED(PackageAccessTrackingOps::NAME_ResetContext);
 
@@ -3165,6 +2960,17 @@ bool UAssetRegistryImpl::GetAllAssets(TArray<FAssetData>& OutAssetData, bool bIn
 	return true;
 }
 
+bool UAssetRegistryImpl::EnumerateAllAssets(TFunctionRef<bool(const FAssetData&)> Callback) const
+{
+	return EnumerateAllAssets(Callback, UE::AssetRegistry::EEnumerateAssetsFlags::None);
+}
+
+bool UAssetRegistryImpl::EnumerateAllAssets(TFunctionRef<bool(const FAssetData&)> Callback, bool bIncludeOnlyOnDiskAssets) const
+{
+	const UE::AssetRegistry::EEnumerateAssetsFlags Flags = bIncludeOnlyOnDiskAssets ? UE::AssetRegistry::EEnumerateAssetsFlags::OnlyOnDiskAssets : UE::AssetRegistry::EEnumerateAssetsFlags::None;
+	return EnumerateAllAssets(Callback, Flags);
+}
+
 bool UAssetRegistryImpl::EnumerateAllAssets(TFunctionRef<bool(const FAssetData&)> Callback, UE::AssetRegistry::EEnumerateAssetsFlags InEnumerateFlags) const
 {
 	TSet<FName> PackageNamesToSkip;
@@ -3214,7 +3020,7 @@ namespace Utils
 void EnumerateAllMemoryAssets(TSet<FName>& OutPackageNamesWithAssets, bool& bOutStopIteration,
 	TFunctionRef<bool(FAssetData&&)> Callback)
 {
-	checkf(IsInGameThread() || IsInAsyncLoadingThread() || IsInParallelLoadingThread(), TEXT("Enumerating memory assets can only be done on the game thread or in the loader, there are too many GetAssetRegistryTags() still not thread-safe."));
+	checkf(IsInGameThread(), TEXT("Enumerating memory assets can only be done on the game thread; it uses non-threadsafe UE::AssetRegistry::Filtering globals."));
 	bOutStopIteration = false;
 	for (FThreadSafeObjectIterator ObjIt; ObjIt; ++ObjIt)
 	{
@@ -4696,6 +4502,11 @@ void UAssetRegistryImpl::AssetsSaved(TArray<FAssetData>&& Assets)
 #endif
 }
 
+void UAssetRegistryImpl::AssetFullyUpdateTags(UObject* Object)
+{
+	AssetUpdateTags(Object, EAssetRegistryTagsCaller::Fast);
+}
+
 void UAssetRegistryImpl::AssetUpdateTags(UObject* Object, EAssetRegistryTagsCaller Caller)
 {
 #if WITH_EDITOR
@@ -4883,7 +4694,12 @@ UE::AssetRegistry::Impl::EGatherStatus UAssetRegistryImpl::TickOnBackgroundThrea
 					if (InterfaceLock.HasWaiters() || IsBackgroundProcessingPaused())
 					{
 #if WITH_EDITOR
-						++GuardedData.GetBackgroundTickInterruptionsCount();
+						// During EngineStartup many packages are loading and need to query the AssetRegistry; do not
+						// count them in the metric for backgroundtick interruptions.
+						if (IsEngineStartupModuleLoadingComplete())
+						{
+							++GuardedData.GetBackgroundTickInterruptionsCount();
+						}
 #endif
 						return true;
 					}
@@ -5758,8 +5574,6 @@ void FAssetRegistryImpl::Shrink()
 	PackagesNeedingDependencyCalculationOnGameThread.Shrink();
 #endif
 	State.Shrink();
-
-	GatherDependencyDataScratchPad->Shrink();
 }
 
 void FAssetRegistryImpl::OnInitialSearchCompleted(Impl::FEventContext& EventContext)
@@ -7309,6 +7123,27 @@ void FAssetRegistryImpl::PathDataGathered(Impl::FEventContext& EventContext, TRi
 	}
 }
 
+namespace DependsNode
+{
+typedef TMap<FDependsNode*, FDependsNode::FPackageFlagSet> FPackageDependencyMap;
+
+/** Uncompressed version of DependsNode data, used to avoid allocations and 
+compression cost when constructing batches of DependsNodes in DependencyDataGathered. */
+struct FConstructData
+{
+	FDependsNode* SourceNode;
+	FPackageDependencyMap PackageDependencies;
+	FDependsNode::FDependsNodeList SearchableNamesNodes;
+
+	FConstructData(FDependsNode* InSourceNode, FPackageDependencyMap InPackageDependencies, FDependsNode::FDependsNodeList InSearchableNamesNodes)
+		: SourceNode(InSourceNode)
+		, PackageDependencies(MoveTemp(InPackageDependencies))
+		, SearchableNamesNodes(MoveTemp(InSearchableNamesNodes))
+	{
+	}
+};
+}
+
 void FAssetRegistryImpl::DependencyDataGathered(TMultiMap<FName, FPackageDependencyData>& DependsResults,
 	TMultiMap<FName, FPackageDependencyData>& OutDeferredDependencyResults,
 	TSet<FName>* OutPackagesNeedingDependencyCalculation, Impl::FInterruptionContext& InOutInterruptionContext,
@@ -7339,13 +7174,13 @@ void FAssetRegistryImpl::DependencyDataGathered(TMultiMap<FName, FPackageDepende
 		Utils::InitializeMountPoints(MountPointsForVerifyAfterGather);
 	}
 
-	TMap<FName, FConstructData>& AssetNameToConstructDataMap = GatherDependencyDataScratchPad->AssetNameToConstructDataMap;
-	TMap<FName, FName>& CachedDepToRedirect = GatherDependencyDataScratchPad->CachedDepToRedirect;
-	TSet<FDependsNode*>& ModifiedDependNodes = GatherDependencyDataScratchPad->ModifiedDependNodes;
-	AssetNameToConstructDataMap.Reserve(DependsResults.Num());
+	TArray<FConstructData> ConstructDatas;
+	TMap<FName, FName> CachedDepToRedirect;
+	TSet<FDependsNode*> ModifiedDependNodes;
+
+	ConstructDatas.Reserve(DependsResults.Num());
 	CachedDepToRedirect.Reserve(DependsResults.Num());
 	ModifiedDependNodes.Reserve(DependsResults.Num());
-	ON_SCOPE_EXIT{ GatherDependencyDataScratchPad->Reset(); };
 
 	State.CachedPackageData.Reserve(DependsResults.Num());
 	State.CachedDependsNodes.Reserve(DependsResults.Num());
@@ -7353,7 +7188,10 @@ void FAssetRegistryImpl::DependencyDataGathered(TMultiMap<FName, FPackageDepende
 	bool bInterrupted = false;
 	int32 IterationCounter = 0;
 	constexpr int32 MinBatchSize = 1024;
-	const bool bUseJobs = DependsResults.Num() > MinBatchSize;
+
+	// Disable parallelism until UE-295029 is resolved
+	const bool bUseJobs = false;
+	//const bool bUseJobs = DependsResults.Num() > MinBatchSize;
 	const EParallelForFlags ParallelForFlags = bUseJobs ? EParallelForFlags::None : EParallelForFlags::ForceSingleThread;
 
 	{
@@ -7361,13 +7199,6 @@ void FAssetRegistryImpl::DependencyDataGathered(TMultiMap<FName, FPackageDepende
 		const bool bShouldSortReferencers = ShouldSortReferencers();
 		for (TMultiMap<FName, FPackageDependencyData>::TIterator Iter(DependsResults); Iter && !bInterrupted; ++Iter)
 		{
-			if (DeferredAssets.Contains(Iter.Key()) || DeferredAssetsForGameThread.Contains(Iter.Key()))
-			{
-				OutDeferredDependencyResults.Add(MoveTemp(*Iter));
-				Iter.RemoveCurrent();
-				continue;
-			}
-
 			ON_SCOPE_EXIT
 			{
 				// ShouldExitEarly calls FPlatformTime::Seconds which isn't super cheap
@@ -7379,68 +7210,52 @@ void FAssetRegistryImpl::DependencyDataGathered(TMultiMap<FName, FPackageDepende
 				}
 			};
 
-			FPackageDependencyData Result = MoveTemp(Iter.Value());
-			Iter.RemoveCurrent();
-
-			checkf(!GIsEditor || Result.bHasPackageData, TEXT("We rely on PackageData being read for every gathered Asset in the editor."));
+			bool bSourceNodeAlreadyExisted = false;
 			FName SourceNodePackageName = Iter.Value().PackageName;
-			if (Result.bHasPackageData)
+			FDependsNode* SourceNode = State.CreateOrFindDependsNode(SourceNodePackageName, &bSourceNodeAlreadyExisted);
+			if (!bSourceNodeAlreadyExisted)
 			{
-				// Update package data
-				FAssetPackageData* PackageData = State.CreateOrGetAssetPackageData(SourceNodePackageName);
-				*PackageData = Result.PackageData;
+				if (FPackageDependencyData* DependencyDependencyData = DependsResults.Find(SourceNodePackageName))
+				{
+					SourceNode->Reserve(DependencyDependencyData->DependsNodeReservations);
+				}
 			}
 
-			// DependsResults might contain multiple values with the same SourcePackageName. 
-			// We don't want to create more constructed nodes than necessary, so we create
-			// or find existing FConstructData and append results as needed.
-			bool bSourceNodeAlreadyExisted = false;
-			FConstructData* ConstructData = nullptr;
-			if (FConstructData* ExistingConstructData = AssetNameToConstructDataMap.Find(SourceNodePackageName))
+			if (DeferredAssets.Contains(Iter.Key()) || DeferredAssetsForGameThread.Contains(Iter.Key()))
 			{
-				ConstructData = ExistingConstructData;
+				OutDeferredDependencyResults.Add(MoveTemp(*Iter));
+				Iter.RemoveCurrent();
+				continue;
 			}
-			else
-			{
-				if (bLocalVerifyMountPointAfterGather)
-				{
-					SourceNodePackageName.ToString(PackageNameString);
-					if (!Utils::IsPathMounted(PackageNameString, *MountPointsForVerifyAfterGather, PackageRoot))
-					{
-						UE_LOG(LogAssetRegistry, Display,
-							TEXT("AssetRegistry: DependencyData has been loaded with an invalid mount point: '%s', Mount Point: '%s'. Ignoring the DependencyData."),
-							*PackageNameString, *PackageRoot);
-						continue;
-					}
-				}
 
-				FDependsNode* SourceNode = State.CreateOrFindDependsNode(SourceNodePackageName, &bSourceNodeAlreadyExisted);
-				if (!bSourceNodeAlreadyExisted)
-				{
-					if (FPackageDependencyData* DependencyDependencyData = DependsResults.Find(SourceNodePackageName))
-					{
-						SourceNode->Reserve(DependencyDependencyData->DependsNodeReservations);
-					}
-				}
+			FPackageDependencyData Result = MoveTemp(Iter.Value());
+			Iter.RemoveCurrent();
 
-				if (Result.bHasDependencyData)
+			if (bLocalVerifyMountPointAfterGather)
+			{
+				SourceNodePackageName.ToString(PackageNameString);
+				if (!Utils::IsPathMounted(PackageNameString, *MountPointsForVerifyAfterGather, PackageRoot))
 				{
-					SourceNode->Reserve(Result.DependsNodeReservations);
-					SourceNode->SetAllowShrinking(false);
-					SourceNode->SetIsReferencersSorted(bShouldSortReferencers);
-
-					ModifiedDependNodes.Add(SourceNode);
-
-					ConstructData = &AssetNameToConstructDataMap.Add(SourceNodePackageName);
-					ConstructData->SourceNode = SourceNode;
+					UE_LOG(LogAssetRegistry, Display,
+						TEXT("AssetRegistry: DependencyData has been loaded with an invalid mount point: '%s', Mount Point: '%s'. Ignoring the DependencyData."),
+						*PackageNameString, *PackageRoot);
+					continue;
 				}
 			}
 
-			if (ConstructData)
+			checkf(!GIsEditor || Result.bHasPackageData, TEXT("We rely on PackageData being read for every gathered Asset in the editor."));
+			if (Result.bHasPackageData)
+			{
+				// Update package data
+				FAssetPackageData* PackageData = State.CreateOrGetAssetPackageData(SourceNodePackageName);
+				*PackageData = Result.PackageData;
+			}
+
+			if (Result.bHasDependencyData)
 			{
 				OutPackagesNeedingDependencyCalculation->Add(SourceNodePackageName);
 
-				FPackageDependencyMap& PackageDependencies = ConstructData->PackageDependencies;
+				FPackageDependencyMap PackageDependencies;
 				PackageDependencies.Reserve(Result.PackageDependencies.Num());
 				for (const FPackageDependencyData::FPackageDependency& DependencyData : Result.PackageDependencies)
 				{
@@ -7502,7 +7317,7 @@ void FAssetRegistryImpl::DependencyDataGathered(TMultiMap<FName, FPackageDepende
 					FlagSet.Add(FDependsNode::PackagePropertiesToByte(DependencyData.Property));
 				}
 
-				FDependsNode::FDependsNodeList& SearchNameDependencyNodes = ConstructData->SearchableNamesNodes;
+				FDependsNode::FDependsNodeList SearchNameDependencyNodes;
 				for (const FPackageDependencyData::FSearchableNamesDependency& NamesDependency : Result.SearchableNameDependencies)
 				{
 					SearchNameDependencyNodes.Reserve(SearchNameDependencyNodes.Num() + NamesDependency.ValueNames.Num());
@@ -7515,30 +7330,31 @@ void FAssetRegistryImpl::DependencyDataGathered(TMultiMap<FName, FPackageDepende
 						SearchNameDependencyNodes.Add(DependsNode);
 					}
 				}
+
+				ModifiedDependNodes.Add(SourceNode);
+
+				SourceNode->Reserve(Result.DependsNodeReservations);
+				SourceNode->SetAllowShrinking(false);
+				SourceNode->SetIsReferencersSorted(bShouldSortReferencers);
+
+				ConstructDatas.Emplace(SourceNode, MoveTemp(PackageDependencies), MoveTemp(SearchNameDependencyNodes));
 			}
 		}
 	}
 
-	if (AssetNameToConstructDataMap.IsEmpty())
+	if (ConstructDatas.IsEmpty())
 	{
 		return;
 	}
 
-	TArray<FConstructData*, TInlineAllocator<1>> ConstructDatas;
-	ConstructDatas.Reserve(AssetNameToConstructDataMap.Num());
-	for (TPair<FName, FConstructData>& Pair : AssetNameToConstructDataMap)
-	{
-		ConstructDatas.Add(&Pair.Value);
-	}
-
 	{
 		TRACE_CPUPROFILER_EVENT_SCOPE(DependencyGatheredData_RemoveReferencers);
 		// Remove referencers - Single threaded due to the poor memory access patterns 
 		// between nodes and their dependencies
-		for (FConstructData* ConstructData : ConstructDatas)
+		for (FConstructData& ConstructData : ConstructDatas)
 		{
-			FDependsNode* SourceNode = ConstructData->SourceNode;
-			const FPackageDependencyMap& PackageDependencies = ConstructData->PackageDependencies;
+			FDependsNode* SourceNode = ConstructData.SourceNode;
+			const FPackageDependencyMap& PackageDependencies = ConstructData.PackageDependencies;
 
 			// We will populate the node dependencies below. Empty the set here in case this file was already read
 			// Also remove references to all existing dependencies, those will be also repopulated below
@@ -7563,10 +7379,10 @@ void FAssetRegistryImpl::DependencyDataGathered(TMultiMap<FName, FPackageDepende
 	ParallelFor(TEXT("DependencyGatheredData_AddDependencies"), ConstructDatas.Num(), MinBatchSize, 
 		[ShouldSortDependencyList = ShouldSortDependencies(), &ConstructDatas] (int Index)
 		{
-			FConstructData* ConstructData = ConstructDatas[Index];
-			FDependsNode* SourceNode = ConstructData->SourceNode;
-			const FPackageDependencyMap& PackageDependencies = ConstructData->PackageDependencies;
-			const FDependsNode::FDependsNodeList& SearchableNameNodes = ConstructData->SearchableNamesNodes;
+			FConstructData& ConstructData = ConstructDatas[Index];
+			FDependsNode* SourceNode = ConstructData.SourceNode;
+			const FPackageDependencyMap& PackageDependencies = ConstructData.PackageDependencies;
+			const FDependsNode::FDependsNodeList& SearchableNameNodes = ConstructData.SearchableNamesNodes;
 
 			SourceNode->ClearDependencies();
 			SourceNode->SetIsDependencyListSorted(EDependencyCategory::All, ShouldSortDependencyList);
@@ -7588,11 +7404,11 @@ void FAssetRegistryImpl::DependencyDataGathered(TMultiMap<FName, FPackageDepende
 	// Add referencers - Single threaded due to the poor memory access patterns between nodes and their dependencies
 	{
 		TRACE_CPUPROFILER_EVENT_SCOPE(DependencyGatheredData_AddReferencers);
-		for (FConstructData* ConstructData : ConstructDatas)
+		for (FConstructData& ConstructData : ConstructDatas)
 		{
-			FDependsNode* SourceNode = ConstructData->SourceNode;
-			const FPackageDependencyMap& PackageDependencies = ConstructData->PackageDependencies;
-			const FDependsNode::FDependsNodeList& SearchableNameNodes = ConstructData->SearchableNamesNodes;
+			FDependsNode* SourceNode = ConstructData.SourceNode;
+			const FPackageDependencyMap& PackageDependencies = ConstructData.PackageDependencies;
+			const FDependsNode::FDependsNodeList& SearchableNameNodes = ConstructData.SearchableNamesNodes;
 
 			// Doubly-link all of the PackageDependencies
 			for (const TPair<FDependsNode*, FDependsNode::FPackageFlagSet>& Pair : PackageDependencies)
@@ -9840,14 +9656,23 @@ void FAssetRegistryImpl::SetManageReferences(FSetManageReferencesContext& Contex
 					EDependencyProperty DependencyProperties, bool bDuplicate)
 					{
 						// Skip dependencies that are EditorOnly and non-build.
-						// Propagate only through UsedInGame or Build dependencies.
+						// Propagate only through UsedInGame or build dependencies.
 						if (EnumHasAnyFlags(DependencyProperties, EDependencyProperty::Game | EDependencyProperty::Build))
 						{
-							// ChunkManagement propagates through either UsedGame or Build dependencies, but CookRules
-							// propagate only through Game dependencies.
+							// Different behavior in the graph searches is temporarily disabled in the first submit,
+							// So we can test the new graph search independently of the desired behavior change.
+							static constexpr bool bCookRulesPaintedOnlyWithGame = false;
+
 							EManageSearchColorBit EdgePaintColors;
-							EdgePaintColors = EnumHasAnyFlags(DependencyProperties, EDependencyProperty::Game) ?
+							if (bCookRulesPaintedOnlyWithGame)
+							{
+								EdgePaintColors = EnumHasAnyFlags(DependencyProperties, EDependencyProperty::Game) ?
 									EManageSearchColorBit::All : EManageSearchColorBit::Chunking;
+							}
+							else
+							{
+								EdgePaintColors = EManageSearchColorBit::All;
+							}
 							EManageSearchColorBit TargetPaintColors = VertexPaintColors & EdgePaintColors;
 							VisitNode(SourceNode, TargetNode, DependencyCategory, DependencyProperties, TargetPaintColors);
 						}
@@ -10059,7 +9884,6 @@ const TSet<FName>& FAssetRegistryImpl::GetCachedEmptyPackages() const
 bool FAssetRegistryImpl::ShouldSkipAsset(FTopLevelAssetPath AssetClass, uint32 PackageFlags) const
 {
 #if WITH_ENGINE && WITH_EDITOR
-	UE::TReadScopeLock ReadLock(SkipClassesLock);
 	return Utils::ShouldSkipAsset(AssetClass, PackageFlags, SkipUncookedClasses, SkipCookedClasses);
 #else
 	return false;
@@ -10069,7 +9893,6 @@ bool FAssetRegistryImpl::ShouldSkipAsset(FTopLevelAssetPath AssetClass, uint32 P
 bool FAssetRegistryImpl::ShouldSkipAsset(const UObject* InAsset) const
 {
 #if WITH_ENGINE && WITH_EDITOR
-	UE::TReadScopeLock ReadLock(SkipClassesLock);
 	return Utils::ShouldSkipAsset(InAsset, SkipUncookedClasses, SkipCookedClasses);
 #else
 	return false;
@@ -10189,6 +10012,19 @@ void FEventContext::Append(FEventContext&& Other)
 
 }
 
+void UAssetRegistryImpl::ReadLockEnumerateTagToAssetDatas(TFunctionRef<void(FName TagName, const TArray<const FAssetData*>& Assets)> Callback) const
+{
+	UE_LOG(LogAssetRegistry, Error, TEXT("ReadLockEnumerateTagToAssetDatas has been deprecated. Use ReadLockEnumerateAllTagToAssetDatas instead."));
+
+	UE::AssetRegistry::FInterfaceReadScopeLock InterfaceScopeLock(InterfaceLock);
+	GuardedData.GetState().EnumerateTags([Callback](FName TagName)
+	{
+		TArray<const FAssetData*> EmptyArray;
+		Callback(TagName, EmptyArray);
+		return true;
+	});
+}
+
 void UAssetRegistryImpl::ReadLockEnumerateAllTagToAssetDatas(TFunctionRef<bool(FName TagName, IAssetRegistry::FEnumerateAssetDatasFunc EnumerateAssets)> Callback) const
 {
 	UE::AssetRegistry::FInterfaceReadScopeLock InterfaceScopeLock(InterfaceLock);
diff --git a/Engine/Source/Runtime/AssetRegistry/Private/AssetRegistryImpl.h b/Engine/Source/Runtime/AssetRegistry/Private/AssetRegistryImpl.h
--- a/Engine/Source/Runtime/AssetRegistry/Private/AssetRegistryImpl.h
+++ b/Engine/Source/Runtime/AssetRegistry/Private/AssetRegistryImpl.h
@@ -1034,33 +1034,7 @@ void RunAssetsThroughFilter(TArray<FAssetData>& AssetDataList, const FARCompiled
 void InitializeSerializationOptionsFromIni(
 	FAssetRegistrySerializationOptions& Options,
 	const ITargetPlatform* TargetPlatform = nullptr,
-	ESerializationTarget Target = ESerializationTarget::ForGame,
-	bool bInitializeFilterTagOptions = true);
-
-struct FFilterTagsRule
-{
-	FTopLevelAssetPath ClassPath;
-	const UClass* ClassPtr = nullptr;
-	int32 InheritanceDepth = 0;
-	bool bInclude = false;
-};
-
-struct FParsedFilterTagRules
-{
-	TArray<FFilterTagsRule> OrderedList;
-	bool bContainsAnyUnresolvedRule = false;
-};
-/**
- * Parses the list of CookedTagsBlacklist or CookedTagsWhitelist gathered from the INI.
- * Returns a mapping from a tag to a list of rules. May contain unresolved rules.
- */
-TMap<FName, FParsedFilterTagRules> ParseFilterTags(const TArray<FString>& FilterListItems, bool bIsFilteringDevelopmentAR, bool bUseAllowListInsteadOfDenyList);
-
-/**
- * Resolves and expands the list of rules to list of all derived loaded classes that satisfy those rules.
- * Then transposes rules to a map with a class as a key and a list of tags as value.
- */
-TMap<FTopLevelAssetPath, TSet<FName>> ResolveFilterTags(TMap<FName, FParsedFilterTagRules>& InOutRules, bool bLoadClasses);
+	ESerializationTarget Target = ESerializationTarget::ForGame);
 
 /** Gets the current availability of an asset, primarily for streaming install purposes. */
 EAssetAvailability::Type GetAssetAvailability(const FAssetData& AssetData);
