IAssetRegistry::OnEnumerateAssetsEvent : make the event thread safe.

#rnx

diff --git a/Engine/Source/Runtime/AssetRegistry/Public/AssetRegistry/IAssetRegistry.h b/Engine/Source/Runtime/AssetRegistry/Public/AssetRegistry/IAssetRegistry.h
--- a/Engine/Source/Runtime/AssetRegistry/Public/AssetRegistry/IAssetRegistry.h
+++ b/Engine/Source/Runtime/AssetRegistry/Public/AssetRegistry/IAssetRegistry.h
@@ -1125,7 +1125,7 @@ public:
 	};
 
 	/** Event triggered when a query is made. The event is triggered only if the filter is valid. */
-	DECLARE_MULTICAST_DELEGATE_TwoParams(FEnumerateAssetsEvent, const FARCompiledFilter& /*Filter*/, UE::AssetRegistry::EEnumerateAssetsFlags /*Flag*/)
+	DECLARE_TS_MULTICAST_DELEGATE_TwoParams(FEnumerateAssetsEvent, const FARCompiledFilter& /*Filter*/, UE::AssetRegistry::EEnumerateAssetsFlags /*Flag*/)
 	virtual FEnumerateAssetsEvent& OnEnumerateAssetsEvent() = 0;
 
 protected:
