Reduce AssetRegistry size when merging in asset registry from DLC cooks by shrinking and trimming duplicates when sorting.

diff --git a/Engine/Source/Runtime/AssetRegistry/Private/AssetRegistry.cpp b/Engine/Source/Runtime/AssetRegistry/Private/AssetRegistry.cpp
--- a/Engine/Source/Runtime/AssetRegistry/Private/AssetRegistry.cpp
+++ b/Engine/Source/Runtime/AssetRegistry/Private/AssetRegistry.cpp
@@ -6201,6 +6201,7 @@ void UAssetRegistryImpl::InitializeTemporaryAssetRegistryState(FAssetRegistrySta
 		{
 			TSet<int32> UnusedChunksToKeep;
 			OutState.InitializeFromExistingAndPrune(State, RequiredPackages, RemovePackages, UnusedChunksToKeep, Options);
+			OutState.Shrink();
 		}
 	}
 	else
diff --git a/Engine/Source/Runtime/AssetRegistry/Private/AssetRegistryState.cpp b/Engine/Source/Runtime/AssetRegistry/Private/AssetRegistryState.cpp
--- a/Engine/Source/Runtime/AssetRegistry/Private/AssetRegistryState.cpp
+++ b/Engine/Source/Runtime/AssetRegistry/Private/AssetRegistryState.cpp
@@ -578,6 +578,7 @@ void FAssetRegistryState::InitializeFromExisting(const FAssetDataMap& AssetDataM
 		// Restore the sortedness that we turned off for performance when creating each DependsNode
 		SetDependencyNodeSorting(true, true);
 	}
+	Shrink();
 }
 
 void FAssetRegistryState::PruneAssetData(const TSet<FName>& RequiredPackages, const TSet<FName>& RemovePackages,
@@ -2805,8 +2806,11 @@ void FAssetRegistryState::SetDependencyNodeSorting(bool bSortDependencies, bool
 	for (TPair<FAssetIdentifier, FDependsNode*>& Pair : CachedDependsNodes)
 	{
 		FDependsNode* DependsNode = Pair.Value;
+		const bool bOldAllowShrinking = DependsNode->GetAllowShrinking();
+		DependsNode->SetAllowShrinking(true);
 		DependsNode->SetIsDependencyListSorted(UE::AssetRegistry::EDependencyCategory::All, bSortDependencies);
 		DependsNode->SetIsReferencersSorted(bSortReferencers);
+		DependsNode->SetAllowShrinking(bOldAllowShrinking);
 	}
 }
 
@@ -3210,12 +3214,11 @@ void FAssetRegistryState::Shrink()
 #endif
 
 	CachedAssetsByPackageName.Shrink();
-	CachedAssetsByPath.Shrink();
-	CachedAssetsByClass.Shrink();
-
+	ShrinkMultimap(CachedAssetsByPath);
+	ShrinkMultimap(CachedAssetsByClass);
 
 #if UE_ASSETREGISTRY_CACHEDASSETSBYTAG
-	CachedAssetsByTag.Shrink();
+	ShrinkMultimap(CachedAssetsByTag);
 #else
 	CachedClassesByTag.Shrink();
 #endif
