revert ShrinkMultiMap change

diff --git a/Engine/Source/Runtime/AssetRegistry/Private/AssetRegistryState.cpp b/Engine/Source/Runtime/AssetRegistry/Private/AssetRegistryState.cpp
--- a/Engine/Source/Runtime/AssetRegistry/Private/AssetRegistryState.cpp
+++ b/Engine/Source/Runtime/AssetRegistry/Private/AssetRegistryState.cpp
@@ -3218,7 +3218,7 @@ void FAssetRegistryState::Shrink()
 	ShrinkMultimap(CachedAssetsByClass);
 
 #if UE_ASSETREGISTRY_CACHEDASSETSBYTAG
-	ShrinkMultimap(CachedAssetsByTag);
+	CachedAssetsByTag.Shrink();
 #else
 	CachedClassesByTag.Shrink();
 #endif
