Adding ability for AppendAssetRegistry to only update new pacakges. this way DLC does not clobber data in the base game.
When appending an AssetRegistry to another only serialize and add the AssetData which respect the current project's serialization rules

diff --git a/Engine/Source/Runtime/AssetRegistry/Private/AssetRegistry.cpp b/Engine/Source/Runtime/AssetRegistry/Private/AssetRegistry.cpp
--- a/Engine/Source/Runtime/AssetRegistry/Private/AssetRegistry.cpp
+++ b/Engine/Source/Runtime/AssetRegistry/Private/AssetRegistry.cpp
@@ -5966,13 +5966,13 @@ void FAssetRegistryImpl::Serialize(FArchive& Ar, Impl::FEventContext& EventConte
 }
 
 /** Append the assets from the incoming state into our own */
-void UAssetRegistryImpl::AppendState(const FAssetRegistryState& InState)
+void UAssetRegistryImpl::AppendState(const FAssetRegistryState& InState, UE::AssetRegistry::EAppendMode AppendMode)
 {
 	UE::AssetRegistry::Impl::FEventContext EventContext;
 	{
 		LLM_SCOPE(ELLMTag::AssetRegistry);
 		UE::AssetRegistry::FInterfaceWriteScopeLock InterfaceScopeLock(InterfaceLock);
-		GuardedData.AppendState(EventContext, InState, FAssetRegistryState::EInitializationMode::Append, /*bEmitAssetEvents*/true);
+		GuardedData.AppendState(EventContext, InState, AppendMode, /*bEmitAssetEvents*/true);
 	}
 	Broadcast(EventContext);
 }
@@ -5997,6 +5997,7 @@ void FAssetRegistryImpl::AppendState(Impl::FEventContext& EventContext, const FA
 		AppendResultPtr);
 
 	CachePathsFromState(EventContext, InState);
+	CachedPathTree.Shrink();
 
 	if (AppendResultPtr)
 	{
@@ -6018,9 +6019,22 @@ void FAssetRegistryImpl::CachePathsFromState(Impl::FEventContext& EventContext,
 	// Refreshes ClassGeneratorNames if out of date due to module load
 	CollectCodeGeneratorClasses();
 
+#if WITH_EDITOR
+	const FAssetRegistrySerializationOptions& Options = DevelopmentSerializationOptions;
+#else
+	const FAssetRegistrySerializationOptions& Options = SerializationOptions;
+#endif
+
 	// Add paths to cache
-	InState.EnumerateAllAssets([this, &EventContext](const FAssetData& AssetData)
+	InState.EnumerateAllAssets([this, &Options, &EventContext](const FAssetData& AssetData)
 	{
+		if (Options.bFilterAssetDataWithNoTags && AssetData.TagsAndValues.Num() == 0 &&
+			!FPackageName::IsLocalizedPackage(WriteToString<256>(AssetData.PackageName)))
+		{
+			// based on current serialization settings this should be skipped.
+			return;
+		}
+
 		AddAssetPath(EventContext, AssetData.PackagePath);
 
 		// Populate the class map if adding blueprint
diff --git a/Engine/Source/Runtime/AssetRegistry/Private/AssetRegistry.h b/Engine/Source/Runtime/AssetRegistry/Private/AssetRegistry.h
--- a/Engine/Source/Runtime/AssetRegistry/Private/AssetRegistry.h
+++ b/Engine/Source/Runtime/AssetRegistry/Private/AssetRegistry.h
@@ -110,7 +110,7 @@ public:
 	virtual void ScanModifiedAssetFiles(const TArray<FString>& InFilePaths) override;
 	virtual void ScanModifiedAssetFiles(const TArray<FString>& InFilePaths, UE::AssetRegistry::EScanFlags ScanFlags) override;
 	virtual void Serialize(FArchive& Ar) override;
-	virtual void AppendState(const FAssetRegistryState& InState) override;
+	virtual void AppendState(const FAssetRegistryState& InState, UE::AssetRegistry::EAppendMode AppendMode = UE::AssetRegistry::EAppendMode::Append) override;
 	virtual SIZE_T GetAllocatedSize(bool bLogDetailed = false) const override;
 	virtual void LoadPackageRegistryData(FArchive& Ar, FLoadPackageRegistryData& InOutData) const override;
 	virtual void LoadPackageRegistryData(const FString& PackageFilename, FLoadPackageRegistryData& InOutData) const override;
diff --git a/Engine/Source/Runtime/AssetRegistry/Private/AssetRegistryState.cpp b/Engine/Source/Runtime/AssetRegistry/Private/AssetRegistryState.cpp
--- a/Engine/Source/Runtime/AssetRegistry/Private/AssetRegistryState.cpp
+++ b/Engine/Source/Runtime/AssetRegistry/Private/AssetRegistryState.cpp
@@ -451,6 +451,13 @@ void FAssetRegistryState::InitializeFromExisting(const FAssetDataMap& AssetDataM
 			continue;
 		}
 
+		if (Options.bFilterAssetDataWithNoTags && AssetData.TagsAndValues.Num() == 0 &&
+			!FPackageName::IsLocalizedPackage(WriteToString<256>(AssetData.PackageName)))
+		{
+			// based on current serialization settings this should be skipped.
+			continue;
+		}
+
 		// Filter asset registry tags now
 		FAssetDataTagMap LocalTagsAndValues;
 		FAssetRegistryState::FilterTags(AssetData.TagsAndValues, LocalTagsAndValues,
diff --git a/Engine/Source/Runtime/AssetRegistry/Private/PathTree.cpp b/Engine/Source/Runtime/AssetRegistry/Private/PathTree.cpp
--- a/Engine/Source/Runtime/AssetRegistry/Private/PathTree.cpp
+++ b/Engine/Source/Runtime/AssetRegistry/Private/PathTree.cpp
@@ -245,3 +245,9 @@ bool FPathTree::EnumerateSubPaths(FName BasePath, TFunctionRef<bool(FName)> Call
 
 	return true;
 }
+
+void FPathTree::Shrink()
+{
+	ParentPathToChildPaths.Shrink();
+	ChildPathToParentPath.Shrink();
+}
diff --git a/Engine/Source/Runtime/AssetRegistry/Public/AssetRegistry/AssetRegistryState.h b/Engine/Source/Runtime/AssetRegistry/Public/AssetRegistry/AssetRegistryState.h
--- a/Engine/Source/Runtime/AssetRegistry/Public/AssetRegistry/AssetRegistryState.h
+++ b/Engine/Source/Runtime/AssetRegistry/Public/AssetRegistry/AssetRegistryState.h
@@ -181,16 +181,7 @@ public:
 	FAssetRegistryState& operator=(const FAssetRegistryState&) = delete;
 	ASSETREGISTRY_API FAssetRegistryState& operator=(FAssetRegistryState&& O);
 
-	/**
-	 * Enum controlling how we initialize this state
-	 */
-	enum class EInitializationMode
-	{
-		Rebuild,
-		OnlyUpdateExisting,
-		Append,
-		OnlyUpdateNew,
-	};
+	using EInitializationMode = UE::AssetRegistry::EAppendMode;
 
 	/**
 	 * Does the given path contain assets?
diff --git a/Engine/Source/Runtime/AssetRegistry/Public/AssetRegistry/IAssetRegistry.h b/Engine/Source/Runtime/AssetRegistry/Public/AssetRegistry/IAssetRegistry.h
--- a/Engine/Source/Runtime/AssetRegistry/Public/AssetRegistry/IAssetRegistry.h
+++ b/Engine/Source/Runtime/AssetRegistry/Public/AssetRegistry/IAssetRegistry.h
@@ -127,6 +127,17 @@ struct FSetManageReferencesContext
 	TPimplPtr<FSetManageReferencesScratch> Scratch;
 };
 
+/**
+ * Enum controlling how FAssetRegistryStates are merged during Append functions, e.g. when adding plugin's serialized AssetRegistry into global AssetRegistry.
+ */
+enum class EAppendMode
+{
+	Rebuild,
+	OnlyUpdateExisting,
+	Append,
+	OnlyUpdateNew,
+};
+
 } // namespace UE::AssetRegistry
 
 USTRUCT(BlueprintType)
@@ -1016,7 +1027,7 @@ public:
 	virtual void Serialize(FStructuredArchive::FRecord Record) = 0;
 
 	/** Append the assets from the incoming state into our own */
-	virtual void AppendState(const FAssetRegistryState& InState) = 0;
+	virtual void AppendState(const FAssetRegistryState& InState, UE::AssetRegistry::EAppendMode AppendMode = UE::AssetRegistry::EAppendMode::Append) = 0;
 
 	/** Returns memory size of entire registry, optionally logging sizes */
 	virtual SIZE_T GetAllocatedSize(bool bLogDetailed = false) const = 0;
diff --git a/Engine/Source/Runtime/AssetRegistry/Public/AssetRegistry/PathTree.h b/Engine/Source/Runtime/AssetRegistry/Public/AssetRegistry/PathTree.h
--- a/Engine/Source/Runtime/AssetRegistry/Public/AssetRegistry/PathTree.h
+++ b/Engine/Source/Runtime/AssetRegistry/Public/AssetRegistry/PathTree.h
@@ -39,6 +39,9 @@ public:
 	/** Recursively enumerates all child paths from the specified base path relative to this node */
 	ASSETREGISTRY_API bool EnumerateSubPaths(FName BasePath, TFunctionRef<bool(FName)> Callback, bool bRecurse = true) const;
 
+	/** Shrinks the internal structures reducing the memory they use. */
+	ASSETREGISTRY_API void Shrink();
+
 	int32 NumPaths() const
 	{
 		return ParentPathToChildPaths.Num();
