[AssetRegistry] Fix crash using expired iterator

#tests compiled

diff --git a/Engine/Source/Runtime/AssetRegistry/Private/AssetRegistry.cpp b/Engine/Source/Runtime/AssetRegistry/Private/AssetRegistry.cpp
--- a/Engine/Source/Runtime/AssetRegistry/Private/AssetRegistry.cpp
+++ b/Engine/Source/Runtime/AssetRegistry/Private/AssetRegistry.cpp
@@ -7206,7 +7206,7 @@ void FAssetRegistryImpl::DependencyDataGathered(TMultiMap<FName, FPackageDepende
 			Iter.RemoveCurrent();
 
 			checkf(!GIsEditor || Result.bHasPackageData, TEXT("We rely on PackageData being read for every gathered Asset in the editor."));
-			FName SourceNodePackageName = Iter.Value().PackageName;
+			FName SourceNodePackageName = Result.PackageName;
 			if (Result.bHasPackageData)
 			{
 				// Update package data
