[AssetRegistry]

- Renames IsDependenciesInitialized to IsScriptDependenciesInitialized and SetIsDependenciesInitialized to SetIsScriptDependenciesInitialized.
- Allow GetDependencies to return dependency information for missing assets by not skipping nodes that were not explicitly gathered as source nodes corresponding to assets in the registry. Corrects behaviour change from 42288628

#rnx

diff --git a/Engine/Source/Runtime/AssetRegistry/Private/AssetRegistry.cpp b/Engine/Source/Runtime/AssetRegistry/Private/AssetRegistry.cpp
--- a/Engine/Source/Runtime/AssetRegistry/Private/AssetRegistry.cpp
+++ b/Engine/Source/Runtime/AssetRegistry/Private/AssetRegistry.cpp
@@ -7291,33 +7291,35 @@ void FAssetRegistryImpl::DependencyDataGathered(TMultiMap<FName, FPackageDepende
 					bool bDependsNodeAlreadyExisted = false;
 					FDependsNode* DependsNode = State.CreateOrFindDependsNode(DependencyPackageName, &bDependsNodeAlreadyExisted);
 					ModifiedDependNodes.Add(DependsNode);
+
 					if (!bDependsNodeAlreadyExisted)
 					{
-						// See if we need to read the script package Guid.
-						// Script nodes are never represented as source nodes so we can use the initialization 
-						// state to check if we need to set the package Guid or not
-						if (DependsNode->IsScriptPath() && !DependsNode->IsDependenciesInitialized())
+						// We created a new node so lookup its reservation
+						if (FPackageDependencyData* DependencyDependencyData = DependsResults.Find(DependencyPackageName))
 						{
-							DependsNode->SetIsDependenciesInitialized(true);
-							FNameBuilder DependencyPackageNameStr(DependencyPackageName);
+							DependsNode->Reserve(DependencyDependencyData->DependsNodeReservations);
+						}
+					}
 
-							// Get the guid off the script package, it is updated when 
-							// script is changed so we need to refresh it every run
-							UPackage* Package = FindPackage(nullptr, *DependencyPackageNameStr);
+					// See if we need to read the script package Guid.
+					// Script nodes are never represented as source nodes so we can use the initialization 
+					// state to check if we need to set the package Guid or not
+					// This node may have been created outside of this loop so we should always check
+					if (DependsNode->IsScriptPath() && !DependsNode->IsScriptDependenciesInitialized())
+					{
+						DependsNode->SetIsScriptDependenciesInitialized(true);
+						FNameBuilder DependencyPackageNameStr(DependencyPackageName);
 
-							if (Package)
-							{
-								FAssetPackageData* ScriptPackageData = State.CreateOrGetAssetPackageData(DependencyPackageName);
-#if WITH_EDITORONLY_DATA
-								ScriptPackageData->SetPackageSavedHash(Package->GetSavedHash());
-#endif
-							}
-						}
+						// Get the guid off the script package, it is updated when 
+						// script is changed so we need to refresh it every run
+						UPackage* Package = FindPackage(nullptr, *DependencyPackageNameStr);
 
-						// We created a new node so lookup its reservation
-						if (FPackageDependencyData* DependencyDependencyData = DependsResults.Find(DependencyPackageName))
+						if (Package)
 						{
-							DependsNode->Reserve(DependencyDependencyData->DependsNodeReservations);
+							FAssetPackageData* ScriptPackageData = State.CreateOrGetAssetPackageData(DependencyPackageName);
+#if WITH_EDITORONLY_DATA
+							ScriptPackageData->SetPackageSavedHash(Package->GetSavedHash());
+#endif
 						}
 					}
 
@@ -7434,9 +7436,6 @@ void FAssetRegistryImpl::DependencyDataGathered(TMultiMap<FName, FPackageDepende
 				DependsNode->SetIsReferencersSorted(false);
 				DependsNode->AddReferencer(SourceNode);
 			}
-
-			// Will be true at the end of this function so set it now while iterating over the new source nodes
-			SourceNode->SetIsDependenciesInitialized(true);
 		}
 	}
 
diff --git a/Engine/Source/Runtime/AssetRegistry/Private/AssetRegistryState.cpp b/Engine/Source/Runtime/AssetRegistry/Private/AssetRegistryState.cpp
--- a/Engine/Source/Runtime/AssetRegistry/Private/AssetRegistryState.cpp
+++ b/Engine/Source/Runtime/AssetRegistry/Private/AssetRegistryState.cpp
@@ -395,7 +395,7 @@ void FAssetRegistryState::InitializeFromExistingAndPrune(const FAssetRegistrySta
 				NewDependency->AddReferencer(NewNode);
 			}
 		});
-		NewNode->SetIsDependenciesInitialized(true);
+		NewNode->SetIsScriptDependenciesInitialized(OldNode->IsScriptDependenciesInitialized());
 	}
 
 	// Remove any orphaned depends nodes. This will leave cycles in but those might represent useful data
@@ -538,9 +538,10 @@ void FAssetRegistryState::InitializeFromExisting(const FAssetDataMap& AssetDataM
 			DependsNodesToAdd = &FilteredDependsNodeMap;
 			for (const TPair<FAssetIdentifier, FDependsNode*>& Pair : DependsNodeMap)
 			{
+				bool bTargetNodeAlreadyExists = false;
 				FDependsNode* SourceNode = Pair.Value;
-				FDependsNode* TargetNode = CreateOrFindDependsNode(Pair.Key);
-				if (!TargetNode->IsDependenciesInitialized())
+				FDependsNode* TargetNode = CreateOrFindDependsNode(Pair.Key, &bTargetNodeAlreadyExists);
+				if (!bTargetNodeAlreadyExists)
 				{
 					FilteredDependsNodeMap.Add(Pair.Key, SourceNode);
 				}
@@ -579,7 +580,7 @@ void FAssetRegistryState::InitializeFromExisting(const FAssetDataMap& AssetDataM
 					TargetDependency->AddReferencer(TargetNode);
 				}
 			});
-			TargetNode->SetIsDependenciesInitialized(true);
+			TargetNode->SetIsScriptDependenciesInitialized(SourceNode->IsScriptDependenciesInitialized());
 		}
 
 		// Restore the sortedness that we turned off for performance when creating each DependsNode
@@ -1636,11 +1637,6 @@ bool FAssetRegistryState::GetDependencies(const FAssetIdentifier& AssetIdentifie
 
 	if (Node != nullptr)
 	{
-		if (!Node->IsDependenciesInitialized())
-		{
-			return false;
-		}
-
 		Node->GetDependencies(OutDependencies, Category, Flags);
 		return true;
 	}
diff --git a/Engine/Source/Runtime/AssetRegistry/Private/DependsNode.cpp b/Engine/Source/Runtime/AssetRegistry/Private/DependsNode.cpp
--- a/Engine/Source/Runtime/AssetRegistry/Private/DependsNode.cpp
+++ b/Engine/Source/Runtime/AssetRegistry/Private/DependsNode.cpp
@@ -1140,7 +1140,7 @@ void FDependsNode::SerializeLoad(FArchive& Ar,
 	SerializeLoad_ManageDependencies<bLatestVersion>(Ar, GetNodeFromSerializeIndex, Scratch, Version);
 	SerializeLoad_Referencers<bLatestVersion>(Ar, GetNodeFromSerializeIndex, Scratch, Version);
 
-	SetIsDependenciesInitialized(true);
+	SetIsScriptDependenciesInitialized(IsScriptPath());
 }
 
 // We need Explicit instantiations of the template function because they are only called from other translation units.
@@ -1275,7 +1275,7 @@ void FDependsNode::LegacySerializeLoad_BeforeFlags(FArchive& Ar,
 		SerializeNodeArray(NumReferencers, Referencers, nullptr, 0, 0, true, true);
 	}
 
-	SetIsDependenciesInitialized(true);
+	SetIsScriptDependenciesInitialized(IsScriptPath());
 }
 
 template <uint32 FlagWidth>
@@ -1437,12 +1437,12 @@ void FDependsNode::SetIsReferencersSorted(bool bValue)
 	ReferencersIsSorted = bValue ? 1 : 0;
 }
 
-bool FDependsNode::IsDependenciesInitialized() const
+bool FDependsNode::IsScriptDependenciesInitialized() const
 {
 	return DependenciesInitialized;
 }
 
-void FDependsNode::SetIsDependenciesInitialized(bool bValue)
+void FDependsNode::SetIsScriptDependenciesInitialized(bool bValue)
 {
 	DependenciesInitialized = bValue ? 1 : 0;
 }
diff --git a/Engine/Source/Runtime/AssetRegistry/Private/DependsNode.h b/Engine/Source/Runtime/AssetRegistry/Private/DependsNode.h
--- a/Engine/Source/Runtime/AssetRegistry/Private/DependsNode.h
+++ b/Engine/Source/Runtime/AssetRegistry/Private/DependsNode.h
@@ -178,8 +178,12 @@ public:
 	void SetIsDependencyListSorted(UE::AssetRegistry::EDependencyCategory Category, bool bValue);
 	bool IsReferencersSorted() const;
 	void SetIsReferencersSorted(bool bValue);
-	bool IsDependenciesInitialized() const;
-	void SetIsDependenciesInitialized(bool bValue);
+	/** 
+	* Returns true if an IsScriptPath() FDependsNode has been initialized. FDependsNodes 
+	* that are not IsScriptPath() will always return false. 
+	*/
+	bool IsScriptDependenciesInitialized() const;
+	void SetIsScriptDependenciesInitialized(bool bValue);
 
 private:
 
