AssetRegistry: Add logs to help diagnose when a deadlock occurs in a listener to OnFilesLoaded.
#rnx

diff --git a/Engine/Source/Runtime/AssetRegistry/Private/AssetRegistry.cpp b/Engine/Source/Runtime/AssetRegistry/Private/AssetRegistry.cpp
--- a/Engine/Source/Runtime/AssetRegistry/Private/AssetRegistry.cpp
+++ b/Engine/Source/Runtime/AssetRegistry/Private/AssetRegistry.cpp
@@ -10285,8 +10285,12 @@ void UAssetRegistryImpl::Broadcast(UE::AssetRegistry::Impl::FEventContext& Event
 		// but they can be before or after non-completion events that occurred after completion was triggered.
 		if (bLocalFileLoadedEventBroadcast)
 		{
+			double BroadcastTimeStart = FPlatformTime::Seconds();
+			UE_LOG(LogAssetRegistry, Display, TEXT("Starting OnFilesLoaded.Broadcast"));
 			FileLoadedEvent.Broadcast();
 			ScanEndedEvent.Broadcast();
+			UE_LOG(LogAssetRegistry, Display, TEXT("Completed OnFilesLoaded.Broadcast: %.3fs"),
+				(float) FPlatformTime::Seconds() - BroadcastTimeStart);
 			EventContext.bHasSentFileLoadedEventBroadcast = true;
 			EventContext.bFileLoadedEventBroadcast = false;
 		}
