AssetRegistry: AssetDataMap: In the compressed version of FAssetDataMap and FAssetPackageNameMap, add more functions to support the TSet API that they emulate, for use in some parallelization of AssetRegistryState functions.
FAssetDataMap: Add the equivalent of TSet::Get() so that it supports ParallelFor enumeration of assets.
FAssetPackageNameMap: Add return value for the number of elements removed by a remove call.
#rnx

diff --git a/Engine/Source/Runtime/AssetRegistry/Private/AssetDataMap.cpp b/Engine/Source/Runtime/AssetRegistry/Private/AssetDataMap.cpp
--- a/Engine/Source/Runtime/AssetRegistry/Private/AssetDataMap.cpp
+++ b/Engine/Source/Runtime/AssetRegistry/Private/AssetDataMap.cpp
@@ -197,6 +197,22 @@ TArray<FAssetData*> FAssetDataMap::Array() const
 	return Result;
 }
 
+int32 FAssetDataMap::GetMaxIndex() const
+{
+	return AssetDatas.Num();
+}
+
+bool FAssetDataMap::IsValidId(FSetElementId SetId) const
+{
+	int32 Index = SetId.AsInteger();
+	return 0 <= Index && Index < AssetDatas.Num() && IsInUse(AssetDatas[Index]);
+}
+
+FAssetData* FAssetDataMap::Get(FSetElementId SetId) const
+{
+	return AssetDatas[SetId.AsInteger()];
+}
+
 bool FAssetDataMap::Contains(const FCachedAssetKey& Key) const
 {
 	return FindId(Key) != AssetDataPtrIndexInvalid;
@@ -655,7 +671,7 @@ void FAssetPackageNameMap::Add(FName PackageName, FAssetDataPtrIndex AssetIndex)
 	}
 }
 
-void FAssetPackageNameMap::Remove(FName PackageName, FAssetDataPtrIndex AssetIndex)
+int32 FAssetPackageNameMap::Remove(FName PackageName, FAssetDataPtrIndex AssetIndex)
 {
 	uint32 PackageNameTypeHash = GetTypeHash(PackageName);
 	const FAssetDataOrArrayIndex* OldPtr = AssetOrArrayByPackageName->FindByHash(PackageNameTypeHash, PackageName);
@@ -671,8 +687,10 @@ void FAssetPackageNameMap::Remove(FName PackageName, FAssetDataPtrIndex AssetInd
 			{
 				AssetOrArrayByPackageName->AddByHash(PackageNameTypeHash, NewStoredValue);
 			}
+			return 1;
 		}
 	}
+	return 0;
 }
 
 int32 FAssetPackageNameMap::Num() const
diff --git a/Engine/Source/Runtime/AssetRegistry/Public/AssetRegistry/AssetDataMap.h b/Engine/Source/Runtime/AssetRegistry/Public/AssetRegistry/AssetDataMap.h
--- a/Engine/Source/Runtime/AssetRegistry/Public/AssetRegistry/AssetDataMap.h
+++ b/Engine/Source/Runtime/AssetRegistry/Public/AssetRegistry/AssetDataMap.h
@@ -156,6 +156,9 @@ public:
 	ASSETREGISTRY_API int32 Num() const;
 	ASSETREGISTRY_API SIZE_T GetAllocatedSize() const;
 	ASSETREGISTRY_API TArray<FAssetData*> Array() const;
+	ASSETREGISTRY_API int32 GetMaxIndex() const;
+	ASSETREGISTRY_API bool IsValidId(FSetElementId SetId) const;
+	ASSETREGISTRY_API FAssetData* Get(FSetElementId SetId) const;
 
 	ASSETREGISTRY_API bool Contains(const FCachedAssetKey& Key) const;
 	ASSETREGISTRY_API FAssetData* const* Find(const FCachedAssetKey& Key) const;
@@ -325,7 +328,7 @@ public:
 
 	ASSETREGISTRY_API void Empty(int32 ReservedSize=0);
 	ASSETREGISTRY_API void Add(FName PackageName, FAssetDataPtrIndex AssetIndex);
-	ASSETREGISTRY_API void Remove(FName PackageName, FAssetDataPtrIndex AssetIndex);
+	ASSETREGISTRY_API int32 Remove(FName PackageName, FAssetDataPtrIndex AssetIndex);
 	ASSETREGISTRY_API void Shrink();
 
 	ASSETREGISTRY_API int32 Num() const;
