[AssetRegistry]

- Remove locking around expensive gatherer ScanPathsSynchronous and WaitOnPath calls
- To do the above, the appropiate ScopeLock is passed from callers to make lock ownership clear (and allow the unlocking)
- Adds private templated helper method for GetRedirectedObjectPath to handle the read and write lock cases with a single implementation
- Changed public GetRedirectedObjectPath to only be used from under a read lock and removed unnecessary arguments from the API that are only used in the writing case
- Added GetAndScanRedirectedObjectPath which handles GetRedirectedObjectPath when a write lock is held.
- Removed unnecessary UAssetRegistryImpl argument from FAsyncConsumer::Wait and FAsyncConsumer::ConditionalLoadPremadeAssetRegistry
- Added new Event to Gatherer for processing gather results, removing the AssetsFound and VerseFilesFound callbacks from FTickContext
- Extended FClassInheritanceContext to allow handle re-entering the Interfacelock when it is intermittently released

#rnx

diff --git a/Engine/Source/Runtime/AssetRegistry/Private/AssetDataGatherer.cpp b/Engine/Source/Runtime/AssetRegistry/Private/AssetDataGatherer.cpp
--- a/Engine/Source/Runtime/AssetRegistry/Private/AssetDataGatherer.cpp
+++ b/Engine/Source/Runtime/AssetRegistry/Private/AssetDataGatherer.cpp
@@ -4782,33 +4782,40 @@ void FAssetDataGatherer::GetAndTrimSearchResults(UE::AssetDataGather::FResults&
 	// GetPackageResults takes its own lock.
 	GetPackageResults(InOutResults);
 
-	FGathererScopeLock ResultsScopeLock(&ResultsLock);
+	{
+		FGathererScopeLock ResultsScopeLock(&ResultsLock);
+
+		MoveAppendRangeToRingBuffer(InOutResults.Paths, DiscoveredPaths);
 
-	MoveAppendRangeToRingBuffer(InOutResults.Paths, DiscoveredPaths);
 
+		MoveAppendRangeToRingBuffer(InOutResults.CookedPackageNamesWithoutAssetData, CookedPackageNamesWithoutAssetDataResults);
+		MoveAppendRangeToRingBuffer(InOutResults.VerseFiles, VerseResults);
 
-	MoveAppendRangeToRingBuffer(InOutResults.CookedPackageNamesWithoutAssetData, CookedPackageNamesWithoutAssetDataResults);
-	MoveAppendRangeToRingBuffer(InOutResults.VerseFiles, VerseResults);
+		InOutResults.BlockedFiles.Append(MoveTemp(BlockedResults));
+		BlockedResults.Reset();
 
-	InOutResults.BlockedFiles.Append(MoveTemp(BlockedResults));
-	BlockedResults.Reset();
+		OutContext.SearchTimes.Append(MoveTemp(SearchTimes));
+		SearchTimes.Reset();
 
-	OutContext.SearchTimes.Append(MoveTemp(SearchTimes));
-	SearchTimes.Reset();
+		OutContext.NumFilesToSearch = FilesToSearch->Num();
+		OutContext.NumPathsToSearch = NumPathsToSearchAtLastSyncPoint;
+		OutContext.bIsDiscoveringFiles = !bDiscoveryIsComplete;
 
-	OutContext.NumFilesToSearch = FilesToSearch->Num();
-	OutContext.NumPathsToSearch = NumPathsToSearchAtLastSyncPoint;
-	OutContext.bIsDiscoveringFiles = !bDiscoveryIsComplete;
+		// Idle means no more work OR we are blocked on external events, but complete means no more work period.
+		bool bLocalIsComplete = bIsIdle && FilesToSearch->Num() == 0;
+		if (bLocalIsComplete && !bIsComplete)
+		{
+			bIsComplete = true;
+			Shrink();
+		}
+		OutContext.bIsSearching = !bLocalIsComplete;
+		OutContext.bAbleToProgress = !bIsIdle;
+	}
 
-	// Idle means no more work OR we are blocked on external events, but complete means no more work period.
-	bool bLocalIsComplete = bIsIdle && FilesToSearch->Num() == 0;
-	if (bLocalIsComplete && !bIsComplete)
+	if (GatheredResultsEvent.IsBound())
 	{
-		bIsComplete = true;
-		Shrink();
+		GatheredResultsEvent.Broadcast(InOutResults);
 	}
-	OutContext.bIsSearching = !bLocalIsComplete;
-	OutContext.bAbleToProgress = !bIsIdle;
 }
 
 FAssetGatherDiagnostics FAssetDataGatherer::GetDiagnostics()
diff --git a/Engine/Source/Runtime/AssetRegistry/Private/AssetDataGatherer.h b/Engine/Source/Runtime/AssetRegistry/Private/AssetDataGatherer.h
--- a/Engine/Source/Runtime/AssetRegistry/Private/AssetDataGatherer.h
+++ b/Engine/Source/Runtime/AssetRegistry/Private/AssetDataGatherer.h
@@ -29,19 +29,16 @@ struct FAssetData;
 class FDiskCachedAssetData;
 class FAssetRegistryReader;
 class FAssetRegistryWriter; // Not defined if !ALLOW_NAME_BATCH_SAVING
+namespace UE::AssetDataGather { struct FResults; }
+namespace UE::AssetDataGather::Private { class FAssetDataDiscovery; }
+namespace UE::AssetDataGather::Private { class FFilesToSearch; }
+namespace UE::AssetDataGather::Private { struct FCachePayload; }
+namespace UE::AssetDataGather::Private { struct FGatheredPathData; }
+namespace UE::AssetDataGather::Private { struct FPathExistence; }
+namespace UE::AssetDataGather::Private { struct FSetPathProperties; }
+namespace UE::AssetDataGather::Private { enum class EPriority : uint8; }
 namespace UE::AssetRegistry { class FAssetRegistryImpl; }
 
-namespace UE::AssetDataGather::Private
-{
-class FAssetDataDiscovery;
-class FFilesToSearch;
-struct FCachePayload;
-struct FGatheredPathData;
-struct FPathExistence;
-struct FSetPathProperties;
-enum class EPriority : uint8;
-}
-
 #if DO_CHECK
 typedef FCriticalSectionQueryable FGathererCriticalSection;
 typedef FScopeLockQueryable FGathererScopeLock;
@@ -54,6 +51,8 @@ typedef FScopeLock FGathererScopeLock;
 #define CHECK_IS_NOT_LOCKED_CURRENT_THREAD(CritSec) do {} while (false)
 #endif
 
+DECLARE_TS_MULTICAST_DELEGATE_OneParam(FGatheredResultsEvent, const UE::AssetDataGather::FResults&);
+
 struct FAssetGatherDiagnostics
 {
 	/** Time spent identifying asset files on disk */
@@ -157,6 +156,8 @@ public:
 	/** Signals to end the thread and waits for it to close before returning */
 	void EnsureCompletion();
 
+	/** Returns reference to delegate invoked when search results have been gathered in GetAndTrimSearchResults. */
+	FGatheredResultsEvent& GetGatheredResultsEvent();
 
 	/** Gets search results from the data gatherer. */
 	void GetAndTrimSearchResults(UE::AssetDataGather::FResults& InOutResults,
@@ -401,6 +402,9 @@ private:
 	 */
 	mutable FGathererCriticalSection ResultsLock;
 
+	/** Delegates invoked when results are gathered in GetAndTrimSearchResults */
+	FGatheredResultsEvent GatheredResultsEvent;
+
 	// Variable section for variables that are constant during threading.
 	UAssetRegistryImpl& AssetRegistry;
 
@@ -551,6 +555,11 @@ inline void FAssetDataGatherer::RequestAsyncCacheSave()
 	bSaveAsyncCacheTriggered = true;
 }
 
+inline FGatheredResultsEvent& FAssetDataGatherer::GetGatheredResultsEvent()
+{
+	return GatheredResultsEvent;
+}
+
 inline void FAssetDataGatherer::SetGatherOnGameThreadOnly(bool bValue)
 {
 	bGatherOnGameThreadOnly.store(bValue, std::memory_order_relaxed);
diff --git a/Engine/Source/Runtime/AssetRegistry/Private/AssetRegistry.cpp b/Engine/Source/Runtime/AssetRegistry/Private/AssetRegistry.cpp
--- a/Engine/Source/Runtime/AssetRegistry/Private/AssetRegistry.cpp
+++ b/Engine/Source/Runtime/AssetRegistry/Private/AssetRegistry.cpp
@@ -78,6 +78,11 @@
  * const functions require a ReadLock critical section; non-const require a WriteLock.
  * The requirement that functions must be called only from within a critical section(and non-const only within a
  * WriteLock) is not enforced technically; change authors need to carefully follow the synchronization model.
+ * FAssetRegistryImpl methods that expect an FInterfaceScope[Read|Write]Lock& may release and re-acquire the lock. Functions
+ * re-entering the lock are expected to validate possible re-entrancy issues that may arise from doing so.
+ * FAssetRegistryImpl methods that expect an FInterfaceScope[Read|Write]Lock* work like above however passing in nullptr
+ * indicates that the function should maintain the lock state from the caller and the caller is responsible for ensuring
+ * the threading model is being adhered to.
 
  * *** Events, Callbacks, and Object Virtuals ***
  * The AssetRegistry provides several Events(e.g.AssetAddedEvent) that can be subscribed to from arbitrary engine or
@@ -649,7 +654,7 @@ void FAsyncConsumer::PrepareForConsume()
 	++ReferenceCount;
 };
 
-void FAsyncConsumer::Wait(UAssetRegistryImpl& UARI, UE::AssetRegistry::FInterfaceWriteScopeLock& ScopeLock)
+void FAsyncConsumer::Wait(UE::AssetRegistry::FInterfaceWriteScopeLock& ScopeLock)
 {
 	// Called within the lock
 	if (ReferenceCount == 0)
@@ -660,8 +665,8 @@ void FAsyncConsumer::Wait(UAssetRegistryImpl& UARI, UE::AssetRegistry::FInterfac
 
 	// Wait outside of the lock so that the AsyncThread can enter the lock to call Consume
 	{
-		UARI.InterfaceLock.WriteUnlock();
-		ON_SCOPE_EXIT{ UARI.InterfaceLock.WriteLock(); };
+		ScopeLock.Lock.WriteUnlock();
+		ON_SCOPE_EXIT{ ScopeLock.Lock.WriteLock(); };
 		check(Consumed != nullptr);
 		Consumed->Wait();
 	}
@@ -733,9 +738,9 @@ struct FGatherDependencyDataScratchpad
 
 }
 
-void FAssetRegistryImpl::ConditionalLoadPremadeAssetRegistry(UAssetRegistryImpl& UARI, Impl::FEventContext& EventContext, UE::AssetRegistry::FInterfaceWriteScopeLock& ScopeLock)
+void FAssetRegistryImpl::ConditionalLoadPremadeAssetRegistry(UE::AssetRegistry::FInterfaceWriteScopeLock& ScopeLock, Impl::FEventContext& EventContext)
 {
-	AsyncConsumer.Wait(UARI, ScopeLock);
+	AsyncConsumer.Wait(ScopeLock);
 }
 
 void FAssetRegistryImpl::ConsumeOrDeferPreloadedPremade(UAssetRegistryImpl& UARI, Impl::FEventContext& EventContext)
@@ -884,7 +889,7 @@ UAssetRegistryImpl::UAssetRegistryImpl(const FObjectInitializer& ObjectInitializ
 		GetInheritanceContextWithRequiredLock(InterfaceScopeLock, Context.InheritanceContext,
 			Context.InheritanceBuffer);
 
-		GuardedData.Initialize(Context);
+		GuardedData.Initialize(InterfaceScopeLock, Context);
 		InitializeEvents(Context);
 	}
 	Broadcast(Context.Events);
@@ -991,11 +996,15 @@ void FAssetRegistryImpl::LoadPremadeAssetRegistry(Impl::FEventContext& EventCont
 	bPreloadingComplete = true;
 }
 
-void FAssetRegistryImpl::Initialize(Impl::FInitializeContext& Context)
+void FAssetRegistryImpl::Initialize(UE::AssetRegistry::FInterfaceWriteScopeLock& ScopeLock, Impl::FInitializeContext& Context)
 {
+	if (bInitializationComplete)
+	{
+		return;
+	}
+
 	const double StartupStartTime = FPlatformTime::Seconds();
 
-	bInitialSearchStarted = false;
 	bInitialSearchCompleted.store(true, std::memory_order_relaxed);
 	bAdditionalMountSearchInProgress.store(false, std::memory_order_relaxed);
 #if WITH_EDITOR
@@ -1109,6 +1118,8 @@ void FAssetRegistryImpl::Initialize(Impl::FInitializeContext& Context)
 	// Make sure first call to LoadCalculatedDependencies builds the Gatherer list. At that point Classes should be loaded.
 	bRegisteredDependencyGathererClassesDirty = true;
 #endif
+
+	bInitializationComplete = true;
 }
 
 }
@@ -1358,7 +1369,7 @@ void FAssetRegistryImpl::InitRedirectors(Impl::FEventContext& EventContext,
 			PathsToSearch.Add(RootPackageName);
 
 			Impl::FScanPathContext Context(EventContext, InheritanceContext, PathsToSearch, TArray<FString>());
-			ScanPathsSynchronous(Context);
+			ScanPathsSynchronous(nullptr /*ScopeLock*/, Context);
 		}
 
 		FName PluginPackageName = FName(*FString::Printf(TEXT("/%s/"), *Plugin->GetName()));
@@ -1992,7 +2003,7 @@ void UAssetRegistryImpl::SearchAllAssets(bool bSynchronousSearch)
 		if (bSynchronousSearch)
 		{
 			// make sure any outstanding async preload is complete
-			GuardedData.ConditionalLoadPremadeAssetRegistry(*this, EventContext, InterfaceScopeLock);
+			GuardedData.ConditionalLoadPremadeAssetRegistry(InterfaceScopeLock, EventContext);
 		}
 		GuardedData.SearchAllAssets(EventContext, InheritanceContext, bSynchronousSearch);
 	}
@@ -2226,7 +2237,7 @@ void UAssetRegistryImpl::WaitForPremadeAssetRegistry()
 		FClassInheritanceContext InheritanceContext;
 		FClassInheritanceBuffer InheritanceBuffer;
 		GetInheritanceContextWithRequiredLock(InterfaceScopeLock, InheritanceContext, InheritanceBuffer);
-		GuardedData.ConditionalLoadPremadeAssetRegistry(*this, EventContext, InterfaceScopeLock);
+		GuardedData.ConditionalLoadPremadeAssetRegistry(InterfaceScopeLock, EventContext);
 	}
 	Broadcast(EventContext);
 }
@@ -2273,7 +2284,7 @@ void UAssetRegistryImpl::WaitForPackage(const FString& PackageName)
 	{
 		LLM_SCOPE(ELLMTag::AssetRegistry);
 		UE::AssetRegistry::FInterfaceWriteScopeLock InterfaceScopeLock(InterfaceLock);
-		GuardedData.TickGatherPackage(EventContext, PackageName, LocalPath);
+		GuardedData.TickGatherPackage(InterfaceScopeLock, EventContext, PackageName, LocalPath);
 	}
 	Broadcast(EventContext);
 }
@@ -3400,10 +3411,10 @@ FSoftObjectPath UAssetRegistryImpl::GetRedirectedObjectPath(const FSoftObjectPat
 	// In that case, we can skip further scanning while looking for a redirected path
 	if (!GuardedData.IsGathering())
 	{
-		UE::AssetRegistry::FInterfaceReadScopeLock InterfaceScopeLock(InterfaceLock);
+		UE::AssetRegistry::FInterfaceReadScopeLock ReadScopeLock(InterfaceLock);
 		if (GuardedData.IsSearchAllAssets())
 		{
-			return GuardedData.GetRedirectedObjectPath(ObjectPath, nullptr, nullptr, /*bNeedsScanning*/ false);
+			return GuardedData.GetRedirectedObjectPath(ReadScopeLock, ObjectPath);
 		}
 	}
 
@@ -3415,7 +3426,7 @@ FSoftObjectPath UAssetRegistryImpl::GetRedirectedObjectPath(const FSoftObjectPat
 		LLM_SCOPE(ELLMTag::AssetRegistry);
 		UE::AssetRegistry::FInterfaceWriteScopeLock WriteScopeLock(InterfaceLock);
 		GetInheritanceContextWithRequiredLock(WriteScopeLock, InheritanceContext, InheritanceBuffer);
-		RedirectedObjectPath = GuardedData.GetRedirectedObjectPath(ObjectPath, &EventContext, &InheritanceContext, /*bNeedsScanning*/ true);
+		RedirectedObjectPath = GuardedData.GetAndScanRedirectedObjectPath(WriteScopeLock, ObjectPath, &EventContext, &InheritanceContext);
 	}	
 	Broadcast(EventContext);
 
@@ -3425,8 +3436,23 @@ FSoftObjectPath UAssetRegistryImpl::GetRedirectedObjectPath(const FSoftObjectPat
 namespace UE::AssetRegistry
 {
 
-FSoftObjectPath FAssetRegistryImpl::GetRedirectedObjectPath(const FSoftObjectPath& ObjectPath, UE::AssetRegistry::Impl::FEventContext* EventContext, UE::AssetRegistry::Impl::FClassInheritanceContext* InheritanceContext, bool bNeedsScanning)
+FSoftObjectPath FAssetRegistryImpl::GetRedirectedObjectPath(FInterfaceReadScopeLock& ScopeLock, const FSoftObjectPath& ObjectPath)
 {
+	return GetRedirectedObjectPath(ScopeLock, ObjectPath, nullptr, nullptr);
+}
+
+FSoftObjectPath FAssetRegistryImpl::GetAndScanRedirectedObjectPath(FInterfaceWriteScopeLock& ScopeLock, const FSoftObjectPath& ObjectPath,
+	UE::AssetRegistry::Impl::FEventContext* EventContext, UE::AssetRegistry::Impl::FClassInheritanceContext* InheritanceContext)
+{
+	return GetRedirectedObjectPath(ScopeLock, ObjectPath, EventContext, InheritanceContext);
+}
+
+template <typename TScopeLock>
+FSoftObjectPath FAssetRegistryImpl::GetRedirectedObjectPath(TScopeLock& ScopeLock, const FSoftObjectPath& ObjectPath, 
+	UE::AssetRegistry::Impl::FEventContext* EventContext, UE::AssetRegistry::Impl::FClassInheritanceContext* InheritanceContext)
+{
+	static_assert(std::is_same_v<TScopeLock, FInterfaceWriteScopeLock> || std::is_same_v<TScopeLock, FInterfaceReadScopeLock>, "Only expect scope lock types of either FInterfaceWriteScopeLock or FInterfaceReadScopeLock");
+	constexpr bool bNeedsScanning = std::is_same_v<TScopeLock, FInterfaceWriteScopeLock>;
 	check(!bNeedsScanning || (EventContext && InheritanceContext));
 
 	FSoftObjectPath RedirectedPath = ObjectPath;
@@ -3463,13 +3489,16 @@ FSoftObjectPath FAssetRegistryImpl::GetRedirectedObjectPath(const FSoftObjectPat
 
 	const FAssetData* AssetData = RetrieveAssetData();
 
-	if (!AssetData && bNeedsScanning)
+	if constexpr (bNeedsScanning)
 	{
-		UE::AssetRegistry::Impl::FScanPathContext Context(*EventContext, *InheritanceContext, {}, { RedirectedPath.ToString() },
-			UE::AssetRegistry::EScanFlags::IgnoreInvalidPathWarning);
-		ScanPathsSynchronous(Context);
+		if (!AssetData)
+		{
+			UE::AssetRegistry::Impl::FScanPathContext Context(*EventContext, *InheritanceContext, {}, { RedirectedPath.ToString() },
+				UE::AssetRegistry::EScanFlags::IgnoreInvalidPathWarning);
+			ScanPathsSynchronous(&ScopeLock, Context);
 
-		AssetData = RetrieveAssetData();
+			AssetData = RetrieveAssetData();
+		}
 	}
 
 	// Most of the time this will either not be a redirector or only have one redirect, so optimize for that case
@@ -3495,12 +3524,15 @@ FSoftObjectPath FAssetRegistryImpl::GetRedirectedObjectPath(const FSoftObjectPat
 		}
 
 		AssetData = State.GetAssetByObjectPath(RedirectedPath);
-		if (!AssetData && bNeedsScanning)
+		if constexpr (bNeedsScanning)
 		{
-			UE::AssetRegistry::Impl::FScanPathContext Context(*EventContext, *InheritanceContext, {}, { RedirectedPath.ToString() });
-			ScanPathsSynchronous(Context);
+			if (!AssetData)
+			{
+				UE::AssetRegistry::Impl::FScanPathContext Context(*EventContext, *InheritanceContext, {}, { RedirectedPath.ToString() });
+				ScanPathsSynchronous(&ScopeLock, Context);
 
-			AssetData = State.GetAssetByObjectPath(RedirectedPath);
+				AssetData = State.GetAssetByObjectPath(RedirectedPath);
+			}
 		}
 
 		SeenPaths.Add(RedirectedPath);
@@ -4279,8 +4311,8 @@ void UAssetRegistryImpl::ScanPathsSynchronousInternal(const TArray<FString>& InD
 		bInitialSearchCompleted = GuardedData.IsInitialSearchCompleted();
 		bAdditionalMountSearchInProgress = GuardedData.IsAdditionalMountSearchInProgress();
 		// make sure any outstanding async preload is complete
-		GuardedData.ConditionalLoadPremadeAssetRegistry(*this, EventContext, InterfaceScopeLock);
-		GuardedData.ScanPathsSynchronous(Context);
+		GuardedData.ConditionalLoadPremadeAssetRegistry(InterfaceScopeLock, EventContext);
+		GuardedData.ScanPathsSynchronous(&InterfaceScopeLock, Context);
 	}
 	if (Context.LocalPaths.IsEmpty())
 	{
@@ -5140,10 +5172,6 @@ Impl::EGatherStatus FAssetRegistryImpl::TickGatherer(Impl::FTickContext& TickCon
 	if (BackgroundResults.VerseFiles.Num())
 	{
 		TickContext.LazyStartTimer();
-		if (TickContext.VerseFilesFoundCallback.IsSet())
-		{
-			TickContext.VerseFilesFoundCallback.GetValue()(BackgroundResults.VerseFiles);
-		}
 
 		VerseFilesGathered(EventContext, BackgroundResults.VerseFiles, InOutInterruptionContext);
 		if (InOutInterruptionContext.ShouldExitEarly())
@@ -5426,16 +5454,6 @@ void FTickContext::RunAssetSearchDataGathered(TMultiMap<FName, TUniquePtr<FAsset
 	if (InAssetResults.Num())
 	{
 		LazyStartTimer();
-		// Mark the first amortize time
-		if (AssetsFoundCallback.IsSet())
-		{
-			TMultiMap<FName, FAssetData*> NonOwningContainer;
-			for (auto Iter = InAssetResults.CreateIterator(); Iter; ++Iter)
-			{
-				NonOwningContainer.Add(Iter.Key(), Iter.Value().Get());
-			}
-			AssetsFoundCallback.GetValue()(NonOwningContainer);
-		}
 
 		GuardedData.AssetSearchDataGathered(EventContext, InAssetResults, OutDeferredAssetResults,
 			InterruptionContext, MountPointsForVerifyAfterGather);
@@ -5656,13 +5674,19 @@ void FAssetRegistryImpl::LogSearchDiagnostics(double StartTime)
 #endif // !NO_LOGGING
 }
 
-void FAssetRegistryImpl::TickGatherPackage(Impl::FEventContext& EventContext, const FString& PackageName, const FString& LocalPath)
+void FAssetRegistryImpl::TickGatherPackage(UE::AssetRegistry::FInterfaceWriteScopeLock& ScopeLock, Impl::FEventContext& EventContext, const FString& PackageName, const FString& LocalPath)
 {
 	if (!GlobalGatherer.IsValid())
 	{
 		return;
 	}
+
+	// Release the lock while performing expensive calls that don't need the lock
+	ScopeLock.Lock.WriteUnlock();
 	GlobalGatherer->WaitOnPath(LocalPath);
+	ScopeLock.Lock.WriteLock();
+	check(GlobalGatherer.IsValid());
+
 	double TimingStartTime = -1.;
 	auto LazyStartTimer = [&TimingStartTime]()
 	{
@@ -6369,8 +6393,10 @@ FScanPathContext::FScanPathContext(FEventContext& InEventContext, FClassInherita
 
 }
 
-void FAssetRegistryImpl::ScanPathsSynchronous(Impl::FScanPathContext& Context)
+void FAssetRegistryImpl::ScanPathsSynchronous(FInterfaceWriteScopeLock* ScopeLock, Impl::FScanPathContext& Context)
 {
+	// We always expect a valid lock to be provided except before initialization is complete.
+	check(ScopeLock || bInitializationComplete);
 	LLM_SCOPE(ELLMTag::AssetRegistry);
 
 	if (!TryConstructGathererIfNeeded())
@@ -6388,13 +6414,87 @@ void FAssetRegistryImpl::ScanPathsSynchronous(Impl::FScanPathContext& Context)
 	}
 	Gatherer.AddRequiredMountPoints(Context.LocalPaths);
 
+	auto IsInRequestedDir = [&Context](const FAssetData& AssetData)
+		{
+			TStringBuilder<128> PackageNameStr;
+			AssetData.PackageName.ToString(PackageNameStr);
+			FStringView PackageName(PackageNameStr.ToString(), PackageNameStr.Len());
+
+			for (const FString& RequestedPackageDir : Context.PackageDirs)
+			{
+				if (FPathViews::IsParentPathOf(RequestedPackageDir, PackageName))
+				{
+					return true;
+				}
+			}
+			return false;
+		};
+
+	TSet<FSoftObjectPath> GatheredAssets;
+	TSet<FName> GatheredVerseFiles;
+	FDelegateHandle GatheredResultsHandle;
+	if (Context.bForceRescan)
+	{
+		GatheredResultsHandle = Gatherer.GetGatheredResultsEvent().AddLambda(
+			[&GatheredAssets, &GatheredVerseFiles, &IsInRequestedDir, OutFoundAssets = Context.OutFoundAssets](const UE::AssetDataGather::FResults& Results)
+			{
+				auto AddGatheredAssets = [&IsInRequestedDir, OutFoundAssets](TSet<FSoftObjectPath>& InOutGatheredAssets, const TMultiMap<FName, TUniquePtr<FAssetData>>& AssetMap)
+					{
+						for (TMultiMap<FName, TUniquePtr<FAssetData>>::TConstIterator It = AssetMap.CreateConstIterator(); It; ++It)
+						{
+							const FAssetData* AssetData = It.Value().Get();
+							if (OutFoundAssets)
+							{
+								// The gatherer may have added other assets that were scanned as part of the ongoing background scan,
+								// so remove any assets that were not in the requested paths
+								if (IsInRequestedDir(*AssetData))
+								{
+									UE_LOG(LogAssetRegistry, VeryVerbose, TEXT("FAssetRegistryImpl::ScanPathsSynchronous: Found Asset: %s"),
+										*AssetData->GetObjectPathString());
+									OutFoundAssets->Add(AssetData->GetSoftObjectPath());
+								}
+							}
+
+							InOutGatheredAssets.Add(AssetData->GetSoftObjectPath());
+						}
+					};
+				AddGatheredAssets(GatheredAssets, Results.Assets);
+				AddGatheredAssets(GatheredAssets, Results.AssetsForGameThread);
+
+				for (FName VerseFile : Results.VerseFiles)
+				{
+					GatheredVerseFiles.Add(VerseFile);
+				}
+			});
+	}
+	ON_SCOPE_EXIT
+	{
+		if (Context.bForceRescan)
+		{
+			Gatherer.GetGatheredResultsEvent().Remove(GatheredResultsHandle);
+		}
+	};
+
+	// Conditionally release the lock while performing expensive calls that don't need the lock. 
+	// If no lock was passed we assume the caller does not want us to unlock.
+	if (ScopeLock)
+	{
+		ScopeLock->Lock.WriteUnlock();
+	}
+	Gatherer.ScanPathsSynchronous(Context.LocalPaths, Context.bForceRescan, Context.bIgnoreDenyListScanFilters);
+	if (ScopeLock)
+	{
+		ScopeLock->Lock.WriteLock();
+		Context.InheritanceContext.OnLockReentered(*this);
+	}
+
 	// If we are forcing a rescan, then delete any old assets that no longer exist. If we are not forcing a rescan,
 	// then there should not be any old assets that no longer exist, so we skip the cost of searching for them.
-	TSet<FSoftObjectPath> OldAssetsToRemove;
-	TSet<FName> OldVerseFilesToRemove;
+	TSet<FSoftObjectPath> ExistingAssets;
+	TSet<FName> ExistingVerseFiles;
 	if (Context.bForceRescan)
 	{
-		// Initialize OldAssetsToRemove to the list of all assets in the given paths.
+		// Initialize ExistingAssets to the list of all assets in the given paths.
 		if (!Context.PackageDirs.IsEmpty())
 		{
 			FARFilter Filter;
@@ -6408,9 +6508,9 @@ void FAssetRegistryImpl::ScanPathsSynchronous(Impl::FScanPathContext& Context)
 			CompileFilter(Context.InheritanceContext, Filter, CompiledFilter);
 			TArray<FAssetData> AssetsInPaths;
 			State.EnumerateAssets(CompiledFilter, TSet<FName>() /* PackageNamesToSkip */,
-				[&OldAssetsToRemove](const FAssetData& AssetData)
+				[&ExistingAssets](const FAssetData& AssetData)
 				{
-					OldAssetsToRemove.Add(AssetData.ToSoftObjectPath());
+					ExistingAssets.Add(AssetData.ToSoftObjectPath());
 					return true;
 				}, UE::AssetRegistry::EEnumerateAssetsFlags::AllowUnfilteredArAssets);
 			for (FName PackagePath : CompiledFilter.PackagePaths)
@@ -6418,16 +6518,16 @@ void FAssetRegistryImpl::ScanPathsSynchronous(Impl::FScanPathContext& Context)
 				TArray<FName>* VerseFiles = CachedVerseFilesByPath.Find(PackagePath);
 				if (VerseFiles)
 				{
-					OldVerseFilesToRemove.Append(*VerseFiles);
+					ExistingVerseFiles.Append(*VerseFiles);
 				}
 			}
 		}
 		for (const FString& PackageName : Context.PackageFiles)
 		{
 			State.EnumerateAssetsByPackageName(FName(*PackageName),
-				[&OldAssetsToRemove](const FAssetData* AssetData)
+				[&ExistingAssets](const FAssetData* AssetData)
 				{
-					OldAssetsToRemove.Add(AssetData->ToSoftObjectPath());
+					ExistingAssets.Add(AssetData->ToSoftObjectPath());
 					return true;
 				});
 			for (const TCHAR* Extension : FAssetDataGatherer::GetVerseFileExtensions())
@@ -6435,96 +6535,13 @@ void FAssetRegistryImpl::ScanPathsSynchronous(Impl::FScanPathContext& Context)
 				FName VerseName(*WriteToString<256>(PackageName, Extension), FNAME_Find);
 				if (!VerseName.IsNone() && CachedVerseFiles.Contains(VerseName))
 				{
-					OldVerseFilesToRemove.Add(VerseName);
+					ExistingVerseFiles.Add(VerseName);
 				}
 			}
 		}
 	}
 
-	Gatherer.ScanPathsSynchronous(Context.LocalPaths, Context.bForceRescan, Context.bIgnoreDenyListScanFilters);
-	TArray<FName> FoundAssetPackageNames;
-
-	auto IsInRequestedDir = [&Context](const FAssetData& AssetData)
-	{
-		TStringBuilder<128> PackageNameStr;
-		AssetData.PackageName.ToString(PackageNameStr);
-		FStringView PackageName(PackageNameStr.ToString(), PackageNameStr.Len());
-
-		for (const FString& RequestedPackageDir : Context.PackageDirs)
-		{
-			if (FPathViews::IsParentPathOf(RequestedPackageDir, PackageName))
-			{
-				return true;
-			}
-		}
-		return false;
-	};
-	auto AssetsFoundCallback =
-		[&IsInRequestedDir, &Context, &FoundAssetPackageNames, &OldAssetsToRemove, this]
-		(const TMultiMap<FName, FAssetData*>& InFoundAssets)
-	{
-		Context.NumFoundAssets = InFoundAssets.Num();
-
-		FoundAssetPackageNames.Reserve(FoundAssetPackageNames.Num() + Context.NumFoundAssets);
-
-		// The gatherer may have added other assets that were scanned as part of the ongoing background scan,
-		// so remove any assets that were not in the requested paths
-		for (const TPair<FName,FAssetData*>& Pair : InFoundAssets)
-		{
-			FAssetData* AssetData = Pair.Value;
-			bool bIsInRequestedPaths = false;
-
-			TStringBuilder<128> PackageNameStr;
-			AssetData->PackageName.ToString(PackageNameStr);
-			FStringView PackageName(PackageNameStr.ToString(), PackageNameStr.Len());
-
-			bIsInRequestedPaths = IsInRequestedDir(*AssetData);
-
-			if (!bIsInRequestedPaths)
-			{
-				for (const FString& RequestedPackageFile : Context.PackageFiles)
-				{
-					if (PackageName.Equals(RequestedPackageFile, ESearchCase::IgnoreCase))
-					{
-						bIsInRequestedPaths = true;
-						break;
-					}
-				}
-			}
-
-			if (bIsInRequestedPaths)
-			{
-				UE_LOG(LogAssetRegistry, VeryVerbose, TEXT("FAssetRegistryImpl::ScanPathsSynchronous: Found Asset: %s"),
-					*AssetData->GetObjectPathString());
-				if (Context.OutFoundAssets)
-				{
-					Context.OutFoundAssets->Add(AssetData->GetSoftObjectPath());
-				}
-				FoundAssetPackageNames.Add(AssetData->PackageName);
-			}
-
-			if (!OldAssetsToRemove.IsEmpty())
-			{
-				OldAssetsToRemove.Remove(AssetData->ToSoftObjectPath());
-			}
-		}
-	};
-	auto VerseFileFoundCallback =
-		[&OldVerseFilesToRemove]
-		(const TRingBuffer<FName>& InFoundVerseFiles)
-		{
-			if (!OldVerseFilesToRemove.IsEmpty())
-			{
-				for (const FName VerseFile : InFoundVerseFiles)
-				{
-					OldVerseFilesToRemove.Remove(VerseFile);
-				}
-			}
-		};
-
 	Impl::FTickContext TickContext(*this, Context.EventContext, Context.InheritanceContext);
-	TickContext.AssetsFoundCallback = Impl::FAssetsFoundCallback(AssetsFoundCallback);
-	TickContext.VerseFilesFoundCallback = Impl::FVerseFilesFoundCallback(VerseFileFoundCallback);
 	Context.Status = TickGatherer(TickContext);
 
 	// Temporary hack/partial solution. The expectation is that this function will return cause all assets
@@ -6538,7 +6555,7 @@ void FAssetRegistryImpl::ScanPathsSynchronous(Impl::FScanPathContext& Context)
 	// we manually identify any deferred assets that fall under the paths we are scanning and ask the asset registry
 	// to process them ignoring any failures of TryPostLoadAssetRegistryTags. We then run a second full Tick to finish
 	// out their processing. See UE-210249 for the desired fix.
-
+	TArray<FName> FoundAssetPackageNames;
 	{
 		// Find any assets that were deferred but fall into the paths we are interested in. Extract them from 
 		// the DeferredAssets and DeferredAssetsForGameThread containers
@@ -6582,7 +6599,6 @@ void FAssetRegistryImpl::ScanPathsSynchronous(Impl::FScanPathContext& Context)
 		bForceCompletionEvenIfPostLoadsFail = bOldForceCompletionEvenIfPostLoadsFail;
 		// Tick to perform any subsequent processing required for these assets beyond AssetSearchDataGathered
 		Impl::FTickContext AssetTickContext(*this, Context.EventContext, Context.InheritanceContext);
-		AssetTickContext.AssetsFoundCallback = Impl::FAssetsFoundCallback(AssetsFoundCallback);
 		Context.Status = TickGatherer(AssetTickContext);
 	}
 	FoundAssetPackageNames.Sort(FNameFastLess());
@@ -6594,6 +6610,8 @@ void FAssetRegistryImpl::ScanPathsSynchronous(Impl::FScanPathContext& Context)
 	LoadCalculatedDependencies(&FoundAssetPackageNames, Context.InheritanceContext,
 		&PackagesNeedingDependencyCalculationOnGameThread, TickContext.InterruptionContext);
 #endif
+
+	TSet<FSoftObjectPath> OldAssetsToRemove = ExistingAssets.Difference(GatheredAssets);
 	for (FSoftObjectPath& OldAssetToRemove : OldAssetsToRemove)
 	{
 		FAssetData* AssetDataToRemove = State.GetMutableAssetByObjectPath(OldAssetToRemove);
@@ -6602,6 +6620,8 @@ void FAssetRegistryImpl::ScanPathsSynchronous(Impl::FScanPathContext& Context)
 			RemoveAssetData(Context.EventContext, AssetDataToRemove);
 		}
 	}
+	
+	TSet<FName> OldVerseFilesToRemove = ExistingVerseFiles.Difference(GatheredVerseFiles);
 	for (FName OldVerseFileToRemove : OldVerseFilesToRemove)
 	{
 		RemoveVerseFile(Context.EventContext, OldVerseFileToRemove);
@@ -7965,7 +7985,7 @@ void UAssetRegistryImpl::OnDirectoryChanged(const TArray<FFileChangeData>& FileC
 		UE::AssetRegistry::Impl::FClassInheritanceContext InheritanceContext;
 		UE::AssetRegistry::Impl::FClassInheritanceBuffer InheritanceBuffer;
 		GetInheritanceContextWithRequiredLock(InterfaceScopeLock, InheritanceContext, InheritanceBuffer);
-		GuardedData.OnDirectoryChanged(EventContext, InheritanceContext, FileChangesProcessed);
+		GuardedData.OnDirectoryChanged(InterfaceScopeLock, EventContext, InheritanceContext, FileChangesProcessed);
 	}
 	Broadcast(EventContext);
 	
@@ -7981,7 +8001,7 @@ void UAssetRegistryImpl::OnDirectoryChanged(const TArray<FFileChangeData>& FileC
 namespace UE::AssetRegistry
 {
 
-void FAssetRegistryImpl::OnDirectoryChanged(Impl::FEventContext& EventContext,
+void FAssetRegistryImpl::OnDirectoryChanged(FInterfaceWriteScopeLock& ScopeLock, Impl::FEventContext& EventContext,
 	Impl::FClassInheritanceContext& InheritanceContext, TArray<FFileChangeData>& FileChangesProcessed)
 {
 	TArray<FString> NewDirs;
@@ -7998,7 +8018,7 @@ void FAssetRegistryImpl::OnDirectoryChanged(Impl::FEventContext& EventContext,
 			}
 			else
 			{
-				OnDirectoryRescanRequired(EventContext, InheritanceContext, FileChangesProcessed[FileIdx].Filename,
+				OnDirectoryRescanRequired(ScopeLock, EventContext, InheritanceContext, FileChangesProcessed[FileIdx].Filename,
 					FileChangesProcessed[FileIdx].TimeStamp);
 			}
 
@@ -8152,14 +8172,14 @@ void FAssetRegistryImpl::OnDirectoryChanged(Impl::FEventContext& EventContext,
 			{
 				Impl::FScanPathContext Context(EventContext, InheritanceContext, NewDirs, NewFiles,
 					UE::AssetRegistry::EScanFlags::None, nullptr /* OutFoundAssets */);
-				ScanPathsSynchronous(Context);
+				ScanPathsSynchronous(&ScopeLock, Context);
 			}
 		}
 	}
-	ScanModifiedAssetFiles(EventContext, InheritanceContext, ModifiedFiles, UE::AssetRegistry::EScanFlags::None);
+	ScanModifiedAssetFiles(ScopeLock, EventContext, InheritanceContext, ModifiedFiles, UE::AssetRegistry::EScanFlags::None);
 }
 
-void FAssetRegistryImpl::OnDirectoryRescanRequired(Impl::FEventContext& EventContext,
+void FAssetRegistryImpl::OnDirectoryRescanRequired(FInterfaceWriteScopeLock& ScopeLock, Impl::FEventContext& EventContext,
 	Impl::FClassInheritanceContext& InheritanceContext, FString& DirPath, int64 BeforeTimeStamp)
 {
 	TArray<TPair<FString,FString>> DirPathsAndPackageNames;
@@ -8297,11 +8317,11 @@ void FAssetRegistryImpl::OnDirectoryRescanRequired(Impl::FEventContext& EventCon
 				TArray<FString> UnusedNewDirs;
 				Impl::FScanPathContext Context(EventContext, InheritanceContext, UnusedNewDirs, FinalResult.NewFiles,
 					UE::AssetRegistry::EScanFlags::None, nullptr /* OutFoundAssets */);
-				ScanPathsSynchronous(Context);
+				ScanPathsSynchronous(&ScopeLock, Context);
 			}
 		}
 	}
-	ScanModifiedAssetFiles(EventContext, InheritanceContext, FinalResult.ModifiedFiles, UE::AssetRegistry::EScanFlags::None);
+	ScanModifiedAssetFiles(ScopeLock, EventContext, InheritanceContext, FinalResult.ModifiedFiles, UE::AssetRegistry::EScanFlags::None);
 }
 
 }
@@ -8503,7 +8523,7 @@ void UAssetRegistryImpl::ScanModifiedAssetFiles(const TArray<FString>& InFilePat
 		UE::AssetRegistry::Impl::FClassInheritanceContext InheritanceContext;
 		UE::AssetRegistry::Impl::FClassInheritanceBuffer InheritanceBuffer;
 		GetInheritanceContextWithRequiredLock(InterfaceScopeLock, InheritanceContext, InheritanceBuffer);
-		GuardedData.ScanModifiedAssetFiles(EventContext, InheritanceContext, InFilePaths, ScanFlags);
+		GuardedData.ScanModifiedAssetFiles(InterfaceScopeLock, EventContext, InheritanceContext, InFilePaths, ScanFlags);
 	}
 
 #if WITH_EDITOR
@@ -8521,12 +8541,21 @@ void UAssetRegistryImpl::ScanModifiedAssetFiles(const TArray<FString>& InFilePat
 namespace UE::AssetRegistry
 {
 
-void FAssetRegistryImpl::ScanModifiedAssetFiles(Impl::FEventContext& EventContext,
+void FAssetRegistryImpl::ScanModifiedAssetFiles(FInterfaceWriteScopeLock& ScopeLock, Impl::FEventContext& EventContext,
 	Impl::FClassInheritanceContext& InheritanceContext, const TArray<FString>& InFilePaths,
 	UE::AssetRegistry::EScanFlags InScanFlags)
 {
 	if (InFilePaths.Num() > 0)
 	{
+		// ScanModifiedAssetFiles always does a force rescan of the given files
+		InScanFlags |= UE::AssetRegistry::EScanFlags::ForceRescan;
+
+		// Re-scan and update the asset registry with the new asset data
+		TArray<FSoftObjectPath> FoundAssets;
+		Impl::FScanPathContext Context(EventContext, InheritanceContext, TArray<FString>(), InFilePaths,
+			InScanFlags, &FoundAssets);
+		ScanPathsSynchronous(&ScopeLock, Context);
+
 		// Convert all the filenames to package names
 		TArray<FString> ModifiedPackageNames;
 		ModifiedPackageNames.Reserve(InFilePaths.Num());
@@ -8556,15 +8585,6 @@ void FAssetRegistryImpl::ScanModifiedAssetFiles(Impl::FEventContext& EventContex
 			}
 		}
 
-		// ScanModifiedAssetFiles always does a force rescan of the given files
-		InScanFlags |= UE::AssetRegistry::EScanFlags::ForceRescan;
-
-		// Re-scan and update the asset registry with the new asset data
-		TArray<FSoftObjectPath> FoundAssets;
-		Impl::FScanPathContext Context(EventContext, InheritanceContext, TArray<FString>(), InFilePaths,
-			InScanFlags, &FoundAssets);
-		ScanPathsSynchronous(Context);
-
 		// Remove any assets that are no longer present in the package
 		for (FSoftObjectPath& OldAssetPath : ExistingAssetDatas)
 		{
@@ -8629,7 +8649,7 @@ void UAssetRegistryImpl::OnContentPathMounted(const FString& InAssetPath, const
 		UE::AssetRegistry::Impl::FClassInheritanceContext InheritanceContext;
 		UE::AssetRegistry::Impl::FClassInheritanceBuffer InheritanceBuffer;
 		GetInheritanceContextWithRequiredLock(InterfaceScopeLock, InheritanceContext, InheritanceBuffer);
-		GuardedData.OnContentPathMounted(EventContext, InheritanceContext, InAssetPath, AssetPathWithTrailingSlash,
+		GuardedData.OnContentPathMounted(InterfaceScopeLock, EventContext, InheritanceContext, InAssetPath, AssetPathWithTrailingSlash,
 			FileSystemPath);
 
 		// Listen for directory changes in this content path
@@ -8659,7 +8679,7 @@ void UAssetRegistryImpl::OnContentPathMounted(const FString& InAssetPath, const
 namespace UE::AssetRegistry
 {
 
-void FAssetRegistryImpl::OnContentPathMounted(Impl::FEventContext& EventContext,
+void FAssetRegistryImpl::OnContentPathMounted(UE::AssetRegistry::FInterfaceWriteScopeLock& ScopeLock, Impl::FEventContext& EventContext,
 	Impl::FClassInheritanceContext& InheritanceContext, const FString& InAssetPath,
 	const FString& AssetPathWithTrailingSlash, const FString& FileSystemPath)
 {
@@ -8671,7 +8691,7 @@ void FAssetRegistryImpl::OnContentPathMounted(Impl::FEventContext& EventContext,
 		if (GlobalGatherer->IsSynchronous())
 		{
 			Impl::FScanPathContext Context(EventContext, InheritanceContext, { FileSystemPath }, TArray<FString>());
-			ScanPathsSynchronous(Context);
+			ScanPathsSynchronous(&ScopeLock, Context);
 		}
 		else
 		{
@@ -9003,8 +9023,8 @@ void UAssetRegistryImpl::GetInheritanceContextWithRequiredLock(UE::AssetRegistry
 	}
 
 	// Note that we have to reread all data since we may have dropped the lock
-	GetInheritanceContextAfterVerifyingLock(CurrentGeneratorClassesVersionNumber, CurrentAllClassesVersionNumber,
-		InheritanceContext, StackBuffer);
+	InheritanceContext.StackBuffer = &StackBuffer;
+	InheritanceContext.SelectBuffer(GuardedData, CurrentGeneratorClassesVersionNumber, CurrentAllClassesVersionNumber);
 }
 
 void UAssetRegistryImpl::GetInheritanceContextWithRequiredLock(UE::AssetRegistry::FInterfaceWriteScopeLock& InOutScopeLock,
@@ -9015,32 +9035,8 @@ void UAssetRegistryImpl::GetInheritanceContextWithRequiredLock(UE::AssetRegistry
 
 	uint64 CurrentGeneratorClassesVersionNumber = FAssetRegistryImpl::GetCurrentGeneratorClassesVersionNumber();
 	uint64 CurrentAllClassesVersionNumber = FAssetRegistryImpl::GetCurrentAllClassesVersionNumber();
-	GetInheritanceContextAfterVerifyingLock(CurrentGeneratorClassesVersionNumber, CurrentAllClassesVersionNumber,
-		InheritanceContext, StackBuffer);
-}
-
-void UAssetRegistryImpl::GetInheritanceContextAfterVerifyingLock(uint64 CurrentGeneratorClassesVersionNumber,
-	uint64 CurrentAllClassesVersionNumber,
-	UE::AssetRegistry::Impl::FClassInheritanceContext& InheritanceContext,
-	UE::AssetRegistry::Impl::FClassInheritanceBuffer& StackBuffer)
-{
-	// If bIsTempCachingAlwaysEnabled, then we are guaranteed that bIsTempCachingEnabled=true.
-	// We rely on this to simplify logic and only check bIsTempCachingEnabled
-	check(!GuardedData.IsTempCachingAlwaysEnabled() || GuardedData.IsTempCachingEnabled());
-
-	bool bCodeGeneratorClassesUpToDate = GuardedData.GetSavedGeneratorClassesVersionNumber() == CurrentGeneratorClassesVersionNumber;
-	if (GuardedData.IsTempCachingEnabled())
-	{
-		// Use the persistent buffer
-		UE::AssetRegistry::Impl::FClassInheritanceBuffer& TempCachedInheritanceBuffer = GuardedData.GetTempCachedInheritanceBuffer();
-		bool bInheritanceMapUpToDate = TempCachedInheritanceBuffer.IsUpToDate(CurrentAllClassesVersionNumber);
-		InheritanceContext.BindToBuffer(TempCachedInheritanceBuffer, GuardedData, bInheritanceMapUpToDate, bCodeGeneratorClassesUpToDate);
-	}
-	else
-	{
-		// Use the StackBuffer for the duration of the caller
-		InheritanceContext.BindToBuffer(StackBuffer, GuardedData, false /* bInInheritanceMapUpToDate */, bCodeGeneratorClassesUpToDate);
-	}
+	InheritanceContext.StackBuffer = &StackBuffer;
+	InheritanceContext.SelectBuffer(GuardedData, CurrentGeneratorClassesVersionNumber, CurrentAllClassesVersionNumber);
 }
 
 #if WITH_EDITOR
@@ -9143,6 +9139,34 @@ void FClassInheritanceContext::ConditionalUpdate()
 	bInheritanceMapUpToDate = true;
 }
 
+void FClassInheritanceContext::SelectBuffer(FAssetRegistryImpl& GuardedData, uint64 CurrentGeneratorClassesVersionNumber, uint64 CurrentAllClassesVersionNumber)
+{
+	// If bIsTempCachingAlwaysEnabled, then we are guaranteed that bIsTempCachingEnabled=true.
+	// We rely on this to simplify logic and only check bIsTempCachingEnabled
+	check(!GuardedData.IsTempCachingAlwaysEnabled() || GuardedData.IsTempCachingEnabled());
+
+	bool bLocalCodeGeneratorClassesUpToDate = GuardedData.GetSavedGeneratorClassesVersionNumber() == CurrentGeneratorClassesVersionNumber;
+	if (GuardedData.IsTempCachingEnabled())
+	{
+		// Use the persistent buffer
+		UE::AssetRegistry::Impl::FClassInheritanceBuffer& TempCachedInheritanceBuffer = GuardedData.GetTempCachedInheritanceBuffer();
+		bool bLocalInheritanceMapUpToDate = TempCachedInheritanceBuffer.IsUpToDate(CurrentAllClassesVersionNumber);
+		BindToBuffer(TempCachedInheritanceBuffer, GuardedData, bLocalInheritanceMapUpToDate, bLocalCodeGeneratorClassesUpToDate);
+	}
+	else
+	{
+		// Use the StackBuffer for the duration of the caller
+		check(StackBuffer);
+		BindToBuffer(*StackBuffer, GuardedData, false /* bInInheritanceMapUpToDate */, bLocalCodeGeneratorClassesUpToDate);
+	}
+}
+
+void FClassInheritanceContext::OnLockReentered(FAssetRegistryImpl& GuardedData)
+{
+	uint64 CurrentGeneratorClassesVersionNumber = FAssetRegistryImpl::GetCurrentGeneratorClassesVersionNumber();
+	uint64 CurrentAllClassesVersionNumber = FAssetRegistryImpl::GetCurrentAllClassesVersionNumber();
+	SelectBuffer(GuardedData, CurrentGeneratorClassesVersionNumber, CurrentAllClassesVersionNumber);
+}
 }
 
 void FAssetRegistryImpl::GetSubClasses(Impl::FClassInheritanceContext& InheritanceContext,
diff --git a/Engine/Source/Runtime/AssetRegistry/Private/AssetRegistry.h b/Engine/Source/Runtime/AssetRegistry/Private/AssetRegistry.h
--- a/Engine/Source/Runtime/AssetRegistry/Private/AssetRegistry.h
+++ b/Engine/Source/Runtime/AssetRegistry/Private/AssetRegistry.h
@@ -296,10 +296,6 @@ private:
 	void GetInheritanceContextWithRequiredLock(UE::AssetRegistry::FInterfaceWriteScopeLock& InOutScopeLock,
 		UE::AssetRegistry::Impl::FClassInheritanceContext& InheritanceContext,
 		UE::AssetRegistry::Impl::FClassInheritanceBuffer& StackBuffer);
-	void GetInheritanceContextAfterVerifyingLock(uint64 CurrentGeneratorClassesVersionNumber,
-		uint64 CurrentAllClassesVersionNumber,
-		UE::AssetRegistry::Impl::FClassInheritanceContext& InheritanceContext,
-		UE::AssetRegistry::Impl::FClassInheritanceBuffer& StackBuffer);
 
 
 #if WITH_EDITOR
diff --git a/Engine/Source/Runtime/AssetRegistry/Private/AssetRegistryImpl.h b/Engine/Source/Runtime/AssetRegistry/Private/AssetRegistryImpl.h
--- a/Engine/Source/Runtime/AssetRegistry/Private/AssetRegistryImpl.h
+++ b/Engine/Source/Runtime/AssetRegistry/Private/AssetRegistryImpl.h
@@ -75,7 +75,7 @@ struct FAsyncConsumer
 	* been decided not to be used. For performance reasons, this must be called within the WriteScopeLock, and the caller
 	* must handle the possibility that it leaves and reenters the lock.
 	*/
-	void Wait(UAssetRegistryImpl& UARI, FInterfaceWriteScopeLock& ScopeLock);
+	void Wait(FInterfaceWriteScopeLock& ScopeLock);
 	/** Callback from the async thread to consume the Premade ARState */
 	void Consume(UAssetRegistryImpl& UARI, UE::AssetRegistry::Impl::FEventContext& EventContext, ELoadResult LoadResult, FAssetRegistryState&& ARState);
 
@@ -359,7 +359,7 @@ public:
 	/** Constructor initializes as little as possible; meaningful operations are first done during Initialize. */
 	FAssetRegistryImpl();
 	/** Construct the AssetRegistryImpl, including initial scans if applicable. */
-	void Initialize(Impl::FInitializeContext& Context);
+	void Initialize(UE::AssetRegistry::FInterfaceWriteScopeLock& ScopeLock, Impl::FInitializeContext& Context);
 
 	/** Event for when duplicated assets are found and needs to be resolved */
 	IAssetRegistry::FAssetCollisionEvent& OnAssetCollision_Private();
@@ -368,8 +368,18 @@ public:
 
 
 	bool HasAssets(const FName PackagePath, const bool bRecursive) const;
-	FSoftObjectPath GetRedirectedObjectPath(const FSoftObjectPath& ObjectPath, UE::AssetRegistry::Impl::FEventContext* EventContext, 
-		UE::AssetRegistry::Impl::FClassInheritanceContext* InheritanceContext, bool bNeedsScanning);
+	/** 
+	* Returns the redirected object path for the provided soft object path. 
+	* This function will not attempt to scan the redirected path to ensure it has been registered with the Asset Registry. 
+	* If it's unclear if the object path is known to the Asset Registry already, call GetAndScanRedirectedObjectPath instead.
+	*/
+	FSoftObjectPath GetRedirectedObjectPath(FInterfaceReadScopeLock& ScopeLock, const FSoftObjectPath& ObjectPath);
+	/**
+	* Returns the redirected object path for the provided soft object path.
+	* This function will scan the redirected path to ensure it has been registered with the Asset Registry.
+	*/
+	FSoftObjectPath GetAndScanRedirectedObjectPath(FInterfaceWriteScopeLock& ScopeLock, const FSoftObjectPath& ObjectPath, 
+		UE::AssetRegistry::Impl::FEventContext* EventContext, UE::AssetRegistry::Impl::FClassInheritanceContext* InheritanceContext);
 	bool GetAncestorClassNames(Impl::FClassInheritanceContext& InheritanceContext, FTopLevelAssetPath ClassName,
 		TArray<FTopLevelAssetPath>& OutAncestorClassNames) const;
 	void CompileFilter(Impl::FClassInheritanceContext& InheritanceContext, const FARFilter& InFilter,
@@ -380,9 +390,9 @@ public:
 	void SearchAllAssets(Impl::FEventContext& EventContext, Impl::FClassInheritanceContext& InheritanceContext,
 		bool bSynchronousSearch);
 	bool GetVerseFilesByPath(FName PackagePath, TArray<FName>* OutFilePaths, bool bRecursive) const;
-	void ScanPathsSynchronous(Impl::FScanPathContext& Context);
+	void ScanPathsSynchronous(FInterfaceWriteScopeLock* ScopeLock, Impl::FScanPathContext& Context);
 	void PrioritizeSearchPath(const FString& PathToPrioritize);
-	void ScanModifiedAssetFiles(Impl::FEventContext& EventContext, Impl::FClassInheritanceContext& InheritanceContext,
+	void ScanModifiedAssetFiles(FInterfaceWriteScopeLock& ScopeLock, Impl::FEventContext& EventContext, Impl::FClassInheritanceContext& InheritanceContext,
 		const TArray<FString>& InFilePaths, UE::AssetRegistry::EScanFlags ScanFlags);
 	void Serialize(FArchive& Ar, Impl::FEventContext& EventContext);
 	void AppendState(Impl::FEventContext& EventContext, const FAssetRegistryState& InState,
@@ -395,13 +405,13 @@ public:
 	bool SetPrimaryAssetIdForObjectPath(Impl::FEventContext& EventContext, const FSoftObjectPath& ObjectPath, FPrimaryAssetId PrimaryAssetId);
 	bool ResolveRedirect(const FString& InPackageName, FString& OutPackageName) const;
 #if WITH_EDITOR
-	void OnDirectoryChanged(Impl::FEventContext& EventContext, Impl::FClassInheritanceContext& InheritanceContext,
+	void OnDirectoryChanged(FInterfaceWriteScopeLock& ScopeLock, Impl::FEventContext& EventContext, Impl::FClassInheritanceContext& InheritanceContext,
 		TArray<FFileChangeData>& FileChangesProcessed);
-	void OnDirectoryRescanRequired(Impl::FEventContext& EventContext,
+	void OnDirectoryRescanRequired(FInterfaceWriteScopeLock& ScopeLock, Impl::FEventContext& EventContext,
 		Impl::FClassInheritanceContext& InheritanceContext, FString& DirPath, int64 BeforeTimeStamp);
 	void AddLoadedAssetToProcess(const UObject& AssetLoaded);
 #endif
-	void OnContentPathMounted(Impl::FEventContext& EventContext, Impl::FClassInheritanceContext& InheritanceContext,
+	void OnContentPathMounted(UE::AssetRegistry::FInterfaceWriteScopeLock& ScopeLock, Impl::FEventContext& EventContext, Impl::FClassInheritanceContext& InheritanceContext,
 		const FString& InAssetPath, const FString& AssetPathWithTrailingSlash, const FString& FileSystemPath);
 	void OnContentPathDismounted(Impl::FEventContext& EventContext, const FString& InAssetPath,
 		const FString& AssetPathNoTrailingSlash, const FString& FileSystemPath);
@@ -444,7 +454,7 @@ public:
 	 */
 	void LogSearchDiagnostics(double StartTime);
 	/** Look for and load a single AssetData result from the gatherer. */
-	void TickGatherPackage(Impl::FEventContext& EventContext, const FString& PackageName, const FString& LocalPath);
+	void TickGatherPackage(UE::AssetRegistry::FInterfaceWriteScopeLock& ScopeLock, Impl::FEventContext& EventContext, const FString& PackageName, const FString& LocalPath);
 	void ClearGathererCache();
 #if WITH_EDITOR
 	void AssetsSaved(UE::AssetRegistry::Impl::FEventContext& EventContext, TArray<FAssetData>&& Assets);
@@ -537,8 +547,7 @@ public:
 	 * For performance reasons, this must be called within the WriteScopeLock, and the caller 
 	 * must handle the possibility that it leaves and reenters the lock.
 	 */
-	void ConditionalLoadPremadeAssetRegistry(UAssetRegistryImpl& UARI, 
-		UE::AssetRegistry::Impl::FEventContext& EventContext, FInterfaceWriteScopeLock& ScopeLock);
+	void ConditionalLoadPremadeAssetRegistry(UE::AssetRegistry::FInterfaceWriteScopeLock& ScopeLock, UE::AssetRegistry::Impl::FEventContext& EventContext);
 
 #if WITH_EDITOR
 	/** Request to pause or resume background processing of scan results.
@@ -556,6 +565,10 @@ private:
 	/** The delegate to execute when an asset collision is detected */
 	IAssetRegistry::FAssetCollisionEvent AssetCollisionEvent;
 
+	/** Helper method for public implementations which vary based on lock type */
+	template<typename TScopeLock>
+	FSoftObjectPath GetRedirectedObjectPath(TScopeLock& ScopeLock, const FSoftObjectPath& ObjectPath, UE::AssetRegistry::Impl::FEventContext* EventContext,
+		UE::AssetRegistry::Impl::FClassInheritanceContext* InheritanceContext);
 	/**
 	 * Enumerate all AssetData in memory and in the State that match the given path.
 	 * The returned AssetData for objects present in memory have no AssetRegistryTags calculated
@@ -746,7 +759,9 @@ private:
 	/** Time the additional mount async search was started */
 	double AdditionalMountSearchStartTime = 0.0f;
 	/** Flag to indicate if we used an initial async search */
-	bool bInitialSearchStarted;
+	bool bInitialSearchStarted = false;
+	/** Flag to indicate Initialize() has run to completion */
+	bool bInitializationComplete = false;
 	/** Flag to indicate if the initial background search has completed. All access are relaxed because
 	 *  the actual search data can only be accessed under a proper lock.
 	 */
@@ -894,6 +909,8 @@ struct FClassInheritanceContext
 {
 	/** The Buffer providing the data for this context. */
 	FClassInheritanceBuffer* Buffer = nullptr;
+	/** The stack allocated Buffer which may be used to provide data for this context. */
+	FClassInheritanceBuffer* StackBuffer = nullptr;
 	/** Back pointer to the AssetRegistryImpl to update the Buffer if necessary */
 	FAssetRegistryImpl* AssetRegistryImpl = nullptr;
 	/** Whether ConditionalUpdate needs to update inheritance data */
@@ -906,6 +923,10 @@ struct FClassInheritanceContext
 		bool bInInheritanceMapUpToDate, bool bInCodeGeneratorClassesUpToDate);
 	/** Update the backing buffer if it is out of date */
 	void ConditionalUpdate();
+	/** Selects and binds a FClassInheritanceBuffer buffer to be used by this context. */
+	void SelectBuffer(FAssetRegistryImpl& GuardedData, uint64 CurrentGeneratorClassesVersionNumber, uint64 CurrentAllClassesVersionNumber);
+	/** Updates the context after the InterfaceLock has been unlocked and locked while the context is in use. */
+	void OnLockReentered(FAssetRegistryImpl& GuardedData);
 };
 
 /** Context to transform, collect, and pass along ScanPathsSynchronous arguments and results */
@@ -931,10 +952,6 @@ struct FScanPathContext
 	EGatherStatus Status = EGatherStatus::TickActiveGatherActive;
 };
 
-/** Callback types for FTickContext */
-typedef TFunctionRef<void(const TMultiMap<FName, FAssetData*>&)> FAssetsFoundCallback;
-typedef TFunctionRef<void(const TRingBuffer<FName>&)> FVerseFilesFoundCallback;
-
 /** Input and output variables for the TickGatherer function. */
 struct FTickContext
 {
@@ -991,8 +1008,6 @@ public:
 	FAssetRegistryImpl& GuardedData;
 	Impl::FEventContext& EventContext;
 	Impl::FClassInheritanceContext& InheritanceContext;
-	TOptional<FAssetsFoundCallback> AssetsFoundCallback;
-	TOptional<FVerseFilesFoundCallback> VerseFilesFoundCallback;
 	TOptional<TSet<FString>> MountPointsForVerifyAfterGather;
 	double TimingStartTime = -1.;
 	int32 NumGatherFromDiskPending = 0;
