AssetRegistryState: Add an optimized sort function GetMutableAssetsSortedByObjectPath for use in CookByTheBookFinished
Preallocates strings for the package FNames in the FAssetData*, to avoid the repeated costs of looking up the string for each FName.
Runs ParallelMergeSort instead of Algo::Sort.
400,000 packages: 5.3s -> 0.6s.
#rnx

diff --git a/Engine/Source/Editor/UnrealEd/Private/Commandlets/AssetRegistryGenerator.cpp b/Engine/Source/Editor/UnrealEd/Private/Commandlets/AssetRegistryGenerator.cpp
--- a/Engine/Source/Editor/UnrealEd/Private/Commandlets/AssetRegistryGenerator.cpp
+++ b/Engine/Source/Editor/UnrealEd/Private/Commandlets/AssetRegistryGenerator.cpp
@@ -2383,17 +2383,8 @@ bool FAssetRegistryGenerator::GenerateAssetChunkInformationCSV(const FString& Ou
 	FString TmpString, TmpStringChunks;
 	ANSICHAR HeaderText[] = "ChunkID, Package Name, Class Type, Hard or Soft Chunk, File Size, Other Chunks\n";
 
-	TArray<const FAssetData*> AssetDataList;
-	State.EnumerateAllAssets([&AssetDataList](const FAssetData& AssetData)
-	{
-		AssetDataList.Add(&AssetData);
-	});
-
-	// Sort list so it's consistent over time
-	AssetDataList.Sort([](const FAssetData& A, const FAssetData& B)
-	{
-		return A.GetSoftObjectPath().LexicalLess(B.GetSoftObjectPath());
-	});
+	// Get the list of assets sorted by packagename, so the CSV is consistent between cooks.
+	TArray<FAssetData*> AssetDataList = State.GetMutableAssetsSortedByObjectPath();
 
 	// Create file for all chunks
 	TUniquePtr<FArchive> AllChunksFile(IFileManager::Get().CreateFileWriter(*FPaths::Combine(*OutputPath, TEXT("AllChunksInfo.csv"))));
diff --git a/Engine/Source/Runtime/AssetRegistry/Internal/AssetRegistry/LinearBlockAllocator.h b/Engine/Source/Runtime/AssetRegistry/Internal/AssetRegistry/LinearBlockAllocator.h
new file mode 100644
--- /dev/null
+++ b/Engine/Source/Runtime/AssetRegistry/Internal/AssetRegistry/LinearBlockAllocator.h
@@ -0,0 +1,225 @@
+// Copyright Epic Games, Inc. All Rights Reserved.
+
+#pragma once
+
+#include "Containers/Array.h"
+#include "CoreTypes.h"
+#include "HAL/CriticalSection.h"
+#include "HAL/UnrealMemory.h"
+#include "Misc/AssertionMacros.h"
+#include "Misc/ScopeLock.h"
+
+#include <atomic>
+
+namespace UE::LinearBlockAllocator::Private
+{
+
+struct FBlock
+{
+	int32 NextOffset;
+};
+
+}
+namespace UE { class FLinearBlockAllocatorThreadAccessor; }
+
+namespace UE
+{
+
+/**
+ * Provides Malloc operation that serves memory from blocks, and does not free up memory until
+ * the allocator goes out of scope. This is an optimization to reduce the cost of Malloc and Free
+ * for dynamically allocated scratch or intermediate data.
+ * 
+ * The expected use case is to provide string allocations for FNames.
+ *
+ * Memory for the blocks is allocated from FMemory::Malloc. Large allocations will cause wasted
+ * memory as the excess memory from the previous block goes unused.
+ * 
+ * When used directly, the allocator is not threadsafe. A high-performance threadsafe adaptor
+ * is available by using FLinearBlockAllocatorThreadAccessor, but in that case every thread using
+ * the allocator must access it only through FLinearBlockAllocatorThreadAccessor.
+ * 
+ */
+class FLinearBlockAllocator
+{
+public:
+	~FLinearBlockAllocator();
+
+	void* Malloc(SIZE_T Size, uint32 Alignment);
+
+private:
+	void* TryAllocateFromBlock(SIZE_T Size, uint32 Alignment, UE::LinearBlockAllocator::Private::FBlock* Block);
+	bool RequiresCustomAllocation(SIZE_T Size, uint32 Alignment) const;
+	UE::LinearBlockAllocator::Private::FBlock* GetNewBlock();
+	void* GetCustomAllocation(SIZE_T Size, uint32 Alignment);
+
+	TArray<void*> BlockAllocs; // guarded by LockForThreadSafeAllocator when threading
+	UE::LinearBlockAllocator::Private::FBlock* LastBlock = nullptr; // unused during threading
+	uint32 BlockAlignment = 4; // constant during threading
+	SIZE_T BlockSize = 1 << 16; // constant during threading
+	FCriticalSection LockForThreadSafeAllocator;
+	std::atomic<int32> ReferenceCount{ 0 };
+
+	friend class UE::FLinearBlockAllocatorThreadAccessor;
+};
+
+/**
+ * Accessor for a LinearBlockAllocator when multithreading. Every thread using the Allocator
+ * must access it through its own FLinearBlockAllocatorThreadAccessor.
+ * 
+ * Intended use is to share an allocator amongst the task threads used within a single function
+ * in e.g. ParallelFor.
+ * 
+ * The lifetime of FLinearBlockAllocatorThreadAccessor must be within the lifetime of the
+ * Accessor, an assertion fires if FLinearBlockAllocatorThreadAccessor are still allocated
+ * when the Allocator is destructed. So the thread that destroys the allocator must wait
+ * on all other threads using the Allocator to destroy their FLinearBlockAllocatorThreadAccessor.
+ */
+class FLinearBlockAllocatorThreadAccessor
+{
+public:
+	explicit FLinearBlockAllocatorThreadAccessor(FLinearBlockAllocator& InAllocator);
+	~FLinearBlockAllocatorThreadAccessor();
+
+	void* Malloc(SIZE_T Size, uint32 Alignment);
+
+private:
+	FLinearBlockAllocator& Allocator;
+	UE::LinearBlockAllocator::Private::FBlock* LastBlock = nullptr;
+};
+
+} // namespace UE
+
+
+///////////////////////////////////////////////////////
+// Inline implementations
+///////////////////////////////////////////////////////
+
+
+namespace UE
+{
+
+inline FLinearBlockAllocator::~FLinearBlockAllocator()
+{
+	using namespace UE::LinearBlockAllocator::Private;
+
+	checkf(ReferenceCount.load(std::memory_order_relaxed) == 0,
+		TEXT("FLinearBlockAllocator is destroyed while still in use by a FLinearBlockAllocatorThreadAccessor."));
+	for (void* BlockAlloc : BlockAllocs)
+	{
+		FMemory::Free(BlockAlloc);
+	}
+}
+
+inline void* FLinearBlockAllocator::Malloc(SIZE_T Size, uint32 Alignment)
+{
+	checkf(ReferenceCount.load(std::memory_order_relaxed) == 0,
+		TEXT("FLinearBlockAllocator::Malloc is being used directly while in use by a FLinearBlockAllocatorThreadAccessor, this is not allowed, all malloc calls must go through a FLinearBlockAllocatorThreadAccessor."));
+
+	void* Result = TryAllocateFromBlock(Size, Alignment, LastBlock);
+	if (Result)
+	{
+		return Result;
+	}
+
+	if (RequiresCustomAllocation(Size, Alignment))
+	{
+		return GetCustomAllocation(Size, Alignment);
+	}
+
+	LastBlock = GetNewBlock();
+	Result = TryAllocateFromBlock(Size, Alignment, LastBlock);
+	check(Result);
+	return Result;
+}
+
+inline void* FLinearBlockAllocator::TryAllocateFromBlock(SIZE_T Size, uint32 Alignment,
+	UE::LinearBlockAllocator::Private::FBlock* Block)
+{
+	if (!Block)
+	{
+		return nullptr;
+	}
+
+	void* NextPtr = (void*)(((SIZE_T)Block) + ((SIZE_T)Block->NextOffset));
+	void* AlignedPtr = Align(NextPtr, Alignment);
+	SIZE_T NewOffset = ((SIZE_T)AlignedPtr) + Size - ((SIZE_T)Block);
+	if (NewOffset > BlockSize)
+	{
+		return nullptr;
+	}
+
+	Block->NextOffset = (int32)NewOffset;
+	return AlignedPtr;
+}
+
+bool FLinearBlockAllocator::RequiresCustomAllocation(SIZE_T Size, uint32 Alignment) const
+{
+	using namespace UE::LinearBlockAllocator::Private;
+
+	// Compute the worse case AlignmentPadding, which occurs if the pointer we received for the block allocation
+	// is aligned to BlockAlignment and nothing higher, which we can calculate by assuming the pointer starts at
+	// 1*BlockAlignment. The starting position is then BlockHeaderSize after that, and we have to AlignUp the
+	// starting position.
+	constexpr uint32 BlockHeaderSize = sizeof(FBlock);
+	uint32 StartingOffset = BlockAlignment + BlockHeaderSize;
+	uint32 AlignedOffset = Align(StartingOffset, Alignment);
+	uint32 AlignmentPadding = AlignedOffset - StartingOffset;
+	SIZE_T RequiredBlockSize = Size + AlignmentPadding + BlockHeaderSize;
+	return RequiredBlockSize > BlockSize;
+}
+
+UE::LinearBlockAllocator::Private::FBlock* FLinearBlockAllocator::GetNewBlock()
+{
+	using namespace UE::LinearBlockAllocator::Private;
+
+	constexpr uint32 BlockHeaderSize = sizeof(FBlock);
+	void* BlockAllocation = FMemory::Malloc(BlockSize, BlockAlignment);
+	FBlock* NewBlock = reinterpret_cast<FBlock*>(BlockAllocation);
+	NewBlock->NextOffset = BlockHeaderSize;
+	BlockAllocs.Add(BlockAllocation);
+	return NewBlock;
+}
+
+void* FLinearBlockAllocator::GetCustomAllocation(SIZE_T Size, uint32 Alignment)
+{
+	using namespace UE::LinearBlockAllocator::Private;
+
+	void* CustomAllocation = FMemory::Malloc(Size, Alignment);
+	BlockAllocs.Add(CustomAllocation);
+	return CustomAllocation;
+}
+
+inline FLinearBlockAllocatorThreadAccessor::FLinearBlockAllocatorThreadAccessor(FLinearBlockAllocator& InAllocator)
+	: Allocator(InAllocator)
+{
+	Allocator.ReferenceCount.fetch_add(1, std::memory_order_relaxed);
+}
+
+inline FLinearBlockAllocatorThreadAccessor::~FLinearBlockAllocatorThreadAccessor()
+{
+	Allocator.ReferenceCount.fetch_add(-1, std::memory_order_relaxed);
+}
+
+inline void* FLinearBlockAllocatorThreadAccessor::Malloc(SIZE_T Size, uint32 Alignment)
+{
+	void* Result = Allocator.TryAllocateFromBlock(Size, Alignment, LastBlock);
+	if (Result)
+	{
+		return Result;
+	}
+
+	{
+		FScopeLock ScopeLockForThreadSafeAllocator(&Allocator.LockForThreadSafeAllocator);
+		if (Allocator.RequiresCustomAllocation(Size, Alignment))
+		{
+			return Allocator.GetCustomAllocation(Size, Alignment);
+		}
+		LastBlock = Allocator.GetNewBlock();
+	}
+	Result = Allocator.TryAllocateFromBlock(Size, Alignment, LastBlock);
+	check(Result);
+	return Result;
+}
+
+} // namespace UE
\ No newline at end of file
diff --git a/Engine/Source/Runtime/AssetRegistry/Internal/AssetRegistry/ParallelSort.h b/Engine/Source/Runtime/AssetRegistry/Internal/AssetRegistry/ParallelSort.h
new file mode 100644
--- /dev/null
+++ b/Engine/Source/Runtime/AssetRegistry/Internal/AssetRegistry/ParallelSort.h
@@ -0,0 +1,390 @@
+// Copyright Epic Games, Inc. All Rights Reserved.
+
+#pragma once
+
+#include "Containers/RingBuffer.h"
+#include "HAL/PlatformMath.h"
+#include "Misc/AssertionMacros.h"
+#include "Tasks/Task.h"
+#include "Templates/Less.h"
+#include "Templates/UnrealTemplate.h"
+
+namespace Algo
+{
+
+/**
+ * Sort a range of elements using a user-defined predicate class.  The sort is stable.
+ * The sort operation will be distributed to task threads using UE::Tasks::Launch, and the Predicate
+ * will be called on elements in the input range in parallel (possibly even the same element being
+ * read from multiple threads at the same time); the Predicate must be thread safe.
+ *
+ * bool Predicate(const T& A, const T& B) -> true if A < B.
+ * 
+ * @param  Range      The range to sort.
+ * @param  Predicate  A binary predicate object used to specify if one element should precede another.
+*/
+template <typename RangeType, typename PredicateType>
+void ParallelSort(RangeType&& Range, PredicateType Predicate);
+
+/** Debug version of ParallelSort that shares implementation but is single threaded. */
+template <typename RangeType, typename PredicateType>
+void ParallelSortForceSingleThreaded(RangeType&& Range, PredicateType Predicate);
+
+} // namespace Algo
+
+
+///////////////////////////////////////////////////////
+// Implementation forward declares
+///////////////////////////////////////////////////////
+
+namespace Algo::MergeSort
+{
+
+template <typename T, typename IndexType, typename PredicateType>
+void ParallelMergeSortPointerAndNum(T* First, IndexType Num, PredicateType Predicate);
+template <typename T, typename IndexType, typename PredicateType>
+void MergeSortPointerAndNum(T* First, IndexType Num, PredicateType Predicate);
+template <typename T, typename IndexType, typename PredicateType>
+void MergeSortRecursive(T* First, IndexType Num, PredicateType Predicate, TRingBuffer<T>& MergeBuffer);
+template <typename T, typename IndexType, typename PredicateType>
+void Merge(T* First, IndexType Num, IndexType FloorMidIndex, PredicateType Predicate, TRingBuffer<T>& MergeBuffer);
+template <typename T, typename IndexType, typename PredicateType>
+void SmallSort(T* First, IndexType Num, PredicateType Predicate);
+
+} // namespace Algo::MergeSort
+
+
+///////////////////////////////////////////////////////
+// Template implementations
+///////////////////////////////////////////////////////
+
+
+namespace Algo
+{
+
+template <typename RangeType, typename PredicateType>
+void ParallelSort(RangeType&& Range, PredicateType Predicate)
+{
+	Algo::MergeSort::ParallelMergeSortPointerAndNum(GetData(Range), GetNum(Range), MoveTemp(Predicate));
+}
+
+template <typename RangeType>
+void ParallelSort(RangeType&& Range)
+{
+	Algo::MergeSort::ParallelMergeSortPointerAndNum(GetData(Range), GetNum(Range), TLess<>());
+}
+
+template <typename RangeType, typename PredicateType>
+void ParallelSortForceSingleThreaded(RangeType&& Range, PredicateType Predicate)
+{
+	Algo::MergeSort::MergeSortPointerAndNum(GetData(Range), GetNum(Range), MoveTemp(Predicate));
+}
+
+template <typename RangeType, typename PredicateType>
+void ParallelSortForceSingleThreaded(RangeType&& Range)
+{
+	Algo::MergeSort::MergeSortPointerAndNum(GetData(Range), GetNum(Range), TLess<>());
+}
+
+} // namespace Algo
+
+namespace Algo::MergeSort
+{
+
+template <typename T, typename IndexType, typename PredicateType>
+void ParallelMergeSortPointerAndNum(T* First, IndexType Num, PredicateType Predicate)
+{
+	using namespace UE::Tasks;
+
+	// In Round 0 we divide the Array up initially into TaskCount segments, and create an initial set of tasks that
+	// recursively mergesorts all segments in parallel.
+
+	// Rounds 1 through NumRounds-1 merge each pair of recursively sorted segments from the previous round into a
+	// combined segment, ending with NumRound -1 merging up into a single segment for the entire range.
+
+	// We require TaskCount is a Power of 2
+	// Set MaxAsymptoticTaskCount somewhat arbitrarily at 16 subtasks. There are diminishing returns to further division;
+	// total time is ~ Sum(i=1,log(threadCount))(N/i) + N/(log(threadCount+1)) log (N/(log(threadCount+1)
+	// e.g. for  8 threads and N = 2^20, it is N + N/2 + N/4 + N/8 + N/16 log(2^20/2^4) == 2*N - 2*N/16 + (N/16)*16 ~ 2.875*N
+	//      for 16 threads and N = 2^20, it is N + N/2 + N/4 + N/8 + N/16 + N/32 log(2^20/2^5) == 2*N - 2*N/32 + (N/32)*15 ~ 2.40*N
+	// Approaches time == 2*N in the limit of Number of threads == N
+	//
+	// Depending on N, MaxTaskCount might be less than MaxAsymptoticTaskCount because we want to avoid the overhead of
+	// launching threads to sort a small number of elements, so we set a MinBatchSize.
+	// TODO: Find some objective way to set MinBatchSize.
+	constexpr int32 MinBatchSize = 16;
+	constexpr int32 MaxAsymptoticTaskCount = 16;
+
+	// MinBatchSize = DivAndRoundUp(Num, MaxTaskCount) == (Num + MaxTaskCount - 1) / MaxTaskCount;
+	// ==>
+	// MaxTaskCount == DivAndRoundDown(Num - 1, MinBatchSize - 1)
+	int32 MaxTaskCount = (Num - 1) / (MinBatchSize - 1);
+	if (MaxTaskCount < 2)
+	{
+		return Algo::MergeSort::MergeSortPointerAndNum(First, Num, Predicate);
+	}
+	MaxTaskCount = FMath::Min(MaxTaskCount, MaxAsymptoticTaskCount);
+	int32 TaskCount = 1 << FPlatformMath::FloorLog2(MaxTaskCount);
+	IndexType BatchSize = (Num + TaskCount - 1) / TaskCount;
+	
+	// Round 0 is the initial recursive sort of segments, Round 1 is the first merge round, ...
+	// Round TaskCount is the merge into the final full Range.
+	const int32 NumRounds = FPlatformMath::FloorLog2(TaskCount) + 1;
+
+	// We define a recursive merge function for round N that launches a task for the Upper in round NumRounds - 1,
+	// then works itself on the Lower for round N - 1, then merges. We end up with TaskCount threads (including the
+	// main thread) doing the single-threaded recursive call to mergesort in round 0, then half that many threads
+	// doing the merge in round 1, ..., and only the main thread doing the merge in round NumRounds - 1.
+
+	// Each thread working on a round needs scratch space for the merge, we give it a mergebuffer. It uses the same
+	// mergebuffer for each of the recursive calls that it works on itself; the tasks it kicks off create their
+	// own mergebuffers and do the same.
+	// Merging two segments, each of size <= BatchSize into a merged 2*BatchSize length segment requires scratch
+	// space equal to BatchSize, (plus 1 more because we push before we pop).
+	auto GetMergeBufferForRound = [BatchSize](int32 Round)
+		{
+			// BatchSize << Round is the size of the segment produced for the round, we only need half that (plus one).
+			const IndexType BufferSize = ((BatchSize << Round) + 1) / 2 + 1;
+			TRingBuffer<T> MergeBuffer;
+			MergeBuffer.Reserve(BufferSize);
+			return MergeBuffer;
+		};
+
+	// Does just the merge job with no recursion.
+	auto MergeRoundAndIndex = [TaskCount, Num, BatchSize, &Predicate, First](int32 Round, int32 IndexInRound, TRingBuffer<T>& MergeBuffer)
+		{
+			const int32 NumMergesForRound = TaskCount >> Round;
+			const IndexType CurrentLevelBatchSize = BatchSize << Round;
+			const IndexType PreviousRoundBatchSize = BatchSize << (Round - 1);
+			const IndexType LowerStart = IndexInRound * CurrentLevelBatchSize;
+			const IndexType UpperStart = IndexInRound * CurrentLevelBatchSize + PreviousRoundBatchSize;
+			const IndexType MergeNum = IndexInRound < NumMergesForRound - 1 ? CurrentLevelBatchSize : Num - LowerStart;
+			Merge(First + LowerStart, MergeNum, UpperStart - LowerStart, Predicate, MergeBuffer);
+		};
+
+	// The base case of recursion; call single-threaded mergesort, which does its own recursion but without any
+	// threading.
+	auto RecursiveMergeRound0 = [TaskCount, Num, BatchSize, &Predicate, First](int32 IndexInRound, TRingBuffer<T>& MergeBuffer)
+		{
+			const IndexType StartIndex = IndexInRound * BatchSize;
+			const IndexType MergeNum = IndexInRound < TaskCount - 1 ? BatchSize : Num - StartIndex;
+			MergeSortRecursive(First + StartIndex, MergeNum, Predicate, MergeBuffer);
+		};
+
+	// Our recursive function that forks itself to handle the upper recursive segment, works itself recursively on the
+	// lower recursive segment, and then merges the two segments. The arguments are the round number, and the index
+	// within the round of the segment it is creating.
+	TFunction<void(int32, int32, TRingBuffer<T>&)> RecursiveMergeRoundAndIndexPtr;
+	auto RecursiveMergeRoundAndIndex =
+		[&MergeRoundAndIndex, &RecursiveMergeRoundAndIndexPtr, &RecursiveMergeRound0, &GetMergeBufferForRound]
+		(int32 Round, int32 IndexInRound, TRingBuffer<T>& MergeBuffer)
+		{
+			if (Round > 0)
+			{
+				FTask Upper = Launch(TEXT("ParallelMergeSort"),
+					[&RecursiveMergeRoundAndIndexPtr, &GetMergeBufferForRound, Round, IndexInRound]()
+					{
+						TRingBuffer<T> MergeBuffer = GetMergeBufferForRound(Round - 1);
+						RecursiveMergeRoundAndIndexPtr(Round - 1, 2 * IndexInRound + 1, MergeBuffer);
+					});
+
+				RecursiveMergeRoundAndIndexPtr(Round - 1, 2 * IndexInRound, MergeBuffer);
+				Upper.Wait();
+				MergeRoundAndIndex(Round, IndexInRound, MergeBuffer);
+			}
+			else
+			{
+				RecursiveMergeRound0(IndexInRound, MergeBuffer);
+			}
+		};
+	RecursiveMergeRoundAndIndexPtr = RecursiveMergeRoundAndIndex;
+
+	// Kick off RecursiveMergeRoundAndIndexPtr for the last round, at its only segment index: index 0.
+	TRingBuffer<T> MergeBuffer = GetMergeBufferForRound(NumRounds - 1);
+	RecursiveMergeRoundAndIndexPtr(NumRounds - 1, 0, MergeBuffer);
+}
+
+template <typename T, typename IndexType, typename PredicateType>
+void MergeSortPointerAndNum(T* First, IndexType Num, PredicateType Predicate)
+{
+	TRingBuffer<T> MergeBuffer;
+	MergeBuffer.Reserve((Num + 1) / 2);
+	MergeSortRecursive(First, Num, Predicate, MergeBuffer);
+}
+
+template <typename T, typename IndexType, typename PredicateType>
+void MergeSortRecursive(T* First, IndexType Num, PredicateType Predicate, TRingBuffer<T>& MergeBuffer)
+{
+	if (Num <= 3)
+	{
+		SmallSort(First, Num, Predicate);
+		return;
+	}
+
+	const IndexType FloorMidIndex = (Num + 1) / 2;
+	MergeSortRecursive(First, FloorMidIndex, Predicate, MergeBuffer);
+	MergeSortRecursive(First + FloorMidIndex, Num - FloorMidIndex, Predicate, MergeBuffer);
+
+	Merge(First, Num, FloorMidIndex, Predicate, MergeBuffer);
+}
+
+template <typename T, typename IndexType, typename PredicateType>
+void Merge(T* First, IndexType Num, IndexType FloorMidIndex, PredicateType Predicate, TRingBuffer<T>& MergeBuffer)
+{
+	IndexType WriteHead = 0;
+	IndexType UpperReadHead = FloorMidIndex;
+
+	// Skip past however many leading elements of Lower come before the first element of Upper
+	while (WriteHead < FloorMidIndex)
+	{
+		if (!Predicate(First[UpperReadHead], First[WriteHead]))
+		{
+			++WriteHead;
+		}
+		else
+		{
+			break;
+		}
+	}
+
+	if (WriteHead == FloorMidIndex)
+	{
+		// All elements of Lower came before Upper, the range is already sorted
+		return;
+	}
+
+	MergeBuffer.Reset();
+	// Our caller is supposed to reserve the MergeBuffer for us. The maximum size of MergeBuffer that we use is
+	// the length of the Upper half, plus 1. Assert that we have that size so that we don't hit the performance
+	// cost of growing the buffer here.
+	checkf(MergeBuffer.Max() >= Num - FloorMidIndex + 1,
+		TEXT("ParallelSort bug: MergeBuffer not preallocated large enough. %d < %d"),
+		MergeBuffer.Max(), Num - FloorMidIndex + 1);
+
+	// Move the UpperReadHead into WriteHead, we already compared it above and know it is lower
+	MergeBuffer.Add(MoveTemp(First[WriteHead]));
+	First[WriteHead] = MoveTemp(First[UpperReadHead]);
+	++UpperReadHead;
+	++WriteHead;
+
+	while (UpperReadHead < Num && WriteHead < FloorMidIndex)
+	{
+		check(WriteHead < UpperReadHead);
+		MergeBuffer.Add(MoveTemp(First[WriteHead]));
+		if (!Predicate(First[UpperReadHead], MergeBuffer[0]))
+		{
+			First[WriteHead] = MergeBuffer.PopFrontValue();
+		}
+		else
+		{
+			First[WriteHead] = MoveTemp(First[UpperReadHead]);
+			++UpperReadHead;
+		}
+		++WriteHead;
+	}
+	while (UpperReadHead < Num && !MergeBuffer.IsEmpty())
+	{
+		check(WriteHead < UpperReadHead);
+		if (!Predicate(First[UpperReadHead], MergeBuffer[0]))
+		{
+			First[WriteHead] = MergeBuffer.PopFrontValue();
+		}
+		else
+		{
+			First[WriteHead] = MoveTemp(First[UpperReadHead]);
+			++UpperReadHead;
+		}
+		++WriteHead;
+	}
+	while (WriteHead < FloorMidIndex)
+	{
+		MergeBuffer.Add(MoveTemp(First[WriteHead]));
+		First[WriteHead] = MergeBuffer.PopFrontValue();
+		++WriteHead;
+	}
+	while (!MergeBuffer.IsEmpty())
+	{
+		check(WriteHead < UpperReadHead);
+		First[WriteHead] = MergeBuffer.PopFrontValue();
+		++WriteHead;
+	}
+
+	// We have finished all lower, and so the remaining elements to write (if any) all come from the
+	// remainder of UpperReadHead. And we don't need to move them because they're already in their destination
+	// position. So the list of remaining writes should be identical to the list of remaining uppers. If we're
+	// done with Uppers as well, then WriteHead == UpperReadHead == Num.
+	check(WriteHead == UpperReadHead);
+}
+
+template <typename T, typename IndexType, typename PredicateType>
+void SmallSort(T* First, IndexType Num, PredicateType Predicate)
+{
+	switch (Num)
+	{
+	case 0:
+		return;
+	case 1:
+		return;
+	case 2:
+		if (Predicate(First[1], First[0]))
+		{
+			Swap(First[1], First[0]);
+		}
+		return;
+	case 3:
+		if (!Predicate(First[1], First[0]))
+		{
+			// 0 < 1
+			if (Predicate(First[2], First[1]))
+			{
+				// 0 < 1 and 2 < 1
+				if (Predicate(First[2], First[0]))
+				{
+					// 2 < 0 < 1
+					T Temp = MoveTemp(First[0]);
+					First[0] = MoveTemp(First[2]);
+					First[2] = MoveTemp(First[1]);
+					First[1] = MoveTemp(Temp);
+				}
+				else
+				{
+					// 0 < 2 < 1
+					Swap(First[2], First[1]);
+				}
+			}
+			// else Already sorted 0 < 1 < 2
+		}
+		else
+		{
+			// 1 < 0
+			if (Predicate(First[2], First[0]))
+			{
+				// 1 < 0 and 2 < 0
+				if (Predicate(First[2], First[1]))
+				{
+					// 2 < 1 < 0
+					Swap(First[2], First[0]);
+				}
+				else
+				{
+					// 1 < 2 < 0
+					T Temp = MoveTemp(First[0]);
+					First[0] = MoveTemp(First[1]);
+					First[1] = MoveTemp(First[2]);
+					First[2] = MoveTemp(Temp);
+				}
+			}
+			else
+			{
+				// 1 < 0 < 2
+				Swap(First[1], First[0]);
+			}
+		}
+		return;
+	default:
+		checkf(false, TEXT("ParallelSort bug: SmallSort was passed Num > 3"));
+		return;
+	}
+}
+
+} // namespace Algo
\ No newline at end of file
diff --git a/Engine/Source/Runtime/AssetRegistry/Private/AssetRegistryState.cpp b/Engine/Source/Runtime/AssetRegistry/Private/AssetRegistryState.cpp
--- a/Engine/Source/Runtime/AssetRegistry/Private/AssetRegistryState.cpp
+++ b/Engine/Source/Runtime/AssetRegistry/Private/AssetRegistryState.cpp
@@ -3,6 +3,8 @@
 #include "AssetRegistry/AssetRegistryState.h"
 
 #include "Algo/Compare.h"
+#include "AssetRegistry/LinearBlockAllocator.h"
+#include "AssetRegistry/ParallelSort.h"
 #include "Algo/Sort.h"
 #include "Algo/Unique.h"
 #include "AssetRegistry/ARFilter.h"
@@ -1844,6 +1846,117 @@ bool FAssetRegistryState::Serialize(FArchive& Ar, const FAssetRegistrySerializat
 	return Ar.IsSaving() ? Save(Ar, Options) : Load(Ar, FAssetRegistryLoadOptions(Options));
 }
 
+static UTF8CHAR* ToLowerStrNDup(UE::FLinearBlockAllocatorThreadAccessor& Allocator, const UTF8CHAR* Src, int32 Len)
+{
+	if (!Src || Len == 0)
+	{
+		return nullptr;
+	}
+
+	UTF8CHAR* Dest = reinterpret_cast<UTF8CHAR*>(Allocator.Malloc((Len + 1) * sizeof(UTF8CHAR), alignof(UTF8CHAR)));
+	for (int32 N = 0; N < Len; ++N)
+	{
+		Dest[N] = TChar<UTF8CHAR>::ToLower(Src[N]);
+	}
+	Dest[Len] = '\0';
+	return Dest;
+}
+
+TArray<FAssetData*> FAssetRegistryState::GetMutableAssetsSortedByObjectPath()
+{
+	struct FSortData
+	{
+		FAssetData* AssetData;
+		// Components of AssetData.GetSoftObjectPath, separated for faster sorting
+		UTF8CHAR* PackageName;
+		FName AssetName;
+		UTF8CHAR* SubPath;
+	};
+
+	TArray<FSortData> SortDatas;
+	SortDatas.Reserve(CachedAssets.Num());
+	EnumerateAllMutableAssets([&SortDatas](FAssetData& AssetData)
+		{
+			SortDatas.Emplace_GetRef().AssetData = &AssetData;
+		});
+
+	UE::FLinearBlockAllocator SharedAllocator;
+	TArray<UE::FLinearBlockAllocatorThreadAccessor> ThreadAllocators;
+	// Set minbatchsize to avoid wasting memory by having threads allocate a 64K block and then not use most of it
+	constexpr int32 MinSetupBatchSize = 500;
+	int32 CachedAssetsMaxIndex = CachedAssets.GetMaxIndex();
+	ParallelForWithTaskContext(TEXT("GetMutableAssetsSortedByObjectPath"), ThreadAllocators, SortDatas.Num(),
+		MinSetupBatchSize,
+		// ContextConstructor
+		[&SharedAllocator](int32 ThreadIndex, int32 NumThreads) -> UE::FLinearBlockAllocator&
+		{
+			return SharedAllocator;
+		},
+		// Body
+		[this, &SortDatas](UE::FLinearBlockAllocatorThreadAccessor& Allocator, int32 Index)
+		{
+			FSortData& Data = SortDatas[Index];
+			FSoftObjectPath AssetId = Data.AssetData->GetSoftObjectPath();
+
+			TUtf8StringBuilder<1024> PackageName(InPlace, AssetId.GetLongPackageFName());
+			Data.PackageName = ToLowerStrNDup(Allocator, PackageName.GetData(), PackageName.Len());
+			Data.AssetName = AssetId.GetAssetFName();
+			const FUtf8String& SubPath = AssetId.GetSubPathUtf8String();
+			Data.SubPath = ToLowerStrNDup(Allocator, GetData(SubPath), SubPath.Len());
+		}
+	);
+
+	auto CompareSortData = [](const FSortData& A, const FSortData& B)
+		{
+			int32 PackageNameDiff = TCString<UTF8CHAR>::Strcmp(A.PackageName, B.PackageName);
+			if (PackageNameDiff != 0)
+			{
+				return PackageNameDiff < 0;
+			}
+			int32 AssetNameDiff = A.AssetName.Compare(B.AssetName);
+			if (AssetNameDiff != 0)
+			{
+				return AssetNameDiff < 0;
+			}
+			if ((A.SubPath == nullptr) != (B.SubPath == nullptr))
+			{
+				return A.SubPath == nullptr;
+			}
+			if (A.SubPath)
+			{
+				int32 SubPathDiff = TCString<UTF8CHAR>::Strcmp(A.SubPath, B.SubPath);
+				if (SubPathDiff != 0)
+				{
+					return SubPathDiff < 0;
+				}
+			}
+			return false;
+		};
+	Algo::ParallelSort(SortDatas, CompareSortData);
+	// TODO: Temporary validation of ParallelSort, Remove after a few days
+	for (int32 Index = 0; Index < SortDatas.Num() - 1; ++Index)
+	{
+		if (CompareSortData(SortDatas[Index+1], SortDatas[Index]))
+		{
+			UE_LOG(LogAssetRegistry, Warning,
+				TEXT("Algo::ParallelSort failed to sort correctly. Falling back to Algo::Sort.\n")
+				TEXT("Failed at %d: '%s' > '%s'"),
+				Index, *SortDatas[Index].AssetData->GetSoftObjectPath().ToString(),
+				 *SortDatas[Index+1].AssetData->GetSoftObjectPath().ToString());
+			Algo::Sort(SortDatas, CompareSortData);
+			break;
+		}
+	}	
+
+	TArray<FAssetData*> Result;
+	Result.Reserve(SortDatas.Num());
+	for (const FSortData& SortData : SortDatas)
+	{
+		Result.Add(SortData.AssetData);
+	}
+	return Result;
+}
+
 bool FAssetRegistryState::Save(FArchive& OriginalAr, const FAssetRegistrySerializationOptions& Options)
 {
 	SCOPED_BOOT_TIMING("FAssetRegistryState::Save");
@@ -1869,23 +1982,12 @@ bool FAssetRegistryState::Save(FArchive& OriginalAr, const FAssetRegistrySeriali
 
 	// Write asset data first
 	{
-		TArray<TPair<FAssetData*, FSoftObjectPath>> SortedAssetsByObjectPath;
-		SortedAssetsByObjectPath.Reserve(AssetCount);
-		EnumerateAllMutableAssets([&SortedAssetsByObjectPath](FAssetData& AssetData)
-			{
-				SortedAssetsByObjectPath.Add({ &AssetData, AssetData.GetSoftObjectPath() });
-			});
-		Algo::Sort(SortedAssetsByObjectPath, [](const TPair<FAssetData*, FSoftObjectPath>& A,
-			const TPair<FAssetData*, FSoftObjectPath>& B)
-		{
-			return A.Value.LexicalLess(B.Value);
-		});
-
-		for (TPair<FAssetData*, FSoftObjectPath>& Asset : SortedAssetsByObjectPath)
+		TArray<FAssetData*> SortedAssetsByObjectPath = GetMutableAssetsSortedByObjectPath();
+		for (FAssetData* Asset : SortedAssetsByObjectPath)
 		{
 			// Hardcoding FAssetRegistryVersion::LatestVersion here so that branches can get optimized out in
 			// the forceinlined SerializeForCache
-			Asset.Key->SerializeForCache(Ar);
+			Asset->SerializeForCache(Ar);
 		}
 	}
 
diff --git a/Engine/Source/Runtime/AssetRegistry/Public/AssetRegistry/AssetRegistryState.h b/Engine/Source/Runtime/AssetRegistry/Public/AssetRegistry/AssetRegistryState.h
--- a/Engine/Source/Runtime/AssetRegistry/Public/AssetRegistry/AssetRegistryState.h
+++ b/Engine/Source/Runtime/AssetRegistry/Public/AssetRegistry/AssetRegistryState.h
@@ -572,6 +572,14 @@ public:
 private:
 	// Accessors of mutable FAssetData*, for friend classes
 	ASSETREGISTRY_API void EnumerateAllMutableAssets(TFunctionRef<void(FAssetData&)> Callback) const;
+	/**
+	 * Copy the results of EnumerateAllMutableAssets into an array of pointers and sort it by the FAssetData's
+	 * objectpaths. Do not use this unless you need them sorted; it is wasteful to sort. Use EnumerateAllMutableAssets
+	 * instead, and make your own array. Uses ParallelSort, so do not call this function within functions that are
+	 * already running in parallel.
+	 */
+	ASSETREGISTRY_API TArray<FAssetData*> GetMutableAssetsSortedByObjectPath();
+
 	FAssetData* GetMutableAssetByObjectPath(const FSoftObjectPath& ObjectPath) const;
 	FAssetData* GetMutableAssetByObjectPath(const UE::AssetRegistry::Private::FCachedAssetKey& Key) const;
 	void EnumerateMutableAssetsByPackageName(const FName PackageName,
diff --git a/Engine/Source/Runtime/AssetRegistry/Tests/ParallelSortTests.cpp b/Engine/Source/Runtime/AssetRegistry/Tests/ParallelSortTests.cpp
new file mode 100644
--- /dev/null
+++ b/Engine/Source/Runtime/AssetRegistry/Tests/ParallelSortTests.cpp
@@ -0,0 +1,139 @@
+// Copyright Epic Games, Inc. All Rights Reserved.
+
+#if WITH_LOW_LEVEL_TESTS
+
+#include "Algo/IsSorted.h"
+#include "AssetRegistry/ParallelSort.h"
+#include "TestHarness.h"
+#include "TestMacros/Assertions.h"
+
+TEST_CASE("Algo::ParallelSort")
+{
+	struct FKeyAndValue
+	{
+		int32 Key;
+		int32 Value;
+	};
+	auto CompareKeyAndValueKey = [](const FKeyAndValue& A, const FKeyAndValue& B)
+		{
+			return A.Key < B.Key;
+		};
+	auto CompareKeyAndValueFull = [](const FKeyAndValue& A, const FKeyAndValue& B)
+		{
+			int32 Diff = A.Key - B.Key;
+			if (Diff != 0)
+			{
+				return Diff < 0;
+			}
+			return A.Value < B.Value;
+		};
+
+	// Small test cases; these do not test parallel code because it falls back to single threaded for small cases.
+	{
+		TArray<int32> TestCase({ 1,2,3,4,5 });
+		Algo::ParallelSort(TestCase);
+		CHECK(Algo::IsSorted(TestCase));
+	}
+	{
+		TArray<int32> TestCase({ 5,4,3,2,1 });
+		Algo::ParallelSort(TestCase);
+		CHECK(Algo::IsSorted(TestCase));
+	}
+	{
+		TArray<int32> TestCase({ 1,5,2,4,3 });
+		Algo::ParallelSort(TestCase);
+		CHECK(Algo::IsSorted(TestCase));
+	}
+	{
+		// Verify the sort is stable
+		TArray<FKeyAndValue> TestCase({ { 1,1 }, { 2,2 }, { 3,3 }, { 4,4 }, { 5,5 },
+			{ 1,6 }, { 2,7 }, { 3,8 }, { 4,9 }, { 5,10 } });
+		Algo::ParallelSort(TestCase, CompareKeyAndValueKey);
+		CHECK(Algo::IsSorted(TestCase, CompareKeyAndValueFull));
+	}
+
+
+	// Test cases large enough to test the parallel code
+	{
+		// Already sorted
+		TArray<int32> TestCase;
+		TestCase.SetNum(1024);
+		for (int32 n = 0; n < TestCase.Num(); ++n)
+		{
+			TestCase[n] = n;
+		}
+		Algo::ParallelSort(TestCase);
+		CHECK(Algo::IsSorted(TestCase));
+	}
+	{
+		// Backwards sorted
+		TArray<int32> TestCase;
+		TestCase.SetNum(1024);
+		for (int32 n = 0; n < TestCase.Num(); ++n)
+		{
+			TestCase[n] = -n;
+		}
+		Algo::ParallelSort(TestCase);
+		CHECK(Algo::IsSorted(TestCase));
+	}
+	{
+		// Parabola with origin halfway through list
+		TArray<int32> TestCase;
+		TestCase.SetNum(1024);
+		for (int32 n = 0; n < TestCase.Num(); ++n)
+		{
+			int32 Recentered = n - TestCase.Num() / 2;
+			TestCase[n] = Recentered*Recentered;
+		}
+		Algo::ParallelSort(TestCase);
+		CHECK(Algo::IsSorted(TestCase));
+	}
+	{
+		// Saw tooth with small period
+		TArray<int32> TestCase;
+		TestCase.SetNum(1024);
+		for (int32 n = 0; n < TestCase.Num(); ++n)
+		{
+			TestCase[n] = n % 5;
+		}
+		Algo::ParallelSort(TestCase);
+		CHECK(Algo::IsSorted(TestCase));
+	}
+	{
+		// Saw tooth with big period
+		TArray<int32> TestCase;
+		TestCase.SetNum(1024);
+		for (int32 n = 0; n < TestCase.Num(); ++n)
+		{
+			TestCase[n] = n % 100;
+		}
+		Algo::ParallelSort(TestCase);
+		CHECK(Algo::IsSorted(TestCase));
+	}
+	{
+		// Verify the sort is stable using a saw tooth with small period
+		TArray<FKeyAndValue> TestCase;
+		TestCase.SetNum(1024);
+		for (int32 n = 0; n < TestCase.Num(); ++n)
+		{
+			TestCase[n] = { n % 5, n };
+		}
+		Algo::ParallelSort(TestCase, CompareKeyAndValueKey);
+		CHECK(Algo::IsSorted(TestCase, CompareKeyAndValueFull));
+	}
+	{
+		// Verify that the sort works in the non-power-of-two length case
+		TArray<int32> TestCase;
+		TestCase.SetNum(999);
+		for (int32 n = 0; n < TestCase.Num(); ++n)
+		{
+			TestCase[n] = n % 5;
+		}
+		Algo::ParallelSort(TestCase);
+		CHECK(Algo::IsSorted(TestCase));
+	}
+
+
+}
+
+#endif
\ No newline at end of file
