AssetRegistry: Fix crash in editor -game: fix typo in check statement with reversed value for variable.
#rnx

diff --git a/Engine/Source/Runtime/AssetRegistry/Private/AssetRegistry.cpp b/Engine/Source/Runtime/AssetRegistry/Private/AssetRegistry.cpp
--- a/Engine/Source/Runtime/AssetRegistry/Private/AssetRegistry.cpp
+++ b/Engine/Source/Runtime/AssetRegistry/Private/AssetRegistry.cpp
@@ -6396,7 +6396,7 @@ FScanPathContext::FScanPathContext(FEventContext& InEventContext, FClassInherita
 void FAssetRegistryImpl::ScanPathsSynchronous(FInterfaceWriteScopeLock* ScopeLock, Impl::FScanPathContext& Context)
 {
 	// We always expect a valid lock to be provided except before initialization is complete.
-	check(ScopeLock || bInitializationComplete);
+	check(ScopeLock || !bInitializationComplete);
 	LLM_SCOPE(ELLMTag::AssetRegistry);
 
 	if (!TryConstructGathererIfNeeded())
