AssetRegistry: Add flag EEnumerateAssetsFlags::Parallel for EnumerateAllAssets. Currently only implemented for the enumeration of OnDisk assets.
#rnx

diff --git a/Engine/Source/Runtime/AssetRegistry/Private/AssetRegistry.cpp b/Engine/Source/Runtime/AssetRegistry/Private/AssetRegistry.cpp
--- a/Engine/Source/Runtime/AssetRegistry/Private/AssetRegistry.cpp
+++ b/Engine/Source/Runtime/AssetRegistry/Private/AssetRegistry.cpp
@@ -2986,9 +2986,11 @@ bool UAssetRegistryImpl::GetAllAssets(TArray<FAssetData>& OutAssetData, bool bIn
 
 bool UAssetRegistryImpl::EnumerateAllAssets(TFunctionRef<bool(const FAssetData&)> Callback, UE::AssetRegistry::EEnumerateAssetsFlags InEnumerateFlags) const
 {
+	using namespace UE::AssetRegistry;
+
 	TSet<FName> PackageNamesToSkip;
 	// All in memory assets
-	if (!EnumHasAnyFlags(InEnumerateFlags, UE::AssetRegistry::EEnumerateAssetsFlags::OnlyOnDiskAssets))
+	if (!EnumHasAnyFlags(InEnumerateFlags, EEnumerateAssetsFlags::OnlyOnDiskAssets))
 	{
 		bool bStopIteration;
 		UE::AssetRegistry::Utils::EnumerateAllMemoryAssets(PackageNamesToSkip, bStopIteration,
@@ -3006,6 +3008,7 @@ bool UAssetRegistryImpl::EnumerateAllAssets(TFunctionRef<bool(const FAssetData&)
 	TArray<FAssetData, TInlineAllocator<128>> OnDiskAssetDatas;
 	{
 		UE::AssetRegistry::FInterfaceReadScopeLock InterfaceScopeLock(InterfaceLock);
+		OnDiskAssetDatas.Reserve(GuardedData.State.GetNumAssets());
 		GuardedData.EnumerateAllDiskAssets(PackageNamesToSkip,
 			[&OnDiskAssetDatas](const FAssetData& AssetData)
 			{
@@ -3014,11 +3017,29 @@ bool UAssetRegistryImpl::EnumerateAllAssets(TFunctionRef<bool(const FAssetData&)
 			}, InEnumerateFlags);
 	}
 
-	for (const FAssetData& AssetData : OnDiskAssetDatas)
+	if (EnumHasAnyFlags(InEnumerateFlags, EEnumerateAssetsFlags::Parallel))
 	{
-		if (!Callback(AssetData))
+		ParallelFor(OnDiskAssetDatas.Num(), [&OnDiskAssetDatas, &Callback](int32 Index)
+			{
+				const FAssetData& AssetData = OnDiskAssetDatas[Index];
+				// The Callback returns bContinueIterating, and it would be nice to have a contract that we never
+				// call the Callback again after it returns bContinueIterating. But guaranteeing that contract in
+				// multithreaded code is not possible without taking a lock which is prohibitively expensive. So
+				// for the parallel version of EnumerateAllAssets, we require that the caller handle the Callback
+				// being called again even after it returns false. Since we have that requirement, and we don't
+				// expect it to be important for performance to early exit as soon as possible, we omit early exiting
+				// entirely to get a performance benefit from not needing to frequently test a contended atomic.
+				(void)Callback(AssetData);
+			});
+	}
+	else
+	{
+		for (const FAssetData& AssetData : OnDiskAssetDatas)
 		{
-			return true;
+			if (!Callback(AssetData))
+			{
+				return true;
+			}
 		}
 	}
 	return true;
diff --git a/Engine/Source/Runtime/AssetRegistry/Private/AssetRegistryState.cpp b/Engine/Source/Runtime/AssetRegistry/Private/AssetRegistryState.cpp
--- a/Engine/Source/Runtime/AssetRegistry/Private/AssetRegistryState.cpp
+++ b/Engine/Source/Runtime/AssetRegistry/Private/AssetRegistryState.cpp
@@ -1530,20 +1530,82 @@ bool FAssetRegistryState::GetAllAssets(const TSet<FName>& PackageNamesToSkip, TA
 	EnumerateFlags);
 }
 
-void FAssetRegistryState::EnumerateAllAssets(TFunctionRef<void(const FAssetData&)> Callback) const
+void FAssetRegistryState::EnumerateAllAssets(TFunctionRef<void(const FAssetData&)> Callback,
+	UE::AssetRegistry::EEnumerateAssetsFlags InEnumerateFlags) const
 {
 	EnumerateAllMutableAssets([&Callback](FAssetData& AssetData)
 		{
 			Callback(AssetData);
-		});
+		}, InEnumerateFlags);
 }
 
-void FAssetRegistryState::EnumerateAllMutableAssets(TFunctionRef<void(FAssetData&)> Callback) const
+void FAssetRegistryState::EnumerateAllMutableAssets(TFunctionRef<void(FAssetData&)> Callback,
+	UE::AssetRegistry::EEnumerateAssetsFlags InEnumerateFlags) const
 {
-	for (FAssetData* AssetData : CachedAssets)
+	using namespace UE::AssetRegistry;
+
+	bool bAllowUnfilteredAssets = EnumHasAnyFlags(InEnumerateFlags, EEnumerateAssetsFlags::AllowUnfilteredArAssets);
+	bool bAllowUnmountedPaths = EnumHasAnyFlags(InEnumerateFlags, EEnumerateAssetsFlags::AllowUnmountedPaths);
+	if (bAllowUnfilteredAssets && bAllowUnmountedPaths)
 	{
-		check(AssetData);
-		Callback(*AssetData);
+		if (EnumHasAnyFlags(InEnumerateFlags, EEnumerateAssetsFlags::Parallel))
+		{
+			int32 MaxId = CachedAssets.GetMaxIndex();
+			ParallelFor(MaxId, [this, &Callback](int32 Index)
+				{
+					FSetElementId Id = FSetElementId::FromInteger(Index);
+					if (CachedAssets.IsValidId(Id))
+					{
+						FAssetData* AssetData = CachedAssets.Get(Id);
+						check(AssetData);
+						Callback(*AssetData);
+					}
+				});
+		}
+		else
+		{
+			for (FAssetData* AssetData : CachedAssets)
+			{
+				check(AssetData);
+				Callback(*AssetData);
+			}
+		}
+	}
+	else
+	{
+		if (EnumHasAnyFlags(InEnumerateFlags, EEnumerateAssetsFlags::Parallel))
+		{
+			int32 MaxId = CachedAssets.GetMaxIndex();
+			ParallelFor(MaxId, [this, &Callback, InEnumerateFlags, bAllowUnfilteredAssets, bAllowUnmountedPaths]
+			(int32 Index)
+				{
+					FSetElementId Id = FSetElementId::FromInteger(Index);
+					if (CachedAssets.IsValidId(Id))
+					{
+						FAssetData* AssetData = CachedAssets.Get(Id);
+						check(AssetData);
+						if ((bAllowUnmountedPaths || !IsPackageUnmountedAndFiltered(AssetData->PackageName))
+							&& (bAllowUnfilteredAssets
+								|| !FFiltering::ShouldSkipAsset(AssetData->AssetClassPath, AssetData->PackageFlags)))
+						{
+							Callback(*AssetData);
+						}
+					}
+				});
+		}
+		else
+		{
+			for (FAssetData* AssetData : CachedAssets)
+			{
+				check(AssetData);
+				if ((bAllowUnmountedPaths || !IsPackageUnmountedAndFiltered(AssetData->PackageName))
+					&& (bAllowUnfilteredAssets
+						|| !FFiltering::ShouldSkipAsset(AssetData->AssetClassPath, AssetData->PackageFlags)))
+				{
+					Callback(*AssetData);
+				}
+			}
+		}
 	}
 }
 
@@ -1553,14 +1615,9 @@ bool FAssetRegistryState::EnumerateAllAssets(const TSet<FName>& PackageNamesToSk
 {
 	using namespace UE::AssetRegistry;
 
-	EnumerateAllMutableAssets([&PackageNamesToSkip, &Callback, InEnumerateFlags, this](const FAssetData& AssetData)
+	EnumerateAllMutableAssets([&PackageNamesToSkip, &Callback, this](const FAssetData& AssetData)
 		{
-			if (!PackageNamesToSkip.Contains(AssetData.PackageName)
-				&& (EnumHasAnyFlags(InEnumerateFlags, EEnumerateAssetsFlags::AllowUnmountedPaths)
-					|| !IsPackageUnmountedAndFiltered(AssetData.PackageName))
-				&& (EnumHasAnyFlags(InEnumerateFlags, EEnumerateAssetsFlags::AllowUnfilteredArAssets)
-					|| !FFiltering::ShouldSkipAsset(AssetData.AssetClassPath,
-						AssetData.PackageFlags)))
+			if (!PackageNamesToSkip.Contains(AssetData.PackageName))
 			{
 				if (!Callback(AssetData))
 				{
@@ -1568,7 +1625,7 @@ bool FAssetRegistryState::EnumerateAllAssets(const TSet<FName>& PackageNamesToSk
 				}
 			}
 			return true; // keep iterating
-		});
+		}, InEnumerateFlags);
 	return true;
 }
 
diff --git a/Engine/Source/Runtime/AssetRegistry/Public/AssetRegistry/AssetRegistryState.h b/Engine/Source/Runtime/AssetRegistry/Public/AssetRegistry/AssetRegistryState.h
--- a/Engine/Source/Runtime/AssetRegistry/Public/AssetRegistry/AssetRegistryState.h
+++ b/Engine/Source/Runtime/AssetRegistry/Public/AssetRegistry/AssetRegistryState.h
@@ -182,6 +182,7 @@ public:
 	ASSETREGISTRY_API FAssetRegistryState& operator=(FAssetRegistryState&& O);
 
 	using EInitializationMode = UE::AssetRegistry::EAppendMode;
+	using EEnumerateAssetsFlags = UE::AssetRegistry::EEnumerateAssetsFlags;
 
 	/**
 	 * Does the given path contain assets?
@@ -212,12 +213,12 @@ public:
 	 * @param Filter filter to apply to the assets in the AssetRegistry
 	 * @param PackageNamesToSkip explicit list of packages to skip, because they were already added
 	 * @param Callback function to call for each asset data enumerated
-	 * @param InEnumerateFlags flags to control enumeration and filtering.
-	 *        @see EEnumerateAssetsFlags.
+	 * @param InEnumerateFlags flags to control enumeration and filtering. Note that the default
+	 *        does not include AllowUnmountedPaths; they will be skipped by default. @see EEnumerateAssetsFlags.
 	 */
 	ASSETREGISTRY_API bool EnumerateAssets(const FARCompiledFilter& Filter, const TSet<FName>& PackageNamesToSkip, 
 		TFunctionRef<bool(const FAssetData&)> Callback,
-		UE::AssetRegistry::EEnumerateAssetsFlags InEnumerateFlags = UE::AssetRegistry::EEnumerateAssetsFlags::AllowUnfilteredArAssets) const;
+		EEnumerateAssetsFlags InEnumerateFlags = EEnumerateAssetsFlags::AllowUnfilteredArAssets) const;
 
 	/**
 	 * Gets asset data for all assets in the registry state.
@@ -232,16 +233,26 @@ public:
 	/**
 	 * Enumerates asset data for all assets in the registry state.
 	 *
+	 * @param Callback function to call for each asset data enumerated
+	 * @param InEnumerateFlags flags to control enumeration and filtering. The default includes
+	 *        all AllowFlags and no assets will be skipped by default. @see EEnumerateAssetsFlags.
+	 */
+	ASSETREGISTRY_API void EnumerateAllAssets(TFunctionRef<void(const FAssetData&)> Callback,
+		EEnumerateAssetsFlags InEnumerateFlags = EEnumerateAssetsFlags::AllowAll) const;
+
+	/**
+	 * Enumerates asset data for all assets in the registry state, except for assets in packages that are
+	 * in PackageNamesToSkip. Used by the AssetRegistry for scans after an in-memory scan has already been
+	 * completed.
+	 *
 	 * @param PackageNamesToSkip explicit list of packages to skip, because they were already added
 	 * @param Callback function to call for each asset data enumerated
-	 * @param InEnumerateFlags flags to control enumeration and filtering.
-	 *        @see EEnumerateAssetsFlags.
+	 * @param InEnumerateFlags flags to control enumeration and filtering. Note that the default
+	 *        does not include AllowUnmountedPaths; they will be skipped by default. @see EEnumerateAssetsFlags.
 	 */
 	ASSETREGISTRY_API bool EnumerateAllAssets(const TSet<FName>& PackageNamesToSkip,
 		TFunctionRef<bool(const FAssetData&)> Callback,
-		UE::AssetRegistry::EEnumerateAssetsFlags InEnumerateFlags = UE::AssetRegistry::EEnumerateAssetsFlags::AllowUnfilteredArAssets) const;
-	ASSETREGISTRY_API void EnumerateAllAssets(TFunctionRef<void(const FAssetData&)> Callback) const;
-
+		EEnumerateAssetsFlags InEnumerateFlags = EEnumerateAssetsFlags::AllowUnfilteredArAssets) const;
 
 	/**
 	 * Calls the callback with the LongPackageName of each path that has assets as direct children.
@@ -571,7 +582,8 @@ public:
 
 private:
 	// Accessors of mutable FAssetData*, for friend classes
-	ASSETREGISTRY_API void EnumerateAllMutableAssets(TFunctionRef<void(FAssetData&)> Callback) const;
+	ASSETREGISTRY_API void EnumerateAllMutableAssets(TFunctionRef<void(FAssetData&)> Callback,
+		EEnumerateAssetsFlags InEnumerateFlags = EEnumerateAssetsFlags::AllowAll) const;
 	/**
 	 * Copy the results of EnumerateAllMutableAssets into an array of pointers and sort it by the FAssetData's
 	 * objectpaths. Do not use this unless you need them sorted; it is wasteful to sort. Use EnumerateAllMutableAssets
diff --git a/Engine/Source/Runtime/AssetRegistry/Public/AssetRegistry/IAssetRegistry.h b/Engine/Source/Runtime/AssetRegistry/Public/AssetRegistry/IAssetRegistry.h
--- a/Engine/Source/Runtime/AssetRegistry/Public/AssetRegistry/IAssetRegistry.h
+++ b/Engine/Source/Runtime/AssetRegistry/Public/AssetRegistry/IAssetRegistry.h
@@ -377,6 +377,8 @@ public:
 	 * @param Callback function to call for each asset data enumerated
 	 * @param InEnumerateFlags flags to control enumeration and filtering.
 	 *        @see EEnumerateAssetsFlags.
+	 *        EEnumerateAssetsFlags::Parallel note: in-memory assets enumeration ignores the parallel
+	 *        flag and is always single-threaded.
 	 * @return False if filter is invalid, otherwise true.
 	 */
 	virtual bool EnumerateAssets(const FARFilter& Filter, TFunctionRef<bool(const FAssetData&)> Callback,
@@ -461,6 +463,8 @@ public:
 	 * @param Callback function to call for each asset data enumerated
 	 * @param InEnumerateFlags flags to control enumeration and filtering.
 	 *        @see IAssetRegistry class header for EEnumerateAssetsFlags.
+	 *        EEnumerateAssetsFlags::Parallel note: in-memory assets enumeration ignores the parallel
+	 *        flag and is always single-threaded.
 	 */
 	virtual bool EnumerateAllAssets(TFunctionRef<bool(const FAssetData&)> Callback,
 		UE::AssetRegistry::EEnumerateAssetsFlags InEnumerateFlags = UE::AssetRegistry::EEnumerateAssetsFlags::None) const = 0;
diff --git a/Engine/Source/Runtime/CoreUObject/Public/Misc/AssetRegistryInterface.h b/Engine/Source/Runtime/CoreUObject/Public/Misc/AssetRegistryInterface.h
--- a/Engine/Source/Runtime/CoreUObject/Public/Misc/AssetRegistryInterface.h
+++ b/Engine/Source/Runtime/CoreUObject/Public/Misc/AssetRegistryInterface.h
@@ -209,10 +209,34 @@ namespace UE::AssetRegistry
 
 	enum class EEnumerateAssetsFlags : uint32
 	{
-		None = 0,							// No flags
-		OnlyOnDiskAssets = (1 << 0),		// should only assets on disk be included in the enumeration. When set only DiskGatheredData will be used, does not calculate from UObjects.
-		AllowUnmountedPaths = (1 << 1),		// should unmounted asset paths be allowed
-		AllowUnfilteredArAssets = (1 << 2),	// skip the filtering of UE::AssetRegistry::FFiltering
+		/**
+		 * No flags. This implies a request to filter by some properties, because filtering is the default state that
+		 * that is turned off by some flags; see e.g. AllowUnfilteredAssets.
+		 **/
+		None = 0,
+		/** If set, only DiskGatheredData are returned, AssetDatas are not created from in-memory UObjects. */
+		OnlyOnDiskAssets = (1 << 0),
+		/**
+		 * If set, unmounted asset paths (!FPackageName::IsValidPath) are returned.
+		 * When not set (the default), they are skipped.
+		 */
+		AllowUnmountedPaths = (1 << 1),
+		/**
+		 * If set, asset paths are returned even if UE::AssetRegistry::FFiltering::ShouldSkipAsset is true for them.
+		 * When not set (the default), they are skipped.
+		 */
+		AllowUnfilteredArAssets = (1 << 2),
+		/**
+		 * All allowflags are set, no discovered assets are filtered. Does not impact memory vs OnDiskAssets; in-memory
+		 * assets may still be skipped even if this flag is set, depending on whether OnlyOnDiskAssets is set.
+		 */
+		AllowAll = AllowUnmountedPaths | AllowUnfilteredArAssets,
+		/**
+		 * Enumeration of the Assets is done with ParallelFor; the callback must be threadsafe. This flag is a request
+		 * only, and some functions may not support it or may ignore it if not performant; caller must behave correctly
+		 * if callbacks are called singlethreadedly.
+		 */
+		Parallel = (1 << 3),
 	};
 	ENUM_CLASS_FLAGS(EEnumerateAssetsFlags);
 	
