When appending an asset registry state to another. Use a more optimized tag serialization method for new assets.
This way the newly added tags takes up less memory and is more optimized.
This is the same tag format that is used when serializing the asset registry from the cooker and loaded as part of the LoadPremadeAssetRegistry.

Adding supporting function:
FixedTagPrivate::InitializeFromBuilder
To more easily go from a StoreBuilder to the in memory FStore.

diff --git a/Engine/Source/Runtime/AssetRegistry/Private/AssetRegistryState.cpp b/Engine/Source/Runtime/AssetRegistry/Private/AssetRegistryState.cpp
--- a/Engine/Source/Runtime/AssetRegistry/Private/AssetRegistryState.cpp
+++ b/Engine/Source/Runtime/AssetRegistry/Private/AssetRegistryState.cpp
@@ -427,6 +427,10 @@ void FAssetRegistryState::InitializeFromExisting(const FAssetDataMap& AssetDataM
 		Reset();
 	}
 
+	// FStoreBuilder is a serialized tag builder which will allow us to store the asset tags much more efficiently.
+	TArray<TPair<FAssetData*, FixedTagPrivate::FPartialMapHandle>> AssetsToUpdate;
+	FixedTagPrivate::FStoreBuilder Builder(FAssetRegistryWriterOptions(Options).Tags);
+	
 	for (const FAssetData* AssetDataPtr : AssetDataMap)
 	{
 		if (AssetDataPtr == nullptr)
@@ -490,11 +494,19 @@ void FAssetRegistryState::InitializeFromExisting(const FAssetDataMap& AssetDataM
 			{
 				OutAppendResult->AddedAssets.Add(NewData);
 			}
+
+			AssetsToUpdate.Add(TPair<FAssetData*, FixedTagPrivate::FPartialMapHandle>(NewData, Builder.AddTagMap(NewData->TagsAndValues)));
 		}
 	}
+	// Construct the optimized tag data then set the handle map back to the assets we need to update
+	TRefCountPtr<const FixedTagPrivate::FStore> FixedStore = FixedTagPrivate::InitializeFromBuilder(Builder);
+	for (TPair<FAssetData*, FixedTagPrivate::FPartialMapHandle>& AssetToHandle : AssetsToUpdate)
+	{
+		FixedTagPrivate::FMapHandle FixedMapHandle = AssetToHandle.Value.MakeFullHandle(FixedStore->Index);
+		AssetToHandle.Key->TagsAndValues = FAssetDataTagMapSharedView(FixedMapHandle);
+	}
 
 	TSet<FAssetIdentifier> ScriptPackages;
-
 	if (InInitializationMode != EInitializationMode::OnlyUpdateExisting)
 	{
 		for (const TPair<FName, FAssetPackageData*>& Pair : AssetPackageDataMap)
diff --git a/Engine/Source/Runtime/CoreUObject/Private/AssetRegistry/AssetDataTagMap.cpp b/Engine/Source/Runtime/CoreUObject/Private/AssetRegistry/AssetDataTagMap.cpp
--- a/Engine/Source/Runtime/CoreUObject/Private/AssetRegistry/AssetDataTagMap.cpp
+++ b/Engine/Source/Runtime/CoreUObject/Private/AssetRegistry/AssetDataTagMap.cpp
@@ -11,6 +11,7 @@
 #include "Misc/ScopeLock.h"
 #include "Serialization/ArrayWriter.h"
 #include "Serialization/MemoryReader.h"
+#include "Serialization/MemoryWriter.h"
 #include <limits>
 
 DEFINE_LOG_CATEGORY_STATIC(LogAssetDataTags, Log, All);
@@ -1386,6 +1387,18 @@ namespace FixedTagPrivate
 		return TRefCountPtr<const FStore>();
 	}
 
+	TRefCountPtr<const FStore> InitializeFromBuilder(FStoreBuilder& InBuilder)
+	{
+		TArray<uint8> Data;
+		FMemoryWriter MemoryWriter(Data, true);
+		FStoreData StoreData = InBuilder.Finalize();
+
+		SaveStore(StoreData, MemoryWriter);
+
+		FMemoryReader MemoryReader(Data, true);
+		return LoadStore(MemoryReader);
+	}
+
 } // end namespace FixedTagPrivate
 
 ////////////////////////////////////////////////////////////////////////////
diff --git a/Engine/Source/Runtime/CoreUObject/Public/AssetRegistry/AssetDataTagMapSerializationDetails.h b/Engine/Source/Runtime/CoreUObject/Public/AssetRegistry/AssetDataTagMapSerializationDetails.h
--- a/Engine/Source/Runtime/CoreUObject/Public/AssetRegistry/AssetDataTagMapSerializationDetails.h
+++ b/Engine/Source/Runtime/CoreUObject/Public/AssetRegistry/AssetDataTagMapSerializationDetails.h
@@ -285,4 +285,7 @@ namespace FixedTagPrivate
 		TOptional<ELoadOrder> Order;
 	};
 
+	// Initializes a managed FStore from an unfinalized builder.
+	COREUOBJECT_API TRefCountPtr<const FStore> InitializeFromBuilder(FStoreBuilder& InBuilder);
+
 } // end namespace FixedTagPrivate
