ParallelSort: Fix the assertion failure in Merge: MergeBuffer.Max() >= Num - FloorMidIndex + 1. It was caused by a missing + 1 in one of the callsites of MergeSortRecursive.
Add unittests, and better diagnostics if the assertion happens again.
#rnx

diff --git a/Engine/Source/Runtime/AssetRegistry/Internal/AssetRegistry/ParallelSort.h b/Engine/Source/Runtime/AssetRegistry/Internal/AssetRegistry/ParallelSort.h
--- a/Engine/Source/Runtime/AssetRegistry/Internal/AssetRegistry/ParallelSort.h
+++ b/Engine/Source/Runtime/AssetRegistry/Internal/AssetRegistry/ParallelSort.h
@@ -39,15 +39,24 @@ void ParallelSortForceSingleThreaded(RangeType&& Range, PredicateType Predicate)
 
 namespace Algo::MergeSort
 {
-
+template <typename T>
+struct TMergeBuffer
+{
+	TRingBuffer<T> Buffer;
+	int32 AllocatedSize = 0;
+	int32 ConstructorSite = 0;
+	int32 InputNum = 0;
+	int32 BatchSize = 0;
+	int32 Round = 0;
+};
 template <typename T, typename IndexType, typename PredicateType>
 void ParallelMergeSortPointerAndNum(T* First, IndexType Num, PredicateType Predicate);
 template <typename T, typename IndexType, typename PredicateType>
 void MergeSortPointerAndNum(T* First, IndexType Num, PredicateType Predicate);
 template <typename T, typename IndexType, typename PredicateType>
-void MergeSortRecursive(T* First, IndexType Num, PredicateType Predicate, TRingBuffer<T>& MergeBuffer);
+void MergeSortRecursive(T* First, IndexType Num, PredicateType Predicate, TMergeBuffer<T>& MergeBuffer);
 template <typename T, typename IndexType, typename PredicateType>
-void Merge(T* First, IndexType Num, IndexType FloorMidIndex, PredicateType Predicate, TRingBuffer<T>& MergeBuffer);
+void Merge(T* First, IndexType Num, IndexType FloorMidIndex, PredicateType Predicate, TMergeBuffer<T>& MergeBuffer);
 template <typename T, typename IndexType, typename PredicateType>
 void SmallSort(T* First, IndexType Num, PredicateType Predicate);
 
@@ -141,17 +150,22 @@ void ParallelMergeSortPointerAndNum(T* First, IndexType Num, PredicateType Predi
 	// own mergebuffers and do the same.
 	// Merging two segments, each of size <= BatchSize into a merged 2*BatchSize length segment requires scratch
 	// space equal to BatchSize, (plus 1 more because we push before we pop).
-	auto GetMergeBufferForRound = [BatchSize](int32 Round)
+	auto GetMergeBufferForRound = [BatchSize, Num](int32 Round, int32 ConstructorSite)
 		{
 			// BatchSize << Round is the size of the segment produced for the round, we only need half that (plus one).
 			const IndexType BufferSize = ((BatchSize << Round) + 1) / 2 + 1;
-			TRingBuffer<T> MergeBuffer;
-			MergeBuffer.Reserve(BufferSize);
+			TMergeBuffer<T> MergeBuffer;
+			MergeBuffer.Buffer.Reserve(BufferSize);
+			MergeBuffer.AllocatedSize = BufferSize;
+			MergeBuffer.ConstructorSite = ConstructorSite;
+			MergeBuffer.InputNum = Num;
+			MergeBuffer.Round = Round;
+			MergeBuffer.BatchSize = BatchSize;
 			return MergeBuffer;
 		};
 
 	// Does just the merge job with no recursion.
-	auto MergeRoundAndIndex = [TaskCount, Num, BatchSize, &Predicate, First](int32 Round, int32 IndexInRound, TRingBuffer<T>& MergeBuffer)
+	auto MergeRoundAndIndex = [TaskCount, Num, BatchSize, &Predicate, First](int32 Round, int32 IndexInRound, TMergeBuffer<T>& MergeBuffer)
 		{
 			const int32 NumMergesForRound = TaskCount >> Round;
 			const IndexType CurrentLevelBatchSize = BatchSize << Round;
@@ -164,7 +178,7 @@ void ParallelMergeSortPointerAndNum(T* First, IndexType Num, PredicateType Predi
 
 	// The base case of recursion; call single-threaded mergesort, which does its own recursion but without any
 	// threading.
-	auto RecursiveMergeRound0 = [TaskCount, Num, BatchSize, &Predicate, First](int32 IndexInRound, TRingBuffer<T>& MergeBuffer)
+	auto RecursiveMergeRound0 = [TaskCount, Num, BatchSize, &Predicate, First](int32 IndexInRound, TMergeBuffer<T>& MergeBuffer)
 		{
 			const IndexType StartIndex = IndexInRound * BatchSize;
 			const IndexType MergeNum = IndexInRound < TaskCount - 1 ? BatchSize : Num - StartIndex;
@@ -174,17 +188,17 @@ void ParallelMergeSortPointerAndNum(T* First, IndexType Num, PredicateType Predi
 	// Our recursive function that forks itself to handle the upper recursive segment, works itself recursively on the
 	// lower recursive segment, and then merges the two segments. The arguments are the round number, and the index
 	// within the round of the segment it is creating.
-	TFunction<void(int32, int32, TRingBuffer<T>&)> RecursiveMergeRoundAndIndexPtr;
+	TFunction<void(int32, int32, TMergeBuffer<T>&)> RecursiveMergeRoundAndIndexPtr;
 	auto RecursiveMergeRoundAndIndex =
 		[&MergeRoundAndIndex, &RecursiveMergeRoundAndIndexPtr, &RecursiveMergeRound0, &GetMergeBufferForRound]
-		(int32 Round, int32 IndexInRound, TRingBuffer<T>& MergeBuffer)
+		(int32 Round, int32 IndexInRound, TMergeBuffer<T>& MergeBuffer)
 		{
 			if (Round > 0)
 			{
 				FTask Upper = Launch(TEXT("ParallelMergeSort"),
 					[&RecursiveMergeRoundAndIndexPtr, &GetMergeBufferForRound, Round, IndexInRound]()
 					{
-						TRingBuffer<T> MergeBuffer = GetMergeBufferForRound(Round - 1);
+						TMergeBuffer<T> MergeBuffer = GetMergeBufferForRound(Round - 1, 100);
 						RecursiveMergeRoundAndIndexPtr(Round - 1, 2 * IndexInRound + 1, MergeBuffer);
 					});
 
@@ -200,20 +214,26 @@ void ParallelMergeSortPointerAndNum(T* First, IndexType Num, PredicateType Predi
 	RecursiveMergeRoundAndIndexPtr = RecursiveMergeRoundAndIndex;
 
 	// Kick off RecursiveMergeRoundAndIndexPtr for the last round, at its only segment index: index 0.
-	TRingBuffer<T> MergeBuffer = GetMergeBufferForRound(NumRounds - 1);
+	TMergeBuffer<T> MergeBuffer = GetMergeBufferForRound(NumRounds - 1, 200);
 	RecursiveMergeRoundAndIndexPtr(NumRounds - 1, 0, MergeBuffer);
 }
 
 template <typename T, typename IndexType, typename PredicateType>
 void MergeSortPointerAndNum(T* First, IndexType Num, PredicateType Predicate)
 {
-	TRingBuffer<T> MergeBuffer;
-	MergeBuffer.Reserve((Num + 1) / 2);
+	TMergeBuffer<T> MergeBuffer;
+	IndexType BufferSize = (Num + 1) / 2 + 1;
+	MergeBuffer.Buffer.Reserve(BufferSize);
+	MergeBuffer.AllocatedSize = BufferSize;
+	MergeBuffer.ConstructorSite = 300;
+	MergeBuffer.InputNum = Num;
+	MergeBuffer.Round = -1;
+	MergeBuffer.BatchSize = -1;
 	MergeSortRecursive(First, Num, Predicate, MergeBuffer);
 }
 
 template <typename T, typename IndexType, typename PredicateType>
-void MergeSortRecursive(T* First, IndexType Num, PredicateType Predicate, TRingBuffer<T>& MergeBuffer)
+void MergeSortRecursive(T* First, IndexType Num, PredicateType Predicate, TMergeBuffer<T>& MergeBuffer)
 {
 	if (Num <= 3)
 	{
@@ -229,7 +249,7 @@ void MergeSortRecursive(T* First, IndexType Num, PredicateType Predicate, TRingB
 }
 
 template <typename T, typename IndexType, typename PredicateType>
-void Merge(T* First, IndexType Num, IndexType FloorMidIndex, PredicateType Predicate, TRingBuffer<T>& MergeBuffer)
+void Merge(T* First, IndexType Num, IndexType FloorMidIndex, PredicateType Predicate, TMergeBuffer<T>& InMergeBuffer)
 {
 	IndexType WriteHead = 0;
 	IndexType UpperReadHead = FloorMidIndex;
@@ -253,16 +273,16 @@ void Merge(T* First, IndexType Num, IndexType FloorMidIndex, PredicateType Predi
 		return;
 	}
 
+	TRingBuffer<T>& MergeBuffer = InMergeBuffer.Buffer;
 	MergeBuffer.Reset();
 	// Our caller is supposed to reserve the MergeBuffer for us. The maximum size of MergeBuffer that we use is
 	// the length of the Upper half, plus 1. Assert that we have that size so that we don't hit the performance
 	// cost of growing the buffer here.
-	// This assertion is currently failing, disable it until we get a fix.
-#if 0
 	checkf(MergeBuffer.Max() >= Num - FloorMidIndex + 1,
-		TEXT("ParallelSort bug: MergeBuffer not preallocated large enough. %d < %d"),
-		MergeBuffer.Max(), Num - FloorMidIndex + 1);
-#endif
+		TEXT("ParallelSort bug: MergeBuffer not preallocated large enough. %d < %d. ")
+		TEXT("AllocatedSize=%d, ConstructorSite=%d, InputNum=%d, Round=%d, BatchSize=%d."),
+		MergeBuffer.Max(), Num - FloorMidIndex + 1,
+		InMergeBuffer.AllocatedSize, InMergeBuffer.ConstructorSite, InMergeBuffer.InputNum, InMergeBuffer.Round, InMergeBuffer.BatchSize);
 
 	// Move the UpperReadHead into WriteHead, we already compared it above and know it is lower
 	MergeBuffer.Add(MoveTemp(First[WriteHead]));
diff --git a/Engine/Source/Runtime/AssetRegistry/Tests/ParallelSortTests.cpp b/Engine/Source/Runtime/AssetRegistry/Tests/ParallelSortTests.cpp
--- a/Engine/Source/Runtime/AssetRegistry/Tests/ParallelSortTests.cpp
+++ b/Engine/Source/Runtime/AssetRegistry/Tests/ParallelSortTests.cpp
@@ -122,15 +122,20 @@ TEST_CASE("Algo::ParallelSort")
 		CHECK(Algo::IsSorted(TestCase, CompareKeyAndValueFull));
 	}
 	{
-		// Verify that the sort works in the non-power-of-two length case
-		TArray<int32> TestCase;
-		TestCase.SetNum(999);
-		for (int32 n = 0; n < TestCase.Num(); ++n)
+		// Verify that the sort works on every possible size up to and including the second-smallest power
+		// of two that uses parallel sort rather than falling back to single-threaded sort.
+		for (int32 Size = 0; Size <= 512; ++Size)
 		{
-			TestCase[n] = n % 5;
+			TArray<int32> TestCase;
+			TestCase.SetNum(Size);
+			// Use a reverse sorted list so that every merge level has work to do and can't early exit.
+			for (int32 n = 0; n < TestCase.Num(); ++n)
+			{
+				TestCase[n] = -n;
+			}
+			Algo::ParallelSort(TestCase);
+			CHECK(Algo::IsSorted(TestCase));
 		}
-		Algo::ParallelSort(TestCase);
-		CHECK(Algo::IsSorted(TestCase));
 	}
 
 
