AssetRegistryState: Optimize Prune and RemoveAssetDatasImpl so that Prune is faster than InitializeFromExistingAndPrune, so we can avoid the copy costs that are necessary in InitializeFromExistingAndPrune. Optimizations included batching and ParallelFor.
#rnx

diff --git a/Engine/Source/Editor/UnrealEd/Private/Commandlets/AssetRegistryGenerator.cpp b/Engine/Source/Editor/UnrealEd/Private/Commandlets/AssetRegistryGenerator.cpp
--- a/Engine/Source/Editor/UnrealEd/Private/Commandlets/AssetRegistryGenerator.cpp
+++ b/Engine/Source/Editor/UnrealEd/Private/Commandlets/AssetRegistryGenerator.cpp
@@ -1762,18 +1762,7 @@ void FAssetRegistryGenerator::FinalizeChunkIDs(const TSet<FName>& InCookedPackag
 		FAssetRegistrySerializationOptions DevelopmentSaveOptions;
 		AssetRegistry.InitializeSerializationOptions(DevelopmentSaveOptions, TargetPlatform, UE::AssetRegistry::ESerializationTarget::ForDevelopment);
 		DevelopmentSaveOptions.bKeepDevelopmentAssetRegistryTags = true;
-
-		// Create a new FAssetRegistryState and call InitializeFromExistingAndPrune, then move the result to State.
-		// This is faster than calling State.PruneAssetData. PruneAssetData removes iteratively all the packages one by one and it is slow.
-		FAssetRegistryState PrunedAssetRegistry;
-		if (!AllPackages.IsEmpty())
-		{
-			// Only call InitializeFromExistingAndPrune if we recorded at least one package.
-			// InitializeFromExistingAndPrune does not handle the empty case because it assumes
-			// that an empty RequiredPackages means no required packages.
-			PrunedAssetRegistry.InitializeFromExistingAndPrune(State, AllPackages, TSet<FName>(), TSet<int32>(), DevelopmentSaveOptions);
-		}
-		State = MoveTemp(PrunedAssetRegistry);
+		State.PruneAssetData(AllPackages, TSet<FName>(), TSet<int32>(), DevelopmentSaveOptions);
 	}
 
 	// Development only packages should have been marked via UpdateAssetRegistryData with a negative size indicating why they were not cooked 
diff --git a/Engine/Source/Runtime/AssetRegistry/Private/AssetRegistryState.cpp b/Engine/Source/Runtime/AssetRegistry/Private/AssetRegistryState.cpp
--- a/Engine/Source/Runtime/AssetRegistry/Private/AssetRegistryState.cpp
+++ b/Engine/Source/Runtime/AssetRegistry/Private/AssetRegistryState.cpp
@@ -14,6 +14,7 @@
 #include "AssetRegistryPrivate.h"
 #include "Async/ParallelFor.h"
 #include "Blueprint/BlueprintSupport.h"
+#include "Containers/MpscQueue.h"
 #include "DependsNode.h"
 #include "GenericPlatform/GenericPlatformChunkInstall.h"
 #include "Misc/CommandLine.h"
@@ -631,125 +632,209 @@ void FAssetRegistryState::Prune(const FAssetRegistryPruneOptions& PruneOptions)
 	const bool bIsFilteredByRequiredPackages = RequiredPackages.Num() != 0;
 	const bool bIsFilteredByRemovedPackages = RemovePackages.Num() != 0;
 
-	TSet<FName> RequiredDependNodePackages;
-
-	// Generate list up front as the maps will get cleaned up
 	TArray<FAssetData*> AllAssetData = CachedAssets.Array();
-	TSet<FDependsNode*> RemoveDependsNodes;
+	int32 NumAssetDatas = AllAssetData.Num();
+	int32 NumRemovedAssetDatas = 0;
 
-	TSet<FPrimaryAssetId> KnownPrimaryAssetIds;
+	TArray<int8> KeepAssetDataIndex;
+	TArray<FDependsNode*> AllDependsNodes;
+	TMap<FDependsNode*, int8> KeepDependsNodes;
 
-	// Remove assets and mark-for-removal any dependencynodes for assets removed due to having no tags
-	for (FAssetData* AssetData : AllAssetData)
+	CachedDependsNodes.GenerateValueArray(AllDependsNodes);
+	KeepAssetDataIndex.SetNum(AllAssetData.Num());
+	KeepDependsNodes.Reserve(AllDependsNodes.Num());
+	for (FDependsNode* DependsNode : AllDependsNodes)
 	{
-		bool bRemoveAssetData = false;
-		bool bRemoveDependencyData = true;
+		KeepDependsNodes.Add(DependsNode, 0);
+	}
 
-		if (bIsFilteredByChunkId &&
-			!AssetData->GetChunkIDs().ContainsByPredicate([&](int32 ChunkId)
-				{
-					return ChunksToKeep.Contains(ChunkId);
-				}))
+	ParallelFor(AllAssetData.Num(),
+		[&AllAssetData, &KeepAssetDataIndex, &KeepDependsNodes, &NumRemovedAssetDatas,
+		this, &RequiredPackages, &RemovePackages, &ChunksToKeep, &Options, &PruneOptions,
+		bIsFilteredByChunkId, bIsFilteredByRequiredPackages, bIsFilteredByRemovedPackages
+		]
+		(int32 AssetDataIndex)
 		{
-			bRemoveAssetData = true;
-		}
-		else if (bIsFilteredByRequiredPackages && !RequiredPackages.Contains(AssetData->PackageName))
-		{
-			bRemoveAssetData = true;
-		}
-		else if (bIsFilteredByRemovedPackages && RemovePackages.Contains(AssetData->PackageName))
-		{
-			bRemoveAssetData = true;
-		}
-		else if (Options.bFilterAssetDataWithNoTags && AssetData->TagsAndValues.Num() == 0 &&
-			!FPackageName::IsLocalizedPackage(WriteToString<256>(AssetData->PackageName)) &&
-			// TODO: Add a package flag for PKG_CookGenerator and check that here as well.
-			!(AssetData->PackageFlags & PKG_CookGenerated))
-		{
-			bRemoveAssetData = true;
-			bRemoveDependencyData = Options.bFilterDependenciesWithNoTags;
-		}
+			FAssetData* AssetData = AllAssetData[AssetDataIndex];
+			// Remove assets, and remove FDependsNode for removed assets.
+			// Exception: if !bFilterDependenciesWithNoTags, then keep the FDependsNode of Assets that were removed
+			// only due to having no tags.
+			bool bKeepAssetData = true;
+			bool bKeepDependsNode = true;
 
-		if (bRemoveAssetData)
-		{
-			bool bRemovedAssetData, bRemovedPackageData;
-			FName AssetPackageName = AssetData->PackageName;
-			// AssetData might be deleted after this call
-			RemoveAssetData(AssetData, false /* bRemoveDependencyData */, bRemovedAssetData, bRemovedPackageData);
-			if (!bRemoveDependencyData)
+			if (bIsFilteredByChunkId &&
+				!AssetData->GetChunkIDs().ContainsByPredicate([&ChunksToKeep](int32 ChunkId)
+					{
+						return ChunksToKeep.Contains(ChunkId);
+					}))
 			{
-				RequiredDependNodePackages.Add(AssetPackageName);
+				bKeepAssetData = false;
+				bKeepDependsNode = false;
 			}
-			else if (bRemovedPackageData)
+			else if (bIsFilteredByRequiredPackages && !RequiredPackages.Contains(AssetData->PackageName))
 			{
-				FDependsNode** RemovedNode = CachedDependsNodes.Find(AssetPackageName);
-				if (RemovedNode)
-				{
-					RemoveDependsNodes.Add(*RemovedNode);
-				}
+				bKeepAssetData = false;
+				bKeepDependsNode = false;
 			}
-		}
-		else if (PruneOptions.bRemoveDependenciesWithoutPackages)
-		{
-			FPrimaryAssetId PrimaryAssetId = AssetData->GetPrimaryAssetId();
-			if (PrimaryAssetId.IsValid())
+			else if (bIsFilteredByRemovedPackages && RemovePackages.Contains(AssetData->PackageName))
 			{
-				KnownPrimaryAssetIds.Add(PrimaryAssetId);
+				bKeepAssetData = false;
+				bKeepDependsNode = false;
+			}
+			else if (Options.bFilterAssetDataWithNoTags && AssetData->TagsAndValues.Num() == 0 &&
+				!FPackageName::IsLocalizedPackage(WriteToString<256>(AssetData->PackageName)) &&
+				// TODO: Add a package flag for PKG_CookGenerator and check that here as well.
+				!(AssetData->PackageFlags & PKG_CookGenerated))
+			{
+				bKeepAssetData = false;
+				bKeepDependsNode = !Options.bFilterDependenciesWithNoTags;
 			}
-		}
-	}
 
-	TArray<FDependsNode*> AllDependsNodes;
-	CachedDependsNodes.GenerateValueArray(AllDependsNodes);
+			KeepAssetDataIndex[AssetDataIndex] = bKeepAssetData;
+			NumRemovedAssetDatas += bKeepAssetData ? 0 : 1;
 
-	// Mark-for-removal all other dependsnodes that are filtered out by our settings
-	for (FDependsNode* DependsNode : AllDependsNodes)
+			if (bKeepDependsNode)
+			{
+				FDependsNode** DependsNode = CachedDependsNodes.Find(AssetData->PackageName);
+				// Multiple assets can share the same dependsnode; keep the depends node if any of them say to keep it.
+				// We write atomically to prevent a technical datarace from multiple writes to the same variable; it's not
+				// an actual datarace because we always write the same value if we write anything at all.
+				if (DependsNode)
+				{
+					int8& KeepDependsNode = KeepDependsNodes.FindChecked(*DependsNode);
+					FPlatformAtomics::AtomicStore_Relaxed(&KeepDependsNode, 1);
+				}
+
+				FPrimaryAssetId PrimaryAssetId = AssetData->GetPrimaryAssetId();
+				if (PrimaryAssetId.IsValid())
+				{
+					FDependsNode** PrimaryAssetDependsNode = CachedDependsNodes.Find(PrimaryAssetId);
+					if (PrimaryAssetDependsNode)
+					{
+						int8& KeepDependsNode = KeepDependsNodes.FindChecked(*PrimaryAssetDependsNode);
+						FPlatformAtomics::AtomicStore_Relaxed(&KeepDependsNode, 1);
+					}
+				}
+			}
+
+		});
+
+	TArray<FAssetData*> RemovedAssetDatas;
+	RemovedAssetDatas.Reserve(NumRemovedAssetDatas);
+	for (int32 AssetDataIndex = 0; AssetDataIndex < NumAssetDatas; ++AssetDataIndex)
 	{
-		const FAssetIdentifier& Id = DependsNode->GetIdentifier();
-		bool bRemoveDependsNode = false;
-		if (RemoveDependsNodes.Contains(DependsNode))
+		if (!KeepAssetDataIndex[AssetDataIndex])
 		{
-			continue;
+			RemovedAssetDatas.Add(AllAssetData[AssetDataIndex]);
 		}
+	}
+	TBitArray<> UnusedRemoveAssetData;
+	TBitArray<> UnusedRemovedDependencyData;
+	RemoveAssetDatas(RemovedAssetDatas, false /* bRemoveDependencyData*/, UnusedRemoveAssetData, UnusedRemovedDependencyData);
 
-		if (Options.bFilterSearchableNames && Id.IsValue())
-		{
-			bRemoveDependsNode = true;
-		}
-		else if (Id.IsPackage() &&
-			!CachedAssetsByPackageName.Contains(Id.PackageName) &&
-			!RequiredDependNodePackages.Contains(Id.PackageName) &&
-			!FPackageName::IsScriptPackage(WriteToString<256>(Id.PackageName)))
-		{
-			bRemoveDependsNode = true;
-		}
-		else if (PruneOptions.bRemoveDependenciesWithoutPackages)
+	// Mark-for-removal all other dependsnodes that are filtered out by our settings
+	int32 MaxId = KeepDependsNodes.GetMaxIndex();
+	ParallelFor(MaxId, [&KeepDependsNodes, &Options, &PruneOptions](int32 Index)
 		{
-			const FPrimaryAssetId PrimaryAssetId = Id.GetPrimaryAssetId();
-			if (PrimaryAssetId.IsValid() && Id.IsObject())
+			FSetElementId SetId = FSetElementId::FromInteger(Index);
+			if (!KeepDependsNodes.IsValidId(SetId))
+			{
+				return;
+			}
+			TPair<FDependsNode*, int8>& KeepPair = KeepDependsNodes.Get(SetId);
+
+			FDependsNode* DependsNode = KeepPair.Key;
+			int8& KeepDependsNode = KeepPair.Value;
+			const FAssetIdentifier& Id = DependsNode->GetIdentifier();
+
+			// bKeepDependsNode is initially the value false if the dependsnode does not correspond to an AssetData
+			// (SearchableName, or it does not exist in AssetDatas) somehow. If it does correspond to an AssetData,
+			// it is true if we decided to keep the package's DependsNode based on whether we kept the AssetData and on
+			// some flags, in the loop above.
+			// Override this initial value with some further type-specific rules for keeping or discarding DependsNode.
+
+			if (Id.IsValue())
 			{
-				if (!KnownPrimaryAssetIds.Contains(PrimaryAssetId))
+				// SearchableName; KeepDependsNode is always initially false, set it instead to remove the
+				// SearchableName depends node if and only if we're removing all SearchableName depends nodes.
+				bool bKeepDependsNode = !Options.bFilterSearchableNames;
+				KeepDependsNode = bKeepDependsNode ? 1 : 0;
+			}
+			else if (Id.PrimaryAssetType.IsValid())
+			{
+				const FPrimaryAssetId PrimaryAssetId = Id.GetPrimaryAssetId();
+				if (!PrimaryAssetId.IsValid())
 				{
-					if (!PruneOptions.RemoveDependenciesWithoutPackagesKeepPrimaryAssetTypes.Contains(
-						PrimaryAssetId.PrimaryAssetType))
-					{
-						bRemoveDependsNode = true;
-					}
+					UE_LOG(LogAssetRegistry, Warning, TEXT("FAssetRegistryState::Prune: Unknown type of DependsNode: %s"), *Id.ToString());
+					KeepDependsNode = 0;
+				}
+				else
+				{
+					// PrimaryAsset. bKeepDependsNode is false if it does not have an AssetData, or if we decided to
+					// discard that AssetData above. Override this to additionally keep the asset if we are not
+					// bRemoveDependenciesWithoutPackages or if it is in
+					// RemoveDependenciesWithoutPackagesKeepPrimaryAssetTypes.
+					bool bKeepDependsNode = !PruneOptions.bRemoveDependenciesWithoutPackages ||
+						PruneOptions.RemoveDependenciesWithoutPackagesKeepPrimaryAssetTypes.Contains(
+							PrimaryAssetId.PrimaryAssetType);
+					KeepDependsNode |= bKeepDependsNode ? 1 : 0;
 				}
 			}
-		}
-		
-		if (bRemoveDependsNode)
+			else if (Id.IsPackage())
+			{
+				// Package. bKeepDependsNode is false if the package does not have an AssetData, or if we decided to
+				// discard above. Override for script packages, which do not have AssetDatas and which we always keep.
+				bool bKeepDependsNode = FPackageName::IsScriptPackage(WriteToString<256>(Id.PackageName));
+				KeepDependsNode |= bKeepDependsNode ? 1 : 0;
+			}
+			else if (!Id.PackageName.IsNone() && !Id.ObjectName.IsNone())
+			{
+				// A DependsNode to an object within a package. Invalid, ignore it.
+				KeepDependsNode = 0;
+			}
+			else
+			{
+				// Invalid depends node with all fields == NAME_None; remove it
+				KeepDependsNode = 0;
+			}
+		});
+
+	ParallelFor(MaxId, [&KeepDependsNodes, &Options, &PruneOptions](int32 Index)
 		{
-			RemoveDependsNodes.Add(DependsNode);
-		}
-	}
+			FSetElementId SetId = FSetElementId::FromInteger(Index);
+			if (!KeepDependsNodes.IsValidId(SetId))
+			{
+				return;
+			}
+			TPair<FDependsNode*, int8>& KeepPair = KeepDependsNodes.Get(SetId);
+			if (KeepPair.Value)
+			{
+				FDependsNode* DependsNode = KeepPair.Key;
+				DependsNode->RemoveLinks([&KeepDependsNodes](const FDependsNode* ExistingDependsNode)
+					{
+						return 0 == KeepDependsNodes.FindChecked(ExistingDependsNode);
+					});
 
-	// Batch-remove all of the marked-for-removal dependsnodes
-	for (FDependsNode* DependsNode : AllDependsNodes)
+				// Mark for remove any orphaned depends nodes.
+				if (DependsNode->GetConnectionCount() == 0)
+				{
+					// We are modifying a value in KeepDependsNodes, which on first consideration might be a problem
+					// because it is being read at the same time by the RemoveLinks call for other DependsNodes in this
+					// ParallelFor. But that data race is not possible, because we're only modifying a value for a
+					// node with connectionCount == 0 after removing links to nodes with keepPair.Value == false,
+					// which means none of the other DependsNode with keepPair.Value == true, which are the only ones
+					// that call RemoveLinks, could encounter it in their RemoveLinks call.
+					KeepPair.Value = 0;
+				}
+			}
+		});
+
+	// Remove all of the marked-for-removal dependsnodes; we have already removed links to them so we don't need to call
+	// the full RemoveDependsNode function.
+	for (TPair<FDependsNode*, int8>& KeepPair : KeepDependsNodes)
 	{
-		check(DependsNode != nullptr);
-		if (RemoveDependsNodes.Contains(DependsNode))
+		FDependsNode* DependsNode = KeepPair.Key;
+		if (0 == KeepPair.Value)
 		{
 			CachedDependsNodes.Remove(DependsNode->GetIdentifier());
 			NumDependsNodes--;
@@ -760,23 +845,6 @@ void FAssetRegistryState::Prune(const FAssetRegistryPruneOptions& PruneOptions)
 				delete DependsNode;
 			}
 		}
-		else
-		{
-			DependsNode->RemoveLinks([&RemoveDependsNodes](const FDependsNode* ExistingDependsNode)
-				{
-					return RemoveDependsNodes.Contains(ExistingDependsNode);
-				});
-		}
-	}
-
-	// Remove any orphaned depends nodes. This will leave cycles in but those might represent useful data
-	CachedDependsNodes.GenerateValueArray(AllDependsNodes);
-	for (FDependsNode* DependsNode : AllDependsNodes)
-	{
-		if (DependsNode->GetConnectionCount() == 0)
-		{
-			RemoveDependsNode(DependsNode->GetIdentifier());
-		}
 	}
 }
 
@@ -3295,88 +3363,180 @@ void FAssetRegistryState::RemoveAssetDatasImpl(TArrayView<FAssetData*> AssetData
 )
 {
 	using namespace UE::AssetRegistry::Private;
+	using namespace UE::Tasks;
 
-	TSet<FAssetIdentifier, DefaultKeyFuncs<FAssetIdentifier>, TInlineSetAllocator<1>> PackagesToRemove;
-	PackagesToRemove.Reserve(Keys.Num());
-
-	for (int32 Index = 0; Index < Keys.Num(); Index++)
-	{
-		bOutRemovedAssetDatas[Index] = false;
-		bOutRemovedPackageDatas[Index] = false;
-
-		const FCachedAssetKey& Key  = Keys[Index];
-		if (!CachedAssets.Find(Key))
+	auto RemoveFromPackageAssets = 
+		[this, AssetDatas, Keys, bRemoveDependencyData, &bOutRemovedPackageDatas
+#if UE_ASSETREGISTRY_INDIRECT_ASSETDATA_POINTERS
+		,AssetIndices
+#endif
+		]()
 		{
-			continue;
-		}
+			TSet<FAssetIdentifier, DefaultKeyFuncs<FAssetIdentifier>, TInlineSetAllocator<1>> PackagesToRemove;
+			PackagesToRemove.Reserve(Keys.Num());
+
+			for (int32 Index = 0; Index < Keys.Num(); Index++)
+			{
+				bOutRemovedPackageDatas[Index] = false;
 
-		FAssetData* AssetData = AssetDatas[Index];
+				FAssetData* AssetData = AssetDatas[Index];
 
 #if !UE_ASSETREGISTRY_INDIRECT_ASSETDATA_POINTERS
-		FAssetData* AssetIndex = AssetDatas[Index];
+				FAssetData* AssetIndex = AssetData;
+
+				TArray<FAssetData*, TInlineAllocator<1>>* OldPackageAssets = CachedAssetsByPackageName.Find(AssetData->PackageName);
+				bool bRemovedAssetFromPackageAssets = OldPackageAssets->RemoveSingleSwap(AssetIndex, EAllowShrinking::No) > 0;
+				bool bOldPackageAssetsEmpty = bRemovedAssetFromPackageAssets && OldPackageAssets->Num() == 0;
+				if (bOldPackageAssetsEmpty)
+				{
+					CachedAssetsByPackageName.Remove(AssetData->PackageName);
+				}
+#else
+				UE::AssetRegistry::Private::FAssetDataPtrIndex AssetIndex = AssetIndices[Index];
+				bool bRemovedAssetFromPackageAssets = CachedAssetsByPackageName.Remove(AssetData->PackageName, AssetIndex) > 0;
+				bool bOldPackageAssetsEmpty = bRemovedAssetFromPackageAssets && !CachedAssetsByPackageName.Contains(AssetData->PackageName);
+#endif
+				if (bOldPackageAssetsEmpty)
+				{
+					PackagesToRemove.Add(AssetData->PackageName);
+				}
+
+				bOutRemovedPackageDatas[Index] = bOldPackageAssetsEmpty;
+			}
+
+			// Only remove dependencies and package data if there are no other known assets in the package
+			if (bRemoveDependencyData && !PackagesToRemove.IsEmpty())
+			{
+				// We need to update the cached dependencies references cache so that they know we no
+				// longer exist and so don't reference them.
+				if (PackagesToRemove.Num() == 1)
+				{
+					// It's common enough to remove a single package at 
+					// a time we should special case it
+					RemoveDependsNode(*PackagesToRemove.FindArbitraryElement());
+				}
+				else
+				{
+					RemoveDependsNodes(PackagesToRemove);
+				}
+			}
 
-		TArray<FAssetData*, TInlineAllocator<1>>* OldPackageAssets = CachedAssetsByPackageName.Find(AssetData->PackageName);
-		OldPackageAssets->RemoveSingleSwap(AssetIndex, EAllowShrinking::No);
-		bool bOldPackageAssetsEmpty = OldPackageAssets->Num() == 0;
-		if (bOldPackageAssetsEmpty)
+			// Remove the package data as well
+			RemovePackageData(PackagesToRemove);
+		};
+
+	auto RemoveFromOtherAccelerators =
+		[this, AssetDatas, Keys
+#if UE_ASSETREGISTRY_INDIRECT_ASSETDATA_POINTERS
+		, AssetIndices
+#endif
+		]()
 		{
-			CachedAssetsByPackageName.Remove(AssetData->PackageName);
-			PackagesToRemove.Add(AssetData->PackageName);
-		}
+#if !UE_ASSETREGISTRY_INDIRECT_ASSETDATA_POINTERS
+			TMap<FTopLevelAssetPath, TSet<FAssetData*>> RemovedAssetsForClass;
 #else
-		UE::AssetRegistry::Private::FAssetDataPtrIndex AssetIndex = AssetIndices[Index];
-		CachedAssetsByPackageName.Remove(AssetData->PackageName, AssetIndex);
-		bool bOldPackageAssetsEmpty = !CachedAssetsByPackageName.Contains(AssetData->PackageName);
+			TMap<FTopLevelAssetPath, TSet<UE::AssetRegistry::Private::FAssetDataPtrIndex>> RemovedAssetsForClass;
+#endif
+
+			for (int32 Index = 0; Index < Keys.Num(); Index++)
+			{
+				FAssetData* AssetData = AssetDatas[Index];
+
+#if !UE_ASSETREGISTRY_INDIRECT_ASSETDATA_POINTERS
+				FAssetData* AssetIndex = AssetData;
+#else
+				UE::AssetRegistry::Private::FAssetDataPtrIndex AssetIndex = AssetIndices[Index];
 #endif
 
-		CachedAssetsByPath.Find(AssetData->PackagePath)->RemoveSingleSwap(AssetIndex);
-		CachedAssetsByClass.Find(AssetData->AssetClassPath)->RemoveSingleSwap(AssetIndex);
+				CachedAssetsByPath.Find(AssetData->PackagePath)->RemoveSingleSwap(AssetIndex);
+				auto& ArrayOfAssetIndex = RemovedAssetsForClass.FindOrAdd(AssetData->AssetClassPath);
+				if (ArrayOfAssetIndex.IsEmpty())
+				{
+					ArrayOfAssetIndex.Reserve(CachedAssetsByClass.Find(AssetData->AssetClassPath)->Num());
+				}
+				ArrayOfAssetIndex.Add(AssetIndex);
 
 #if UE_ASSETREGISTRY_CACHEDASSETSBYTAG
-		for (auto TagIt = AssetData->TagsAndValues.CreateConstIterator(); TagIt; ++TagIt)
-		{
-			CachedAssetsByTag.Find(TagIt.Key())->Remove(AssetIndex);
-		}
+				for (auto TagIt = AssetData->TagsAndValues.CreateConstIterator(); TagIt; ++TagIt)
+				{
+					CachedAssetsByTag.Find(TagIt.Key())->Remove(AssetIndex);
+				}
 #else
-		// For CachedClassesByTag, we do not need to remove the asset's class from the entries
-		// for the old tags. The class is not changed and still has the possibility of containing the tags.
+				// For CachedClassesByTag, we do not need to remove the asset's class from the entries
+				// for the old tags. The class is not changed and still has the possibility of containing the tags.
 #endif
+			}
 
-		// In the UE_ASSETREGISTRY_INDIRECT_ASSETDATA_POINTERS case the other containers hold an index into CachedAssets,
-		// so we can only remove from CachedAssets after removing from all other containers.
-		CachedAssets.Remove(Key);
+			// Remove the FAssetData from their class list. We do this as one operation per class, and with the removed
+			// AssetDatas in a TSet, because the class list is a TArray and it would be expensive to find each removed
+			// AssetData in the TArray and call RemoveSwap on it.
+			for (const auto& ClassPathAndSetOfAssetIndexPair : RemovedAssetsForClass)
+			{
+				const FTopLevelAssetPath& ClassPath = ClassPathAndSetOfAssetIndexPair.Key;
+				const auto& SetOfRemovedAssetIndex = ClassPathAndSetOfAssetIndexPair.Value;
+				CachedAssetsByClass.Find(ClassPath)->RemoveAllSwap([&SetOfRemovedAssetIndex](auto ExistingAssetIndex)
+					{
+						return SetOfRemovedAssetIndex.Contains(ExistingAssetIndex);
+					}, EAllowShrinking::Yes);
+			}
+		};
 
-		// If the assets were preallocated in a block, we can't delete them one at a time,
-		// only the whole chunk in the destructor
-		if (PreallocatedAssetDataBuffers.Num() == 0)
+	auto RemoveFromCachedAssets =
+		[this, Keys, &bOutRemovedAssetDatas]()
 		{
-			delete AssetData;
-		}
+			for (int32 Index = 0; Index < Keys.Num(); Index++)
+			{
+				const FCachedAssetKey& Key = Keys[Index];
+				int32 NumRemoved = CachedAssets.Remove(Key);
+				bOutRemovedAssetDatas[Index] = NumRemoved > 0;
+				if (NumRemoved > 0)
+				{
+					NumAssets--;
+				}
+			}
+		};
+
+	constexpr int32 MinCountForMultiThreaded = 100;
+	if (Keys.Num() < MinCountForMultiThreaded)
+	{
+		// Avoid the overhead of creating Tasks and just run the lambdas in the current thread
+		RemoveFromPackageAssets();
+		RemoveFromOtherAccelerators();
+		RemoveFromCachedAssets();
+	}
+	else
+	{
+		FTask RemoveFromPackageAssetsTask = Launch(TEXT("RemoveAssetDatasImpl_PackageAssets"),
+			MoveTemp(RemoveFromPackageAssets), ETaskPriority::Default, EExtendedTaskPriority::None);
+		FTask RemoveFromOtherAcceleratorsTask = Launch(TEXT("RemoveAssetDatasImpl_OtherAccelerators"),
+			MoveTemp(RemoveFromOtherAccelerators), ETaskPriority::Default, EExtendedTaskPriority::None);
+		FTask RemoveFromCachedAssetsTask = Launch(TEXT("RemoveAssetDatasImpl_CachedAssets"),
+			MoveTemp(RemoveFromCachedAssets), 
+#if UE_ASSETREGISTRY_INDIRECT_ASSETDATA_POINTERS
+			// In the UE_ASSETREGISTRY_INDIRECT_ASSETDATA_POINTERS case the other containers hold an index into
+			// CachedAssets, so we can only remove from CachedAssets after removing from all other containers.
+			UE::Tasks::Prerequisites(RemoveFromPackageAssetsTask, RemoveFromOtherAcceleratorsTask),
+#endif
+			ETaskPriority::Default, EExtendedTaskPriority::None);
 
-		NumAssets--;
-		bOutRemovedAssetDatas[Index] = true;
-		bOutRemovedPackageDatas[Index] = bOldPackageAssetsEmpty;
+		RemoveFromPackageAssetsTask.Wait();
+		RemoveFromOtherAcceleratorsTask.Wait();
+		RemoveFromCachedAssetsTask.Wait();
 	}
 
-	// Only remove dependencies and package data if there are no other known assets in the package
-	if (bRemoveDependencyData && !PackagesToRemove.IsEmpty())
+	// If the assets were preallocated in a block, we can't delete them one at a time,
+	// only the whole chunk in the destructor
+	if (PreallocatedAssetDataBuffers.Num() == 0)
 	{
-		// We need to update the cached dependencies references cache so that they know we no
-		// longer exist and so don't reference them.
-		if (PackagesToRemove.Num() == 1)
-		{
-			// It's common enough to remove a single package at 
-			// a time we should special case it
-			RemoveDependsNode(*PackagesToRemove.FindArbitraryElement());
-		}
-		else
+		for (int32 Index = 0; Index < Keys.Num(); Index++)
 		{
-			RemoveDependsNodes(PackagesToRemove);
+			FAssetData* AssetData = AssetDatas[Index];
+			if (bOutRemovedAssetDatas[Index])
+			{
+				delete AssetData;
+			}
 		}
 	}
-
-	// Remove the package data as well
-	RemovePackageData(PackagesToRemove);
 }
 
 void FAssetRegistryState::Shrink()
@@ -3638,14 +3798,11 @@ bool FAssetRegistryState::RemovePackageData(const TSet<FAssetIdentifier, SetKeyF
 	bool bRemoved = false;
 	for (const FAssetIdentifier& AssetIdentifier : AssetIdentifiers)
 	{
-		FAssetPackageData** DataPtr = CachedPackageData.Find(AssetIdentifier.PackageName);
-
-		if (DataPtr != nullptr)
+		FAssetPackageData* Data;
+		if (CachedPackageData.RemoveAndCopyValue(AssetIdentifier.PackageName, Data))
 		{
-			FAssetPackageData* Data = *DataPtr;
 			if (Data != nullptr)
 			{
-				CachedPackageData.Remove(AssetIdentifier.PackageName);
 				NumPackageData--;
 
 				// If the package data was preallocated in a block, we can't delete them one at a time,
