AssetRegistryState: Optimize InitlializeFromExisting. Optimizations included batching and ParallelFor.
#rnx

diff --git a/Engine/Source/Runtime/AssetRegistry/Private/AssetRegistryState.cpp b/Engine/Source/Runtime/AssetRegistry/Private/AssetRegistryState.cpp
--- a/Engine/Source/Runtime/AssetRegistry/Private/AssetRegistryState.cpp
+++ b/Engine/Source/Runtime/AssetRegistry/Private/AssetRegistryState.cpp
@@ -423,82 +423,151 @@ void FAssetRegistryState::InitializeFromExisting(const FAssetDataMap& AssetDataM
 	EInitializationMode InInitializationMode,
 	FAssetRegistryAppendResult* OutAppendResult)
 {
+	TArray<FAssetData*> ExistingAssetForNewAsset;
+	int32 AssetDataMapMaxIndex = AssetDataMap.GetMaxIndex();
 	if (InInitializationMode == EInitializationMode::Rebuild)
 	{
 		Reset();
 	}
+	else
+	{
+		ExistingAssetForNewAsset.SetNum(AssetDataMapMaxIndex);
+		for (int32 Index = 0; Index < AssetDataMapMaxIndex; ++Index)
+		{
+			FAssetData* ExistingAsset = nullptr;
+			FSetElementId SetId = FSetElementId::FromInteger(Index);
+			if (AssetDataMap.IsValidId(SetId))
+			{
+				FAssetData* AssetData = AssetDataMap.Get(SetId);
+				if (FAssetData* const* Ptr = CachedAssets.Find(FCachedAssetKey(AssetData)))
+				{
+					ExistingAsset = *Ptr;
+				}
+			}
+			ExistingAssetForNewAsset[Index] = ExistingAsset;
+		}
+	}
 
 	// FStoreBuilder is a serialized tag builder which will allow us to store the asset tags much more efficiently.
 	TArray<TPair<FAssetData*, FixedTagPrivate::FPartialMapHandle>> AssetsToUpdate;
 	FixedTagPrivate::FStoreBuilder Builder(FAssetRegistryWriterOptions(Options).Tags);
-	
-	for (const FAssetData* AssetDataPtr : AssetDataMap)
-	{
-		if (AssetDataPtr == nullptr)
-		{
-			// don't do anything 
-			continue;
-		}
-		const FAssetData& AssetData = *AssetDataPtr;
 
-		FAssetData* ExistingData = nullptr;
-		if (InInitializationMode != EInitializationMode::Rebuild) // minor optimization to avoid lookup in rebuild mode
+	// We accumulate Assets to add to *this in a queue inside the ParallelFor. Construction of each asset is expensive
+	// and we want to work on it on multple threads. But we can't add them to *this simultaneously because we can't append
+	// to our CachedAsset and query accelerator maps in parallel. So we need to add them all in a batch after the
+	// ParallelFor is over. But waiting until after the ParallelFor is over causes cache misses on the ids, so instead
+	// of doing them all in a single batch, we add them in smaller batches as we go. Each worker of the parallelfor
+	// tries to take responsibility for adding the batch under the lock when the batch reaches our desired size, and
+	// the winner of the race is the one that adds the batch to *this.
+	const int32 QueueFlushLength = 100; // Might want to change this to be proportional to platform's Cache size.
+	std::atomic<int32> AssetsToAddNum(0);
+	FCriticalSection FlushQueueLock;
+
+	TMpscQueue<TPair<FAssetData*, FAssetData*>> AssetsToAdd;
+
+	auto FlushQueue = [&AssetsToAdd, this, OutAppendResult, &Builder, &AssetsToUpdate]()
 		{
-			if (FAssetData*const* Ptr = CachedAssets.Find(FCachedAssetKey(AssetData)))
+			TPair<FAssetData*, FAssetData*> NewAndExisting;
+			while (AssetsToAdd.Dequeue(NewAndExisting))
 			{
-				ExistingData = *Ptr;
+				FAssetData* NewData = NewAndExisting.Key;
+				FAssetData* ExistingData = NewAndExisting.Value;
+				if (ExistingData)
+				{
+					bool bIsModified = false;
+					UpdateAssetData(ExistingData, MoveTemp(*NewData), &bIsModified);
+					if (OutAppendResult && bIsModified)
+					{
+						OutAppendResult->UpdatedAssets.Add(ExistingData);
+					}
+					delete NewData;
+				}
+				else
+				{
+					AddAssetData(NewData);
+					if (OutAppendResult)
+					{
+						OutAppendResult->AddedAssets.Add(NewData);
+					}
+					AssetsToUpdate.Add(TPair<FAssetData*, FixedTagPrivate::FPartialMapHandle>(NewData,
+						Builder.AddTagMap(NewData->TagsAndValues)));
+				}
 			}
-		}
-		if (InInitializationMode == EInitializationMode::OnlyUpdateExisting && ExistingData == nullptr)
-		{
-			continue;
-		}
-		if (InInitializationMode == EInitializationMode::OnlyUpdateNew && ExistingData != nullptr)
-		{
-			continue;
-		}
+		};
 
-		if (Options.bFilterAssetDataWithNoTags && AssetData.TagsAndValues.Num() == 0 &&
-			!FPackageName::IsLocalizedPackage(WriteToString<256>(AssetData.PackageName)))
+	auto ConditionalAddDuplicatedAssetToQueue =
+		[InInitializationMode, &Options, &AssetsToAdd, &AssetDataMap, &ExistingAssetForNewAsset,
+		&AssetsToAddNum, QueueFlushLength, &FlushQueue, &FlushQueueLock]
+		(int32 Index)
 		{
-			// based on current serialization settings this should be skipped.
-			continue;
-		}
+			FSetElementId SetId = FSetElementId::FromInteger(Index);
+			if (!AssetDataMap.IsValidId(SetId))
+			{
+				return;
+			}
+			const FAssetData* AssetDataPtr = AssetDataMap.Get(SetId);
 
-		// Filter asset registry tags now
-		FAssetDataTagMap LocalTagsAndValues;
-		FAssetRegistryState::FilterTags(AssetData.TagsAndValues, LocalTagsAndValues,
-			Options.CookFilterlistTagsByClass.Find(AssetData.AssetClassPath), Options);
+			if (AssetDataPtr == nullptr)
+			{
+				return;
+			}
+			const FAssetData& AssetData = *AssetDataPtr;
 
-		// Filter AssetBundles separately, LocalAssetBundles will point to the same TaggedAssetBundles if it is not different
-		TSharedPtr<FAssetBundleData> LocalAssetBundles = FAssetRegistryState::FilterBundles(AssetData.TaggedAssetBundles, Options);
+			FAssetData* ExistingData = nullptr;
+			if (InInitializationMode != EInitializationMode::Rebuild) // minor optimization to avoid lookup in rebuild mode
+			{
+				ExistingData = ExistingAssetForNewAsset[Index];
+			}
+			if (InInitializationMode == EInitializationMode::OnlyUpdateExisting && ExistingData == nullptr)
+			{
+				return;
+			}
+			if (InInitializationMode == EInitializationMode::OnlyUpdateNew && ExistingData != nullptr)
+			{
+				return;
+			}
 
-		if (ExistingData)
-		{
-			FAssetData NewData(AssetData);
-			NewData.TagsAndValues = FAssetDataTagMapSharedView(MoveTemp(LocalTagsAndValues));
-			NewData.TaggedAssetBundles = MoveTemp(LocalAssetBundles);
-			bool bIsModified = false;
-			UpdateAssetData(ExistingData, MoveTemp(NewData), &bIsModified);
-			if (OutAppendResult && bIsModified)
+			if (Options.bFilterAssetDataWithNoTags && AssetData.TagsAndValues.Num() == 0 &&
+				!FPackageName::IsLocalizedPackage(WriteToString<256>(AssetData.PackageName)))
 			{
-				OutAppendResult->UpdatedAssets.Add(ExistingData);
+				// based on current serialization settings this should be skipped.
+				return;
 			}
-		}
-		else
-		{
+
+			// Filter asset registry tags now
+			FAssetDataTagMap LocalTagsAndValues;
+			FAssetRegistryState::FilterTags(AssetData.TagsAndValues, LocalTagsAndValues,
+				Options.CookFilterlistTagsByClass.Find(AssetData.AssetClassPath), Options);
+
+			// Filter AssetBundles separately, LocalAssetBundles will point to the same TaggedAssetBundles if it is not different
+			TSharedPtr<FAssetBundleData> LocalAssetBundles = FAssetRegistryState::FilterBundles(AssetData.TaggedAssetBundles, Options);
+
 			FAssetData* NewData = new FAssetData(AssetData);
 			NewData->TagsAndValues = FAssetDataTagMapSharedView(MoveTemp(LocalTagsAndValues));
 			NewData->TaggedAssetBundles = MoveTemp(LocalAssetBundles);
-			AddAssetData(NewData);
-			if (OutAppendResult)
-			{
-				OutAppendResult->AddedAssets.Add(NewData);
+			AssetsToAdd.Enqueue(NewData, ExistingData);
+			int32 OldNum = AssetsToAddNum.fetch_add(1, std::memory_order_relaxed);
+			if (OldNum == QueueFlushLength)
+			{
+				AssetsToAddNum.store(0, std::memory_order_relaxed);
+				// Only the winner of the race that happens to add the element that sets num to QueueFlushLength will
+				// enter this OldNum == QueueFlushLength block, but as the other threads continue adding, it is
+				// possible that another thread will encounter it the next time around after we reset to 0, so to
+				// avoid simultaneous access from all winners of the races, we need to also enter a critical section.
+				if (FlushQueueLock.TryLock())
+				{
+					FlushQueue();
+					FlushQueueLock.Unlock();
+				}
 			}
+		};
 
-			AssetsToUpdate.Add(TPair<FAssetData*, FixedTagPrivate::FPartialMapHandle>(NewData, Builder.AddTagMap(NewData->TagsAndValues)));
-		}
-	}
+	ParallelFor(AssetDataMapMaxIndex,
+		[&AssetDataMap, &ConditionalAddDuplicatedAssetToQueue](int32 Index)
+		{
+			ConditionalAddDuplicatedAssetToQueue(Index);
+		});
+	FlushQueue();
 	// Construct the optimized tag data then set the handle map back to the assets we need to update
 	TRefCountPtr<const FixedTagPrivate::FStore> FixedStore = FixedTagPrivate::InitializeFromBuilder(Builder);
 	for (TPair<FAssetData*, FixedTagPrivate::FPartialMapHandle>& AssetToHandle : AssetsToUpdate)
