AssetRegistryState: Optimize InitializeFromExistingAndPrune by using multiple threads.
#rnx

diff --git a/Engine/Source/Runtime/AssetRegistry/Private/AssetRegistryState.cpp b/Engine/Source/Runtime/AssetRegistry/Private/AssetRegistryState.cpp
--- a/Engine/Source/Runtime/AssetRegistry/Private/AssetRegistryState.cpp
+++ b/Engine/Source/Runtime/AssetRegistry/Private/AssetRegistryState.cpp
@@ -265,13 +265,14 @@ void FAssetRegistryState::InitializeFromExistingAndPrune(const FAssetRegistrySta
 	const bool bIsFilteredByRequiredPackages = RequiredPackages.Num() != 0;
 	const bool bIsFilteredByRemovedPackages = RemovePackages.Num() != 0;
 
-	TSet<FName> RequiredDependNodePackages;
+	TMpscQueue<FName> RequiredDependNodePackagesQueue;
+	TMpscQueue<FAssetData*> NewAssetDatas;
 
 	// Duplicate asset data entries
 	ExistingState.EnumerateAllMutableAssets(
 		[&RequiredPackages, &RemovePackages, &ChunksToKeep, &Options,
 		bIsFilteredByChunkId, bIsFilteredByRequiredPackages, bIsFilteredByRemovedPackages,
-		&RequiredDependNodePackages, this]
+		&RequiredDependNodePackagesQueue, &NewAssetDatas, this]
 		(FAssetData& AssetData)
 	{
 		bool bRemoveAssetData = false;
@@ -304,7 +305,7 @@ void FAssetRegistryState::InitializeFromExistingAndPrune(const FAssetRegistrySta
 		{
 			if (!bRemoveDependencyData)
 			{
-				RequiredDependNodePackages.Add(AssetData.PackageName);
+				RequiredDependNodePackagesQueue.Enqueue(AssetData.PackageName);
 			}
 			return;
 		}
@@ -328,27 +329,63 @@ void FAssetRegistryState::InitializeFromExistingAndPrune(const FAssetRegistrySta
 		NewAssetData->TaggedAssetBundles = FAssetRegistryState::FilterBundles(AssetData.TaggedAssetBundles, Options);
 
 		// Add asset to new state
-		AddAssetData(NewAssetData);
-	});
+		NewAssetDatas.Enqueue(NewAssetData);
+	}, UE::AssetRegistry::EEnumerateAssetsFlags::Parallel | UE::AssetRegistry::EEnumerateAssetsFlags::AllowAll);
+
+	TSet<FName> RequiredDependNodePackages;
+	FName DependNodePackageName;
+	while (RequiredDependNodePackagesQueue.Dequeue(DependNodePackageName))
+	{
+		RequiredDependNodePackages.Add(DependNodePackageName);
+	}
+
+	FAssetData* NewAssetDataToAdd;
+	while (NewAssetDatas.Dequeue(NewAssetDataToAdd))
+	{
+		AddAssetData(NewAssetDataToAdd);
+	}
 
 	// Create package data for all script and required packages
-	for (const TPair<FName, FAssetPackageData*>& Pair : ExistingState.CachedPackageData)
+	// Do this in a ParallelFor since it requires a call to IsScriptPackage which can be expensive when called 1 million times
+	TArray<bool> ShouldAddAssetPackageData;
+	int32 CachedPackageDataMaxIndex = ExistingState.CachedPackageData.GetMaxIndex();
+	ShouldAddAssetPackageData.SetNumUninitialized(CachedPackageDataMaxIndex);
+	ParallelFor(CachedPackageDataMaxIndex,
+		[&ExistingState, &ShouldAddAssetPackageData, this](int32 Index)
+		{
+			FSetElementId Id = FSetElementId::FromInteger(Index);
+			bool bShouldAdd = false;
+			if (ExistingState.CachedPackageData.IsValidId(Id))
+			{
+				const TPair<FName, FAssetPackageData*>& Pair = ExistingState.CachedPackageData.Get(Id);
+				if (Pair.Value)
+				{
+					// Only add if also in asset data map, or script package
+					if (CachedAssetsByPackageName.Find(Pair.Key) ||
+						FPackageName::IsScriptPackage(WriteToString<256>(Pair.Key)))
+					{
+						bShouldAdd = true;
+					}
+				}
+			}
+			ShouldAddAssetPackageData[Index] = bShouldAdd;
+		});
+	for (int32 Index = 0; Index < CachedPackageDataMaxIndex; ++Index)
 	{
-		if (Pair.Value)
+		if (ShouldAddAssetPackageData[Index])
 		{
-			// Only add if also in asset data map, or script package
-			if (CachedAssetsByPackageName.Find(Pair.Key) ||
-				FPackageName::IsScriptPackage(WriteToString<256>(Pair.Key)))
-			{
-				FAssetPackageData* NewData = CreateOrGetAssetPackageData(Pair.Key);
-				*NewData = *Pair.Value;
-			}
+			FSetElementId Id = FSetElementId::FromInteger(Index);
+			check(ExistingState.CachedPackageData.IsValidId(Id));
+			const TPair<FName, FAssetPackageData*>& Pair = ExistingState.CachedPackageData.Get(Id);
+			FAssetPackageData* NewData = CreateOrGetAssetPackageData(Pair.Key);
+			*NewData = *Pair.Value;
 		}
 	}
 
-	// Find valid dependency nodes for all script and required packages
-	TSet<FDependsNode*> ValidDependsNodes;
-	ValidDependsNodes.Reserve(ExistingState.CachedDependsNodes.Num());
+	// Find valid dependency nodes for all script and required packages, and create a new
+	// node corresponding to each old node.
+	TMap<FDependsNode*, FDependsNode*> OldNodeToNewNode;
+	OldNodeToNewNode.Reserve(ExistingState.CachedDependsNodes.Num());
 	for (const TPair<FAssetIdentifier, FDependsNode*>& Pair : ExistingState.CachedDependsNodes)
 	{
 		FDependsNode* Node = Pair.Value;
@@ -362,58 +399,84 @@ void FAssetRegistryState::InitializeFromExistingAndPrune(const FAssetRegistrySta
 		else if (Id.IsPackage() &&
 			!CachedAssetsByPackageName.Contains(Id.PackageName) &&
 			!RequiredDependNodePackages.Contains(Id.PackageName) &&
-			!FPackageName::IsScriptPackage(WriteToString<256>(Id.PackageName)))
+			!Node->IsScriptPath())
 		{
 			bRemoveDependsNode = true;
 		}
 
 		if (!bRemoveDependsNode)
 		{
-			ValidDependsNodes.Add(Node);
+			FDependsNode*& ExistingNewNode = OldNodeToNewNode.FindOrAdd(Node);
+			if (!ExistingNewNode)
+			{
+				ExistingNewNode = CreateOrFindDependsNode(Node->GetIdentifier());
+			}
 		}
 	}
 
-	// Duplicate dependency nodes
-	for (FDependsNode* OldNode : ValidDependsNodes)
-	{
-		FDependsNode* NewNode = CreateOrFindDependsNode(OldNode->GetIdentifier());
-		NewNode->Reserve(OldNode);
-	}
-	
-	for (FDependsNode* OldNode : ValidDependsNodes)
-	{
-		FDependsNode* NewNode = CreateOrFindDependsNode(OldNode->GetIdentifier());
-		OldNode->IterateOverDependencies([&, OldNode, NewNode](FDependsNode* InDependency,
-			UE::AssetRegistry::EDependencyCategory InCategory,
-			UE::AssetRegistry::EDependencyProperty InFlags,
-			bool bDuplicate)
-		{
-			if (ValidDependsNodes.Contains(InDependency))
-			{
-				// Only add link if it's part of the filtered asset set
-				FDependsNode* NewDependency = CreateOrFindDependsNode(InDependency->GetIdentifier());
-				NewNode->SetIsDependencyListSorted(InCategory, false);
-				NewNode->AddDependency(NewDependency, InCategory, InFlags);
-				NewDependency->SetIsReferencersSorted(false);
-				NewDependency->AddReferencer(NewNode);
+	TMpscQueue<TPair<FDependsNode*, FDependsNode*>> ReferencerPairs;
+
+	int32 OldNodeToNewNodeMaxId = OldNodeToNewNode.GetMaxIndex();
+	ParallelFor(OldNodeToNewNodeMaxId, [&OldNodeToNewNode, &ReferencerPairs](int32 Index)
+		{
+			FSetElementId Id = FSetElementId::FromInteger(Index);
+			if (OldNodeToNewNode.IsValidId(Id))
+			{
+				const TPair<FDependsNode*, FDependsNode*>& Pair = OldNodeToNewNode.Get(Id);
+				FDependsNode* OldNode = Pair.Key;
+				FDependsNode* NewNode = Pair.Value;
+				NewNode->Reserve(OldNode);
+				NewNode->SetIsDependencyListSorted(UE::AssetRegistry::EDependencyCategory::All, false);
+				NewNode->SetIsReferencersSorted(false);
+
+				OldNode->IterateOverDependencies([&, OldNode, NewNode](FDependsNode* InDependency,
+					UE::AssetRegistry::EDependencyCategory InCategory,
+					UE::AssetRegistry::EDependencyProperty InFlags,
+					bool bDuplicate)
+					{
+						// Only add link if it's part of the filtered asset set
+						FDependsNode** NewDependency = OldNodeToNewNode.Find(InDependency);
+						if (NewDependency)
+						{
+							NewNode->AddDependency(*NewDependency, InCategory, InFlags);
+							ReferencerPairs.Enqueue(NewNode, *NewDependency);
+						}
+					});
+				NewNode->SetIsScriptDependenciesInitialized(OldNode->IsScriptDependenciesInitialized());
+				NewNode->SetIsDependencyListSorted(UE::AssetRegistry::EDependencyCategory::All, true);
 			}
 		});
-		NewNode->SetIsScriptDependenciesInitialized(OldNode->IsScriptDependenciesInitialized());
+
+	// Add the Referencer link for each dependency link we added
+	TPair<FDependsNode*, FDependsNode*> ReferencerPair;
+	while (ReferencerPairs.Dequeue(ReferencerPair))
+	{
+		ReferencerPair.Value->AddReferencer(ReferencerPair.Key);
 	}
 
 	// Remove any orphaned depends nodes. This will leave cycles in but those might represent useful data
-	TArray<FDependsNode*> AllDependsNodes;
-	CachedDependsNodes.GenerateValueArray(AllDependsNodes);
-	for (FDependsNode* DependsNode : AllDependsNodes)
+	for (TMap<FDependsNode*, FDependsNode*>::TIterator Iter(OldNodeToNewNode); Iter; ++Iter)
 	{
+		FDependsNode* DependsNode = Iter->Value;
 		if (DependsNode->GetConnectionCount() == 0)
 		{
 			RemoveDependsNode(DependsNode->GetIdentifier());
+			Iter.RemoveCurrent();
 		}
 	}
+	OldNodeToNewNodeMaxId = OldNodeToNewNode.GetMaxIndex();
 
-	// Restore the sortedness that we turned off for performance when creating each DependsNode
-	SetDependencyNodeSorting(true, true);
+	// Restore the referencer sortedness that we turned off for performance when creating each DependsNode
+	ParallelFor(OldNodeToNewNodeMaxId, [&OldNodeToNewNode](int32 Index)
+		{
+			FSetElementId Id = FSetElementId::FromInteger(Index);
+			if (OldNodeToNewNode.IsValidId(Id))
+			{
+				const TPair<FDependsNode*, FDependsNode*>& Pair = OldNodeToNewNode.Get(Id);
+				FDependsNode* NewNode = Pair.Value;
+				NewNode->SetIsReferencersSorted(true);
+			}
+		});
 }
 
 void FAssetRegistryState::InitializeFromExisting(const FAssetDataMap& AssetDataMap,
