AssetRegistry: DiscoveryCache: WindowsPlatformFileJournal: speculative fix for stale data in the AssetDataGatherer's DiscoveryCache. We have learned that the FileReferenceNumber assigned to a directory and used by the USN journal can change in response to some events on the machine, such as defragmentation passes. We speculate but have not yet confirmed that these reassignments are reported in the USN journal in a record with the directory's old id as the FileReferenceNumber and with a reason such as USN_REASON_BASIC_INFO_CHANGE. We were accidentally handling those reassignments, if they exist, for directories with parent directories also in the list of content directories, but we were not handling it for the reassignments of the FileReferenceNumber for top-level mount directories. Add an invalidation of the cached data for directories with change records themselves, in addition to our current invalidation for directories that are the parent directory of paths with a change record.
#rnx

diff --git a/Engine/Source/Runtime/AssetRegistry/Private/AssetDataGatherer.cpp b/Engine/Source/Runtime/AssetRegistry/Private/AssetDataGatherer.cpp
--- a/Engine/Source/Runtime/AssetRegistry/Private/AssetDataGatherer.cpp
+++ b/Engine/Source/Runtime/AssetRegistry/Private/AssetDataGatherer.cpp
@@ -2186,7 +2186,7 @@ void FAssetDataDiscovery::TickInternal(bool bTickAll)
 			{
 				// If we fail to iterate this directory, fall back to the old way. Meaning we will fail to cache this but still pick up on the assets in the directories
 				if (FPlatformFileManager::Get().GetPlatformFile().FileJournalIterateDirectory(Data.DirLocalAbsPath.ToString(),
-					[this, &Data, &CacheDataToAdd, &ProcessIterData]
+					[&ProcessIterData]
 					(const TCHAR* IterFilename, const FFileJournalData& IterData)
 					{
 						return ProcessIterData(IterFilename, IterData.bIsDirectory, IterData.ModificationTime, IterData.bIsReadOnly, IterData.JournalHandle, IterData.bIsReparsePoint);
@@ -2204,7 +2204,7 @@ void FAssetDataDiscovery::TickInternal(bool bTickAll)
 						// If we failed once, run once more time but grab an error so we can send some telemetry on the issue
 						FString OutError;
 						if (!FPlatformFileManager::Get().GetPlatformFile().FileJournalIterateDirectory(Data.DirLocalAbsPath.ToString(),
-							[this, &Data, &CacheDataToAdd, &ProcessIterData]
+							[]
 							(const TCHAR* IterFilename, const FFileJournalData& IterData)
 							{
 								return true; // if it doesnt return true it acts as a failure for this function
@@ -2226,7 +2226,7 @@ void FAssetDataDiscovery::TickInternal(bool bTickAll)
 			if (!bIteratedDirectory)
 			{
 				IFileManager::Get().IterateDirectoryStat(Data.DirLocalAbsPath.ToString(),
-					[this, &Data, &ProcessIterData]
+					[&ProcessIterData]
 					(const TCHAR* IterFilename, const FFileStatData& IterData)
 					{
 						return ProcessIterData(IterFilename, IterData.bIsDirectory, IterData.ModificationTime, IterData.bIsReadOnly, FFileJournalFileHandle(), false /* bIsReparsePoint */);
diff --git a/Engine/Source/Runtime/AssetRegistry/Private/AssetDataGathererDiscoveryCache.cpp b/Engine/Source/Runtime/AssetRegistry/Private/AssetDataGathererDiscoveryCache.cpp
--- a/Engine/Source/Runtime/AssetRegistry/Private/AssetDataGathererDiscoveryCache.cpp
+++ b/Engine/Source/Runtime/AssetRegistry/Private/AssetDataGathererDiscoveryCache.cpp
@@ -220,6 +220,13 @@ void FAssetDataDiscoveryCache::LoadAndUpdateCache()
 					// We cannot remove the DirData yet because we will need it to get the list of its child
 					// directories that exist in the tree, at the point when we rescan it
 					DirData->bCacheValid = false;
+					// For modified directories, also clear their JournalHandle; the directory might have been
+					// marked modified because its JournalHandle changed, so we need to recalculate it.
+					// For directories under the top-level directory we recalculate it automatically when we
+					// rescan their parent directory (which we do because parent directories of changed directories
+					// are also marked modified) but for top-level directories we never recalculate their JournalHandle
+					// unless it has been set to invalid.
+					DirData->JournalHandle = FileJournalFileHandleInvalid;
 				}
 			}
 		}
diff --git a/Engine/Source/Runtime/Core/Private/Windows/WindowsPlatformFileJournal.cpp b/Engine/Source/Runtime/Core/Private/Windows/WindowsPlatformFileJournal.cpp
--- a/Engine/Source/Runtime/Core/Private/Windows/WindowsPlatformFileJournal.cpp
+++ b/Engine/Source/Runtime/Core/Private/Windows/WindowsPlatformFileJournal.cpp
@@ -440,7 +440,6 @@ EFileJournalResult FWindowsPlatformFile::FileJournalReadModified(const TCHAR* Vo
 
 	// Records are expected to be aligned to a 64-bit boundary. The USN taking the first 64-bits.
 	alignas(sizeof(uint64)) uint8 ReadBuffer[4096];
-	DWORD ReadBufferSize = UE_ARRAY_COUNT(ReadBuffer);
 
 	HANDLE JournalWindowsHandle = NULL;
 	TStringBuilder<16> NormalizedVolumeName;
@@ -602,10 +601,28 @@ EFileJournalResult FWindowsPlatformFile::FileJournalReadModified(const TCHAR* Vo
 				return EFileJournalResult::UnhandledJournalVersion;
 			}
 
+			FFileJournalFileHandle TargetJournalHandle = ToFileHandle(Record->FileReferenceNumber);
+			FString* TargetPathName = KnownDirectories.Find(TargetJournalHandle);
+			if (TargetPathName)
+			{
+				// If the TargetPath is one of our tracked directories itself, mark it as modified
+				// because we might need to update its FileReferenceNumber.
+				OutModifiedDirectories.Add(*TargetPathName);
+
+				if (bShowFilenames)
+				{
+					FStringView Filename = GetFileName(Record);
+					UE_LOG(LogWindows, Display, TEXT("FileJournalReadModified: TargetDirectory %s"),
+						*TargetJournalHandle.ToString());
+				}
+			}
 			FFileJournalFileHandle ParentJournalHandle = ToFileHandle(Record->ParentFileReferenceNumber);
 			FString* ParentDirectoryName = KnownDirectories.Find(ParentJournalHandle);
 			if (ParentDirectoryName)
 			{
+				// If a file or directory that is a child under one of our tracked directories has changed,
+				// then mark the tracked directory as modified because we need to rescan it to refresh its
+				// list of existing child paths.
 				OutModifiedDirectories.Add(*ParentDirectoryName);
 
 				if (bShowFilenames)
