AssetRegistry: DiscoveryCache: Bump version to invalidate all existing discoverycaches after the speculative fix for the bug that corrupts the discovery cache and causes stale file lists to be kept.
#rnx

diff --git a/Engine/Source/Runtime/AssetRegistry/Private/AssetDataGathererDiscoveryCache.cpp b/Engine/Source/Runtime/AssetRegistry/Private/AssetDataGathererDiscoveryCache.cpp
--- a/Engine/Source/Runtime/AssetRegistry/Private/AssetDataGathererDiscoveryCache.cpp
+++ b/Engine/Source/Runtime/AssetRegistry/Private/AssetDataGathererDiscoveryCache.cpp
@@ -27,7 +27,7 @@ FArchive& operator<<(FArchive& Ar, FFileJournalFileHandle& H)
 
 namespace AssetDataGathererConstants
 {
-const FGuid DiscoveryCacheVersion(TEXT("415D05FD867949DCA49FFDE505005D49"));
+const FGuid DiscoveryCacheVersion(0x18A13604, 0x1BC44658, 0x97F7E0F0, 0x47A84CA2);
 }
 
 namespace UE::AssetDataGather::Private
