AssetRegistry: DiscoveryCache: Fix typo: if -> else if, fixes spurious report of Cache Corruption when there is a cache version mismatch.
#rnx

diff --git a/Engine/Source/Runtime/AssetRegistry/Private/AssetDataGathererDiscoveryCache.cpp b/Engine/Source/Runtime/AssetRegistry/Private/AssetDataGathererDiscoveryCache.cpp
--- a/Engine/Source/Runtime/AssetRegistry/Private/AssetDataGathererDiscoveryCache.cpp
+++ b/Engine/Source/Runtime/AssetRegistry/Private/AssetDataGathererDiscoveryCache.cpp
@@ -312,7 +312,7 @@ bool FAssetDataDiscoveryCache::TryReadCacheFile()
 				TEXT("Version mismatch for AssetDiscovery cache %s. AssetRegistry discovery of files will be uncached."),
 				*Filename);
 		}
-		if (bCacheCorrupt)
+		else if (bCacheCorrupt)
 		{
 			UE_LOG(LogAssetRegistry, Warning,
 				TEXT("Corrupt AssetDiscovery cache %s. AssetRegistry discovery of files will be uncached."),
