Asset Registry - On GC, clean up entries in the AssetDataObjectPathsUpdatedOnLoad set
* This prevented package reloading from picking up changes to asset tags

diff --git a/Engine/Source/Runtime/AssetRegistry/Private/AssetRegistry.cpp b/Engine/Source/Runtime/AssetRegistry/Private/AssetRegistry.cpp
--- a/Engine/Source/Runtime/AssetRegistry/Private/AssetRegistry.cpp
+++ b/Engine/Source/Runtime/AssetRegistry/Private/AssetRegistry.cpp
@@ -953,6 +953,13 @@ FAssetRegistryImpl::FAssetRegistryImpl()
 {
 }
 
+FAssetRegistryImpl::~FAssetRegistryImpl()
+{
+#if WITH_EDITOR
+	FCoreUObjectDelegates::PreGarbageCollectConditionalBeginDestroy.RemoveAll(this);
+#endif
+}
+
 void FAssetRegistryImpl::LoadPremadeAssetRegistry(Impl::FEventContext& EventContext,
 	Premade::ELoadResult LoadResult, FAssetRegistryState&& ARState)
 {
@@ -1164,6 +1171,12 @@ void FAssetRegistryImpl::Initialize(UE::AssetRegistry::FInterfaceWriteScopeLock&
 #if WITH_EDITOR
 	// Make sure first call to LoadCalculatedDependencies builds the Gatherer list. At that point Classes should be loaded.
 	bRegisteredDependencyGathererClassesDirty = true;
+
+	// On GC, clean unloaded objects from the AssetDataObjectPathsUpdatedOnLoad set
+	if (GIsEditor && !IsRunningCommandlet())
+	{
+		FCoreUObjectDelegates::PreGarbageCollectConditionalBeginDestroy.AddRaw(this, &UE::AssetRegistry::FAssetRegistryImpl::CleanupAssetDataObjectPathsUpdatedOnLoad);
+	}
 #endif
 
 	bInitializationComplete = true;
@@ -1171,6 +1184,28 @@ void FAssetRegistryImpl::Initialize(UE::AssetRegistry::FInterfaceWriteScopeLock&
 
 }
 
+void UE::AssetRegistry::FAssetRegistryImpl::CleanupAssetDataObjectPathsUpdatedOnLoad()
+{
+#if WITH_EDITOR
+	TRACE_CPUPROFILER_EVENT_SCOPE(FAssetRegistryImpl::CleanupAssetDataObjectPathsUpdatedOnLoad);
+
+	// Iterate over all entries of AssetDataObjectPathsUpdatedOnLoad
+	// Retrieve the object associated with the FSoftObjectPath
+	// If it can't be retrieved, or is Unreachable, remove the entry from the set
+	for (TSet<FSoftObjectPath>::TIterator It = AssetDataObjectPathsUpdatedOnLoad.CreateIterator(); It; ++It)
+	{
+		const FSoftObjectPath& SoftPath = *It;
+		UObject* ResolvedObject = SoftPath.ResolveObject();
+
+		// Remove if the object doesn't exist or is unreachable
+		if (!ResolvedObject || ResolvedObject->IsUnreachable())
+		{
+			It.RemoveCurrent();
+		}
+	}
+#endif
+}
+
 IAssetRegistry::FAssetCollisionEvent& UE::AssetRegistry::FAssetRegistryImpl::OnAssetCollision_Private()
 {
 	return AssetCollisionEvent;
diff --git a/Engine/Source/Runtime/AssetRegistry/Private/AssetRegistryImpl.h b/Engine/Source/Runtime/AssetRegistry/Private/AssetRegistryImpl.h
--- a/Engine/Source/Runtime/AssetRegistry/Private/AssetRegistryImpl.h
+++ b/Engine/Source/Runtime/AssetRegistry/Private/AssetRegistryImpl.h
@@ -358,6 +358,8 @@ class FAssetRegistryImpl
 public:
 	/** Constructor initializes as little as possible; meaningful operations are first done during Initialize. */
 	FAssetRegistryImpl();
+	~FAssetRegistryImpl();
+
 	/** Construct the AssetRegistryImpl, including initial scans if applicable. */
 	void Initialize(UE::AssetRegistry::FInterfaceWriteScopeLock& ScopeLock, Impl::FInitializeContext& Context);
 
@@ -692,6 +694,9 @@ private:
 
 	/** Free slack memory accumulated by the asset registry. */
 	void Shrink();
+
+	void CleanupAssetDataObjectPathsUpdatedOnLoad();
+
 private:
 
 	/** Internal state of the cached asset registry */
