[AssetRegistry] Resolve runtime filter tag rules lazily

#tests local build

diff --git a/Engine/Source/Runtime/AssetRegistry/Private/AssetRegistry.cpp b/Engine/Source/Runtime/AssetRegistry/Private/AssetRegistry.cpp
--- a/Engine/Source/Runtime/AssetRegistry/Private/AssetRegistry.cpp
+++ b/Engine/Source/Runtime/AssetRegistry/Private/AssetRegistry.cpp
@@ -2,6 +2,7 @@
 
 #include "AssetRegistry.h"
 
+#include "Algo/AnyOf.h"
 #include "Algo/HeapSort.h"
 #include "Algo/Transform.h"
 #include "Algo/Unique.h"
@@ -11,6 +12,7 @@
 #include "AssetRegistry/AssetDependencyGatherer.h"
 #include "AssetRegistry/AssetRegistryTelemetry.h"
 #include "AssetRegistryConsoleCommands.h"
+#include "AssetRegistryImpl.h"
 #include "AssetRegistryPrivate.h"
 #include "Async/Async.h"
 #include "Async/ParallelFor.h"
@@ -23,6 +25,7 @@
 #include "HAL/PlatformMisc.h"
 #include "HAL/ThreadHeartBeat.h"
 #include "Interfaces/IPluginManager.h"
+#include "Math/UnrealMathUtility.h"
 #include "Misc/CommandLine.h"
 #include "Misc/ConfigCacheIni.h"
 #include "Misc/CoreDelegates.h"
@@ -39,6 +42,7 @@
 #include "Misc/TrackedActivity.h"
 #include "Misc/TransactionallySafeRWLock.h"
 #include "AssetRegistry/PackageReader.h"
+#include "Modules/ModuleManager.h"
 #include "ProfilingDebugging/MiscTrace.h"
 #include "Serialization/ArrayReader.h"
 #include "Serialization/CompactBinarySerialization.h"
@@ -1122,9 +1126,8 @@ void FAssetRegistryImpl::Initialize(UE::AssetRegistry::FInterfaceWriteScopeLock&
 #endif
 
 	// Read default serialization options
-	constexpr static bool bInitializeFilterTagOptions = false;
-	Utils::InitializeSerializationOptionsFromIni(SerializationOptions, nullptr, UE::AssetRegistry::ESerializationTarget::ForGame, bInitializeFilterTagOptions);
-	Utils::InitializeSerializationOptionsFromIni(DevelopmentSerializationOptions, nullptr, UE::AssetRegistry::ESerializationTarget::ForDevelopment, bInitializeFilterTagOptions);
+	Utils::InitializeSerializationOptionsFromIni(SerializationOptions, &CachedParsedFilterRules, nullptr, UE::AssetRegistry::ESerializationTarget::ForGame);
+	Utils::InitializeSerializationOptionsFromIni(DevelopmentSerializationOptions, nullptr, nullptr, UE::AssetRegistry::ESerializationTarget::ForDevelopment);
 
 	bool bStartedAsyncGather = false;
 	if (ShouldSearchAllAssetsAtStart())
@@ -1481,10 +1484,7 @@ void UAssetRegistryImpl::OnInitialPluginLoadingComplete()
 	{
 		LLM_SCOPE(ELLMTag::AssetRegistry);
 		UE::AssetRegistry::FInterfaceWriteScopeLock InterfaceScopeLock(InterfaceLock);
-		// Read default serialization options
-		constexpr static bool bInitializeFilterTagOptions = true;
-		UE::AssetRegistry::Utils::InitializeSerializationOptionsFromIni(GuardedData.SerializationOptions, nullptr, UE::AssetRegistry::ESerializationTarget::ForGame, bInitializeFilterTagOptions);
-		UE::AssetRegistry::Utils::InitializeSerializationOptionsFromIni(GuardedData.DevelopmentSerializationOptions, nullptr, UE::AssetRegistry::ESerializationTarget::ForDevelopment, bInitializeFilterTagOptions);
+		UE::AssetRegistry::Utils::UpdateSerializationOptions(GuardedData.SerializationOptions, GuardedData.CachedParsedFilterRules);
 		GuardedData.OnPluginLoadingComplete(true);
 	}
 
@@ -1552,8 +1552,9 @@ void UAssetRegistryImpl::InitializeSerializationOptions(FAssetRegistrySerializat
 	}
 	else
 	{
-		constexpr static bool bInitializeFilterTagOptions = true;
-		UE::AssetRegistry::Utils::InitializeSerializationOptionsFromIni(Options, TargetPlatform, Target, bInitializeFilterTagOptions);
+		// Temp container to parse and resolve FilterTags
+		UE::AssetRegistry::Impl::FFilterTagRules TempParsedFilterRules;
+		UE::AssetRegistry::Utils::InitializeSerializationOptionsFromIni(Options, &TempParsedFilterRules, TargetPlatform, Target);
 	}
 }
 
@@ -1589,9 +1590,9 @@ static TSet<FName> MakeNameSet(const TArray<FString>& Strings)
 
 void InitializeSerializationOptionsFromIni(
 	FAssetRegistrySerializationOptions& Options,
+	Impl::FFilterTagRules* OutRules,
 	const ITargetPlatform* TargetPlatform,
-	UE::AssetRegistry::ESerializationTarget Target,
-	bool bInitializeFilterTagOptions)
+	UE::AssetRegistry::ESerializationTarget Target)
 {
 	FString PlatformIniName;
 	EBuildTargetType TargetType = EBuildTargetType::Game;
@@ -1646,10 +1647,9 @@ void InitializeSerializationOptionsFromIni(
 		}
 	}
 
-	EngineIni->GetBool(TEXT("AssetRegistry"), TEXT("bUseAssetRegistryTagsWhitelistInsteadOfBlacklist"), Options.bUseAssetRegistryTagsAllowListInsteadOfDenyList);
-
-	if (bInitializeFilterTagOptions)
+	if (OutRules)
 	{
+		EngineIni->GetBool(TEXT("AssetRegistry"), TEXT("bUseAssetRegistryTagsWhitelistInsteadOfBlacklist"), Options.bUseAssetRegistryTagsAllowListInsteadOfDenyList);
 		TArray<FString> FilterListItems;
 		if (Options.bUseAssetRegistryTagsAllowListInsteadOfDenyList)
 		{
@@ -1663,8 +1663,13 @@ void InitializeSerializationOptionsFromIni(
 		const bool bKeepDevelopmentTags = bForDevelopment || FParse::Param(FCommandLine::Get(), TEXT("ARKeepDevTags"));
 		const bool bLoadClasses = true;
 		const bool bIsFilteringDevelopmentAR = bKeepDevelopmentTags;
-		TMap<FName, FParsedFilterTagRules> Rules = ParseFilterTags(FilterListItems, bKeepDevelopmentTags, Options.bUseAssetRegistryTagsAllowListInsteadOfDenyList);
-		Options.CookFilterlistTagsByClass = ResolveFilterTags(Rules, bLoadClasses);
+		OutRules->Rules = ParseFilterTags(FilterListItems, bKeepDevelopmentTags, Options.bUseAssetRegistryTagsAllowListInsteadOfDenyList);
+		UpdateSerializationOptions(Options, *OutRules);
+	}
+	else
+	{
+		// By default all tags are included
+		Options.bUseAssetRegistryTagsAllowListInsteadOfDenyList = false;
 	}
 
 	{
@@ -1679,6 +1684,16 @@ void InitializeSerializationOptionsFromIni(
 	}
 }
 
+void UpdateSerializationOptions(
+	FAssetRegistrySerializationOptions& InOutOptions,
+	Impl::FFilterTagRules& InOutFilterRules)
+{
+	if (InOutFilterRules.bContainsAnyUnresolvedRule || (InOutFilterRules.ClassesFingerprint != GetRegisteredNativeClassesVersionNumber()))
+	{
+		InOutOptions.CookFilterlistTagsByClass = ResolveFilterTags(InOutFilterRules);
+	}
+}
+
 static int32 CalculateInheritanceDepth(const UClass* Class)
 {
 	int32 InheritanceDepth = 0;
@@ -1692,9 +1707,9 @@ static int32 CalculateInheritanceDepth(const UClass* Class)
 	return InheritanceDepth;
 }
 
-TMap<FName, FParsedFilterTagRules> ParseFilterTags(const TArray<FString>& FilterListItems, bool bIsFilteringDevelopmentAR , bool bUseAllowListInsteadOfDenyList)
+TMap<FName, Impl::FParsedFilterTagRules> ParseFilterTags(const TArray<FString>& FilterListItems, bool bIsFilteringDevelopmentAR , bool bUseAllowListInsteadOfDenyList)
 {
-	TMap<FName, FParsedFilterTagRules> TagsOrder;
+	TMap<FName, Impl::FParsedFilterTagRules> TagsOrder;
 
 	// Takes on the pattern "(Class=SomeClass,Tag=SomeTag)"
 	// Optional key KeepInDevOnly for tweaking a DevelopmentAssetRegistry (additive if allow list, subtractive if deny list)
@@ -1761,8 +1776,8 @@ TMap<FName, FParsedFilterTagRules> ParseFilterTags(const TArray<FString>& Filter
 
 		if (!ClassName.IsEmpty() && !TagName.IsEmpty())
 		{
-			FParsedFilterTagRules& Rules = TagsOrder.FindOrAdd(FName(TagName.Len(), *TagName));
-			TArray<FFilterTagsRule>& OrderedList = Rules.OrderedList;
+			Impl::FParsedFilterTagRules& Rules = TagsOrder.FindOrAdd(FName(TagName.Len(), *TagName));
+			TArray<Impl::FParsedFilterTagRules::FFilterTagsRule>& OrderedList = Rules.OrderedList;
 			const UClass* FilterlistClass = nullptr;
 			// Skip known wildcard
 			if (ClassName != TEXTVIEW("*"))
@@ -1780,7 +1795,15 @@ TMap<FName, FParsedFilterTagRules> ParseFilterTags(const TArray<FString>& Filter
 				OrderedList.Emplace(FilterlistClass->GetClassPathName(), FilterlistClass, InheritanceDepth, bAddToList);
 				if (!Rules.bContainsAnyUnresolvedRule)
 				{
-					AlgoImpl::HeapSiftUp(OrderedList.GetData(), 0, OrderedList.Num() - 1, [](const FFilterTagsRule& Info) { return Info.InheritanceDepth; }, TLess<int32>());
+					AlgoImpl::HeapSiftUp(
+						OrderedList.GetData(),
+						0,
+						OrderedList.Num() - 1,
+						[](const Impl::FParsedFilterTagRules::FFilterTagsRule& Info)
+							{
+								return Info.InheritanceDepth; 
+							},
+						TLess<int32>());
 				}
 			}
 			else
@@ -1820,7 +1843,22 @@ TMap<FName, FParsedFilterTagRules> ParseFilterTags(const TArray<FString>& Filter
 	return TagsOrder;
 }
 
-static void ResolveRules(FStringView TagName, TArray<FFilterTagsRule>& OrderedList, bool bLoadClasses)
+/**
+ * Returns the status of Engine module existance.
+ * Used to suppress emitting warnings about missing Engine classes.
+ */
+static bool HasEngineModuleLoaded()
+{
+	static bool bIsEngineModuleAvaliable = false;
+	if (bIsEngineModuleAvaliable)
+	{
+		return true;
+	}
+	bIsEngineModuleAvaliable = FModuleManager::Get().IsModuleLoaded(TEXT("Engine"));
+	return bIsEngineModuleAvaliable;
+}
+
+static void ResolveRules(FStringView TagName, TArray<Impl::FParsedFilterTagRules::FFilterTagsRule>& OrderedList, bool* bAreRulesResolved)
 {
 	auto RuleIt = OrderedList.CreateIterator();
 	if (RuleIt->ClassPath == UE::AssetRegistry::WildcardPathName)
@@ -1840,13 +1878,22 @@ static void ResolveRules(FStringView TagName, TArray<FFilterTagsRule>& OrderedLi
 			);
 			RuleIt.RemoveCurrentSwap();
 		}
+		// Wildcard is considered resolved rule
+		if (bAreRulesResolved)
+		{
+			*bAreRulesResolved = true;
+		}
 		return;
 	}
 
+	bool bContainAnyUnresolvedRules = false;
 	for (; RuleIt; ++RuleIt)
 	{
+		// Currently there is no case when we don't want to cause sync load of the class we try to resolve.
+		// Expose as a function parameter if that ever will need to change.
+		constexpr bool bLoadClasses = true;
 		RuleIt->ClassPtr = bLoadClasses
-			? StaticLoadClass(UObject::StaticClass(), nullptr, RuleIt->ClassPath.ToString())
+			? StaticLoadClass(UObject::StaticClass(), nullptr, RuleIt->ClassPath.ToString(), {}, LOAD_Quiet)
 			: Cast<UClass>(StaticFindObject(UClass::StaticClass(), RuleIt->ClassPath));
 		if (RuleIt->ClassPtr)
 		{
@@ -1855,56 +1902,66 @@ static void ResolveRules(FStringView TagName, TArray<FFilterTagsRule>& OrderedLi
 		else
 		{
 			// Silently ignore same unresolved rules
-			const bool bIsDuplicate = (RuleIt.GetIndex() != 0) && (OrderedList[0].bInclude == RuleIt->bInclude);
-			UE_CLOG(!bIsDuplicate, LogAssetRegistry, Warning,
+			const bool bIsDuplicate = (RuleIt.GetIndex() != 0) && (OrderedList[RuleIt.GetIndex()-1].bInclude == RuleIt->bInclude);
+			UE_CLOG(!bIsDuplicate && HasEngineModuleLoaded(), LogAssetRegistry, Warning,
 				TEXT("There is a rule to %s tag \"%.*s\" for class \"%s\". Class is unknown. Ignoring."),
 				RuleIt->bInclude ? TEXT("include") : TEXT("exclude"),
 				TagName.Len(), TagName.GetData(),
 				*RuleIt->ClassPath.ToString()
 			);
-			RuleIt.RemoveCurrentSwap();
+			// Setting unresolved rule to max_int causes resolved rules to be applied anyway and any unresolved to be skipped.
+			RuleIt->InheritanceDepth = MAX_int32;
+			bContainAnyUnresolvedRules = true;
 		}
 	}
-	Algo::HeapSortBy(OrderedList, [](const FFilterTagsRule& Info) { return Info.InheritanceDepth; }, TLess<int32>());
+
+	Algo::HeapSortBy(OrderedList, [](const Impl::FParsedFilterTagRules::FFilterTagsRule& Info) { return Info.InheritanceDepth; }, TLess<int32>());
+
+	if (bAreRulesResolved)
+	{
+		*bAreRulesResolved = !bContainAnyUnresolvedRules;
+	}
 }
 
-TMap<FTopLevelAssetPath, TSet<FName>> ResolveFilterTags(TMap<FName, FParsedFilterTagRules>& InOutRules, bool bLoadClasses)
+TMap<FTopLevelAssetPath, TSet<FName>> ResolveFilterTags(Impl::FFilterTagRules& InOutRules)
 {
 	TMap<FTopLevelAssetPath, TSet<FName>> Result;
 	// keep temp allocation around while we traverse the ordered list
 	TSet<const UClass*> ClassesToAdd;
+	InOutRules.bContainsAnyUnresolvedRule = false;
 	// Populate CookFilterlistTagsByClass by walking the ordered list
-	for (TTuple<FName, FParsedFilterTagRules>& Pair : InOutRules)
+	for (TTuple<FName, Impl::FParsedFilterTagRules>& Pair : InOutRules.Rules)
 	{
 		const FName Tag = Pair.Key;
 		TStringBuilder<512> TagNameBuilder;
 		Pair.Key.ToString(TagNameBuilder);
 		const FStringView TagName = TagNameBuilder.ToView();
-		FParsedFilterTagRules& TagRules = Pair.Value;
-		TArray<FFilterTagsRule>& OrderedList = TagRules.OrderedList;
+		Impl::FParsedFilterTagRules& TagRules = Pair.Value;
+		TArray<Impl::FParsedFilterTagRules::FFilterTagsRule>& OrderedList = TagRules.OrderedList;
 		if (TagRules.bContainsAnyUnresolvedRule)
 		{
-			const bool bIsWildcard = OrderedList.IsValidIndex(0) && (OrderedList[0].ClassPath == UE::AssetRegistry::WildcardPathName);
-			UE_CLOG(!bIsWildcard, LogAssetRegistry, Warning,
-				TEXT("Tag \"%.*s\" has unresolved rules because some classes were not found during INI parsing. Will try to resolve again."),
-				TagName.Len(), TagName.GetData()
-			);
-			ResolveRules(TagName, OrderedList, bLoadClasses);
-			TagRules.bContainsAnyUnresolvedRule = false;
+			bool bAreRulesResolved = false;
+			ResolveRules(TagName, OrderedList, &bAreRulesResolved);
+			TagRules.bContainsAnyUnresolvedRule = !bAreRulesResolved;
+		}
+		if (TagRules.bContainsAnyUnresolvedRule)
+		{
+			InOutRules.bContainsAnyUnresolvedRule = true;
 		}
 
 		ClassesToAdd.Reset();
 		// The list is sorted in heap order of Inheritance depth,
 		// so we can just walk through it and add/remove as needed
 		TArray<UClass*> Temp;
-		for (const FFilterTagsRule& Info : OrderedList)
+		for (const Impl::FParsedFilterTagRules::FFilterTagsRule& Info : OrderedList)
 		{
-			// Add Wildcard or unknown class (we guarantee to only have one of those)
 			if (!Info.ClassPtr)
 			{
-				UE_CLOG(OrderedList.Num() > 1, LogAssetRegistry, Warning, 
-					TEXT("There is more than one rule for tag \"%.*s\" with unknown class, ignoring."),
-					TagName.Len(), TagName.GetData()
+				const bool bIsWildcard = (Info.ClassPath == UE::AssetRegistry::WildcardPathName);
+				UE_CLOG(!bIsWildcard && HasEngineModuleLoaded(), LogAssetRegistry, Warning, 
+					TEXT("There is at least one rule for tag \"%.*s\" with unknown class \"%s\", ignoring the rest of them."),
+					TagName.Len(), TagName.GetData(),
+					*Info.ClassPath.ToString()
 				);
 				if (Info.bInclude)
 				{
@@ -1937,6 +1994,7 @@ TMap<FTopLevelAssetPath, TSet<FName>> ResolveFilterTags(TMap<FName, FParsedFilte
 			Result.FindOrAdd(Class->GetClassPathName()).Add(Tag);
 		}
 	}
+	InOutRules.ClassesFingerprint = GetRegisteredNativeClassesVersionNumber();
 
 	return Result;
 }
@@ -6300,13 +6358,16 @@ void FAssetRegistryImpl::AppendState(Impl::FEventContext& EventContext, const FA
 	FAssetRegistryAppendResult LocalAppendResult;
 	FAssetRegistryAppendResult* AppendResultPtr = bEmitAssetEvents ? &LocalAppendResult : nullptr;
 
-	State.InitializeFromExisting(
-		InState,
 #if WITH_EDITOR
-		DevelopmentSerializationOptions,
+	const FAssetRegistrySerializationOptions& Options = DevelopmentSerializationOptions;
+	//Development options have all filters disabled, so there is no need to resolve runtime class filter options
 #else
-		SerializationOptions,
+	FAssetRegistrySerializationOptions& Options = SerializationOptions;
+	Utils::UpdateSerializationOptions(Options, CachedParsedFilterRules);
 #endif
+	State.InitializeFromExisting(
+		InState,
+		Options,
 		Mode,
 		AppendResultPtr);
 
@@ -6407,9 +6468,10 @@ void FAssetRegistryImpl::GetAllocatedSize(bool bLogDetailed, SIZE_T& StateSize,
 		UE_LOG(LogAssetRegistry, Log, TEXT("AssetRegistry ClassGeneratorNames Size: %" SIZE_T_FMT "k"), ClassGeneratorNames.GetAllocatedSize() / 1024);
 		UE_LOG(LogAssetRegistry, Log, TEXT("AssetRegistry BackgroundResults Size: %" SIZE_T_FMT "k"), BackgroundResults.GetAllocatedSize() / 1024);
 		UE_LOG(LogAssetRegistry, Log, TEXT("AssetRegistry CachedPathTree Size: %" SIZE_T_FMT "k"), CachedPathTree.GetAllocatedSize() / 1024);
+		UE_LOG(LogAssetRegistry, Log, TEXT("AssetRegistry CachedParsedFilterRules Size: %" SIZE_T_FMT "k"), CachedParsedFilterRules.GetAllocatedSize() / 1024);
 	}
 	StaticSize = CachedEmptyPackages.GetAllocatedSize() + CachedBPInheritanceMap.GetAllocatedSize() + ClassGeneratorNames.GetAllocatedSize();
-	SearchSize = BackgroundResults.GetAllocatedSize() + CachedPathTree.GetAllocatedSize();
+	SearchSize = BackgroundResults.GetAllocatedSize() + CachedPathTree.GetAllocatedSize() + CachedParsedFilterRules.GetAllocatedSize();
 
 	if (bIsTempCachingEnabled && !bIsTempCachingAlwaysEnabled)
 	{
diff --git a/Engine/Source/Runtime/AssetRegistry/Private/AssetRegistryImpl.h b/Engine/Source/Runtime/AssetRegistry/Private/AssetRegistryImpl.h
--- a/Engine/Source/Runtime/AssetRegistry/Private/AssetRegistryImpl.h
+++ b/Engine/Source/Runtime/AssetRegistry/Private/AssetRegistryImpl.h
@@ -331,6 +331,29 @@ private:
 	// True if we ran out of time, EarlyExitCallback returned true, or RequestEarlyExit() was called
 	bool OutInterrupted = false;
 };
+
+/** Container holding information for runtime resolution of parsed filter tag rules */
+struct FParsedFilterTagRules
+{
+	struct FFilterTagsRule
+	{
+		FTopLevelAssetPath ClassPath;
+		const UClass* ClassPtr = nullptr;
+		int32 InheritanceDepth = 0; // used as a sorting criteria to order classes from base to leaf.
+		bool bInclude = false;
+	};
+	TArray<FFilterTagsRule> OrderedList;
+	bool bContainsAnyUnresolvedRule = false;
+};
+
+struct FFilterTagRules
+{
+	SIZE_T GetAllocatedSize() const { return Rules.GetAllocatedSize(); }
+
+	TMap<FName, FParsedFilterTagRules> Rules;
+	uint64 ClassesFingerprint = 0;
+	bool bContainsAnyUnresolvedRule = false;
+};
 } // namespace Impl within namespace UE::AssetRegistry
 } // namespace UE::AssetRegistry
 
@@ -716,6 +739,9 @@ private:
 	/** The map of classes to their parents, only full for offline blueprints */
 	TMap<FTopLevelAssetPath, FTopLevelAssetPath> CachedBPInheritanceMap;
 
+	/** Helper container to resolve filter rules */
+	Impl::FFilterTagRules CachedParsedFilterRules;
+
 	/** If true, search caching is enabled */
 	bool bIsTempCachingEnabled;
 
@@ -1050,37 +1076,34 @@ bool RunAssetThroughFilter_Unchecked(const FAssetData& AssetData, const FARCompi
  */
 void RunAssetsThroughFilter(TArray<FAssetData>& AssetDataList, const FARCompiledFilter& Filter, const EFilterMode FilterMode);
 
-/** This will always read the ini, public version may return cache */
+/** 
+ * This will always read the ini, public version may return cache.
+ * if OutRules is nullptr, no tags will be filtered.
+ */
 void InitializeSerializationOptionsFromIni(
 	FAssetRegistrySerializationOptions& Options,
+	Impl::FFilterTagRules* OutRules,
 	const ITargetPlatform* TargetPlatform = nullptr,
-	ESerializationTarget Target = ESerializationTarget::ForGame,
-	bool bInitializeFilterTagOptions = true);
+	ESerializationTarget Target = ESerializationTarget::ForGame);
 
-struct FFilterTagsRule
-{
-	FTopLevelAssetPath ClassPath;
-	const UClass* ClassPtr = nullptr;
-	int32 InheritanceDepth = 0;
-	bool bInclude = false;
-};
+/**
+ * Updates FAssetRegistrySerializationOptions filters with the information based on runtime classes.
+ * Options are parsed based on class name and need to be resolved using the class hierarchy,
+ * and since the class hierarchy can change during startup we may need to reexecute the resolve.
+ */
+void UpdateSerializationOptions(FAssetRegistrySerializationOptions& InOutOptions, Impl::FFilterTagRules& InOutFilterRules);
 
-struct FParsedFilterTagRules
-{
-	TArray<FFilterTagsRule> OrderedList;
-	bool bContainsAnyUnresolvedRule = false;
-};
 /**
  * Parses the list of CookedTagsBlacklist or CookedTagsWhitelist gathered from the INI.
  * Returns a mapping from a tag to a list of rules. May contain unresolved rules.
  */
-TMap<FName, FParsedFilterTagRules> ParseFilterTags(const TArray<FString>& FilterListItems, bool bIsFilteringDevelopmentAR, bool bUseAllowListInsteadOfDenyList);
+TMap<FName, Impl::FParsedFilterTagRules> ParseFilterTags(const TArray<FString>& FilterListItems, bool bIsFilteringDevelopmentAR, bool bUseAllowListInsteadOfDenyList);
 
 /**
  * Resolves and expands the list of rules to list of all derived loaded classes that satisfy those rules.
  * Then transposes rules to a map with a class as a key and a list of tags as value.
  */
-TMap<FTopLevelAssetPath, TSet<FName>> ResolveFilterTags(TMap<FName, FParsedFilterTagRules>& InOutRules, bool bLoadClasses);
+TMap<FTopLevelAssetPath, TSet<FName>> ResolveFilterTags(Impl::FFilterTagRules& InOutRules);
 
 /** Gets the current availability of an asset, primarily for streaming install purposes. */
 EAssetAvailability::Type GetAssetAvailability(const FAssetData& AssetData);
diff --git a/Engine/Source/Runtime/AssetRegistry/Public/AssetRegistry/AssetRegistryState.h b/Engine/Source/Runtime/AssetRegistry/Public/AssetRegistry/AssetRegistryState.h
--- a/Engine/Source/Runtime/AssetRegistry/Public/AssetRegistry/AssetRegistryState.h
+++ b/Engine/Source/Runtime/AssetRegistry/Public/AssetRegistry/AssetRegistryState.h
@@ -769,6 +769,8 @@ inline void FAssetRegistrySerializationOptions::DisableFilters()
 	bFilterAssetDataWithNoTags = false;
 	bFilterDependenciesWithNoTags = false;
 	bFilterSearchableNames = false;
+	bUseAssetRegistryTagsAllowListInsteadOfDenyList = false;
+	CookFilterlistTagsByClass.Empty();
 }
 
 inline void FAssetRegistrySerializationOptions::InitForDevelopment()
