[AssetRegistry] Suppress filter tags rule resolve warnings at initial load.

#tests local editor build
#rnx

diff --git a/Engine/Source/Runtime/AssetRegistry/Private/AssetRegistry.cpp b/Engine/Source/Runtime/AssetRegistry/Private/AssetRegistry.cpp
--- a/Engine/Source/Runtime/AssetRegistry/Private/AssetRegistry.cpp
+++ b/Engine/Source/Runtime/AssetRegistry/Private/AssetRegistry.cpp
@@ -1126,8 +1126,9 @@ void FAssetRegistryImpl::Initialize(UE::AssetRegistry::FInterfaceWriteScopeLock&
 #endif
 
 	// Read default serialization options
-	Utils::InitializeSerializationOptionsFromIni(SerializationOptions, &CachedParsedFilterRules, nullptr, UE::AssetRegistry::ESerializationTarget::ForGame);
-	Utils::InitializeSerializationOptionsFromIni(DevelopmentSerializationOptions, nullptr, nullptr, UE::AssetRegistry::ESerializationTarget::ForDevelopment);
+	const bool bSuppressWarnings = true;
+	Utils::InitializeSerializationOptionsFromIni(SerializationOptions, &CachedParsedFilterRules, nullptr, UE::AssetRegistry::ESerializationTarget::ForGame, bSuppressWarnings);
+	Utils::InitializeSerializationOptionsFromIni(DevelopmentSerializationOptions, nullptr, nullptr, UE::AssetRegistry::ESerializationTarget::ForDevelopment, bSuppressWarnings);
 
 	bool bStartedAsyncGather = false;
 	if (ShouldSearchAllAssetsAtStart())
@@ -1484,7 +1485,8 @@ void UAssetRegistryImpl::OnInitialPluginLoadingComplete()
 	{
 		LLM_SCOPE(ELLMTag::AssetRegistry);
 		UE::AssetRegistry::FInterfaceWriteScopeLock InterfaceScopeLock(InterfaceLock);
-		UE::AssetRegistry::Utils::UpdateSerializationOptions(GuardedData.SerializationOptions, GuardedData.CachedParsedFilterRules);
+		const bool bSuppressWarnings = !UE::AssetRegistry::Utils::HasEngineModuleLoaded();
+		UE::AssetRegistry::Utils::UpdateSerializationOptions(GuardedData.SerializationOptions, GuardedData.CachedParsedFilterRules, bSuppressWarnings);
 		GuardedData.OnPluginLoadingComplete(true);
 	}
 
@@ -1554,7 +1556,8 @@ void UAssetRegistryImpl::InitializeSerializationOptions(FAssetRegistrySerializat
 	{
 		// Temp container to parse and resolve FilterTags
 		UE::AssetRegistry::Impl::FFilterTagRules TempParsedFilterRules;
-		UE::AssetRegistry::Utils::InitializeSerializationOptionsFromIni(Options, &TempParsedFilterRules, TargetPlatform, Target);
+		const bool bSuppressWarnings = !UE::AssetRegistry::Utils::HasEngineModuleLoaded();
+		UE::AssetRegistry::Utils::InitializeSerializationOptionsFromIni(Options, &TempParsedFilterRules, TargetPlatform, Target, bSuppressWarnings);
 	}
 }
 
@@ -1592,7 +1595,8 @@ void InitializeSerializationOptionsFromIni(
 	FAssetRegistrySerializationOptions& Options,
 	Impl::FFilterTagRules* OutRules,
 	const ITargetPlatform* TargetPlatform,
-	UE::AssetRegistry::ESerializationTarget Target)
+	UE::AssetRegistry::ESerializationTarget Target,
+	bool bSuppressWarnings)
 {
 	FString PlatformIniName;
 	EBuildTargetType TargetType = EBuildTargetType::Game;
@@ -1664,7 +1668,7 @@ void InitializeSerializationOptionsFromIni(
 		const bool bLoadClasses = true;
 		const bool bIsFilteringDevelopmentAR = bKeepDevelopmentTags;
 		OutRules->Rules = ParseFilterTags(FilterListItems, bKeepDevelopmentTags, Options.bUseAssetRegistryTagsAllowListInsteadOfDenyList);
-		UpdateSerializationOptions(Options, *OutRules);
+		UpdateSerializationOptions(Options, *OutRules, bSuppressWarnings);
 	}
 	else
 	{
@@ -1686,11 +1690,12 @@ void InitializeSerializationOptionsFromIni(
 
 void UpdateSerializationOptions(
 	FAssetRegistrySerializationOptions& InOutOptions,
-	Impl::FFilterTagRules& InOutFilterRules)
+	Impl::FFilterTagRules& InOutFilterRules,
+	bool bSuppressWarnings)
 {
 	if (InOutFilterRules.bContainsAnyUnresolvedRule || (InOutFilterRules.ClassesFingerprint != GetRegisteredNativeClassesVersionNumber()))
 	{
-		InOutOptions.CookFilterlistTagsByClass = ResolveFilterTags(InOutFilterRules);
+		InOutOptions.CookFilterlistTagsByClass = ResolveFilterTags(InOutFilterRules, bSuppressWarnings);
 	}
 }
 
@@ -1843,11 +1848,7 @@ TMap<FName, Impl::FParsedFilterTagRules> ParseFilterTags(const TArray<FString>&
 	return TagsOrder;
 }
 
-/**
- * Returns the status of Engine module existance.
- * Used to suppress emitting warnings about missing Engine classes.
- */
-static bool HasEngineModuleLoaded()
+bool HasEngineModuleLoaded()
 {
 	static bool bIsEngineModuleAvaliable = false;
 	if (bIsEngineModuleAvaliable)
@@ -1858,7 +1859,11 @@ static bool HasEngineModuleLoaded()
 	return bIsEngineModuleAvaliable;
 }
 
-static void ResolveRules(FStringView TagName, TArray<Impl::FParsedFilterTagRules::FFilterTagsRule>& OrderedList, bool* bAreRulesResolved)
+static void ResolveRules(
+	FStringView TagName,
+	TArray<Impl::FParsedFilterTagRules::FFilterTagsRule>& OrderedList,
+	bool bSuppressWarnings,
+	bool* bAreRulesResolved)
 {
 	auto RuleIt = OrderedList.CreateIterator();
 	if (RuleIt->ClassPath == UE::AssetRegistry::WildcardPathName)
@@ -1867,7 +1872,7 @@ static void ResolveRules(FStringView TagName, TArray<Impl::FParsedFilterTagRules
 		for (++RuleIt; RuleIt; ++RuleIt)
 		{
 			const bool bIsDuplicate = (OrderedList[0].bInclude == RuleIt->bInclude);
-			UE_CLOG(!bIsDuplicate, LogAssetRegistry, Warning, 
+			UE_CLOG(!bSuppressWarnings && !bIsDuplicate, LogAssetRegistry, Warning, 
 				TEXT("There is a wildcard rule to %s tag \"%.*s\", ")
 				TEXT("but there is another rule to %s that tag for class \"%s\". ")
 				TEXT("Only wildcard rule is accepted. Ignoring another rule."),
@@ -1903,7 +1908,7 @@ static void ResolveRules(FStringView TagName, TArray<Impl::FParsedFilterTagRules
 		{
 			// Silently ignore same unresolved rules
 			const bool bIsDuplicate = (RuleIt.GetIndex() != 0) && (OrderedList[RuleIt.GetIndex()-1].bInclude == RuleIt->bInclude);
-			UE_CLOG(!bIsDuplicate && HasEngineModuleLoaded(), LogAssetRegistry, Warning,
+			UE_CLOG(!bSuppressWarnings && !bIsDuplicate, LogAssetRegistry, Warning,
 				TEXT("There is a rule to %s tag \"%.*s\" for class \"%s\". Class is unknown. Ignoring."),
 				RuleIt->bInclude ? TEXT("include") : TEXT("exclude"),
 				TagName.Len(), TagName.GetData(),
@@ -1923,7 +1928,7 @@ static void ResolveRules(FStringView TagName, TArray<Impl::FParsedFilterTagRules
 	}
 }
 
-TMap<FTopLevelAssetPath, TSet<FName>> ResolveFilterTags(Impl::FFilterTagRules& InOutRules)
+TMap<FTopLevelAssetPath, TSet<FName>> ResolveFilterTags(Impl::FFilterTagRules& InOutRules, bool bSuppressWarnings)
 {
 	TMap<FTopLevelAssetPath, TSet<FName>> Result;
 	// keep temp allocation around while we traverse the ordered list
@@ -1941,7 +1946,7 @@ TMap<FTopLevelAssetPath, TSet<FName>> ResolveFilterTags(Impl::FFilterTagRules& I
 		if (TagRules.bContainsAnyUnresolvedRule)
 		{
 			bool bAreRulesResolved = false;
-			ResolveRules(TagName, OrderedList, &bAreRulesResolved);
+			ResolveRules(TagName, OrderedList, bSuppressWarnings, &bAreRulesResolved);
 			TagRules.bContainsAnyUnresolvedRule = !bAreRulesResolved;
 		}
 		if (TagRules.bContainsAnyUnresolvedRule)
@@ -1957,12 +1962,6 @@ TMap<FTopLevelAssetPath, TSet<FName>> ResolveFilterTags(Impl::FFilterTagRules& I
 		{
 			if (!Info.ClassPtr)
 			{
-				const bool bIsWildcard = (Info.ClassPath == UE::AssetRegistry::WildcardPathName);
-				UE_CLOG(!bIsWildcard && HasEngineModuleLoaded(), LogAssetRegistry, Warning, 
-					TEXT("There is at least one rule for tag \"%.*s\" with unknown class \"%s\", ignoring the rest of them."),
-					TagName.Len(), TagName.GetData(),
-					*Info.ClassPath.ToString()
-				);
 				if (Info.bInclude)
 				{
 					Result.FindOrAdd(Info.ClassPath).Add(Tag);
@@ -6363,7 +6362,8 @@ void FAssetRegistryImpl::AppendState(Impl::FEventContext& EventContext, const FA
 	//Development options have all filters disabled, so there is no need to resolve runtime class filter options
 #else
 	FAssetRegistrySerializationOptions& Options = SerializationOptions;
-	Utils::UpdateSerializationOptions(Options, CachedParsedFilterRules);
+	const bool bSuppressWarnings = !UE::AssetRegistry::Utils::HasEngineModuleLoaded();
+	Utils::UpdateSerializationOptions(Options, CachedParsedFilterRules, bSuppressWarnings);
 #endif
 	State.InitializeFromExisting(
 		InState,
diff --git a/Engine/Source/Runtime/AssetRegistry/Private/AssetRegistryImpl.h b/Engine/Source/Runtime/AssetRegistry/Private/AssetRegistryImpl.h
--- a/Engine/Source/Runtime/AssetRegistry/Private/AssetRegistryImpl.h
+++ b/Engine/Source/Runtime/AssetRegistry/Private/AssetRegistryImpl.h
@@ -1084,14 +1084,24 @@ void InitializeSerializationOptionsFromIni(
 	FAssetRegistrySerializationOptions& Options,
 	Impl::FFilterTagRules* OutRules,
 	const ITargetPlatform* TargetPlatform = nullptr,
-	ESerializationTarget Target = ESerializationTarget::ForGame);
+	ESerializationTarget Target = ESerializationTarget::ForGame,
+	bool bSuppressWarnings = false);
+
+/**
+ * Returns the status of Engine module existance.
+ * Used to suppress emitting warnings about missing Engine classes.
+ */
+bool HasEngineModuleLoaded();
 
 /**
  * Updates FAssetRegistrySerializationOptions filters with the information based on runtime classes.
  * Options are parsed based on class name and need to be resolved using the class hierarchy,
  * and since the class hierarchy can change during startup we may need to reexecute the resolve.
  */
-void UpdateSerializationOptions(FAssetRegistrySerializationOptions& InOutOptions, Impl::FFilterTagRules& InOutFilterRules);
+void UpdateSerializationOptions(
+	FAssetRegistrySerializationOptions& InOutOptions,
+	Impl::FFilterTagRules& InOutFilterRules,
+	bool bSuppressWarnings);
 
 /**
  * Parses the list of CookedTagsBlacklist or CookedTagsWhitelist gathered from the INI.
@@ -1103,7 +1113,7 @@ TMap<FName, Impl::FParsedFilterTagRules> ParseFilterTags(const TArray<FString>&
  * Resolves and expands the list of rules to list of all derived loaded classes that satisfy those rules.
  * Then transposes rules to a map with a class as a key and a list of tags as value.
  */
-TMap<FTopLevelAssetPath, TSet<FName>> ResolveFilterTags(Impl::FFilterTagRules& InOutRules);
+TMap<FTopLevelAssetPath, TSet<FName>> ResolveFilterTags(Impl::FFilterTagRules& InOutRules, bool bSuppressWarnings);
 
 /** Gets the current availability of an asset, primarily for streaming install purposes. */
 EAssetAvailability::Type GetAssetAvailability(const FAssetData& AssetData);
