Small potential bug fix for IAssetRegistry/UAssetRegistryImpl::GetInMemoryAssets
- The asset filter flag bIncludeOnlyOnDiskAssets is ignored in the implementation, however EnumerateAssetsEvent is broadcast with the original filter, which could potentially have bIncludeOnlyOnDiskAssets set to true.
- This may cause problems later if a delegate is relying on the bIncludeOnlyOnDiskAssets filter flag to be set correctly (e.g the cooker)
- Removed the EnumerateAssetsEvent call from GetInMemoryAssets to reduce the chance of confusion later (it does not affect anything right now since the only place using this API is the content browser via UContentBrowserAssetDataSource, and the only thing binding to this event is the cooker)

#tests Cooked CitySample, Lyra

diff --git a/Engine/Source/Runtime/AssetRegistry/Private/AssetRegistry.cpp b/Engine/Source/Runtime/AssetRegistry/Private/AssetRegistry.cpp
--- a/Engine/Source/Runtime/AssetRegistry/Private/AssetRegistry.cpp
+++ b/Engine/Source/Runtime/AssetRegistry/Private/AssetRegistry.cpp
@@ -2811,9 +2811,6 @@ bool UAssetRegistryImpl::GetInMemoryAssets(const FARCompiledFilter& CompiledFilt
 {
 	using namespace UE::AssetRegistry::Utils;
 
-	const UE::AssetRegistry::EEnumerateAssetsFlags Flags = bSkipARFilteredAssets ? UE::AssetRegistry::EEnumerateAssetsFlags::None : UE::AssetRegistry::EEnumerateAssetsFlags::AllowUnfilteredArAssets;
-	EnumerateAssetsEvent.Broadcast(CompiledFilter, Flags);
-
 	TSet<FName> PackagesToSkipUnused;
 	bool bStopIterationUnused;
 	EnumerateMemoryAssets(CompiledFilter, PackagesToSkipUnused, bStopIterationUnused,
diff --git a/Engine/Source/Runtime/AssetRegistry/Public/AssetRegistry/IAssetRegistry.h b/Engine/Source/Runtime/AssetRegistry/Public/AssetRegistry/IAssetRegistry.h
--- a/Engine/Source/Runtime/AssetRegistry/Public/AssetRegistry/IAssetRegistry.h
+++ b/Engine/Source/Runtime/AssetRegistry/Public/AssetRegistry/IAssetRegistry.h
@@ -1139,7 +1139,11 @@ public:
 		~FPauseBackgroundProcessingScope() { IAssetRegistry::GetChecked().RequestResumeBackgroundProcessing(); }
 	};
 
-	/** Event triggered when a query is made. The event is triggered only if the filter is valid. */
+	/**
+	 * Event triggered when a query of the AssetRegistry's on-disk assets is made.
+	 * The event is triggered only if the filter is valid.
+	 * The event is not triggered for queries of in-memory assets only, such as GetInMemoryAssets.
+	 */
 	DECLARE_TS_MULTICAST_DELEGATE_TwoParams(FEnumerateAssetsEvent, const FARCompiledFilter& /*Filter*/, UE::AssetRegistry::EEnumerateAssetsFlags /*Flag*/)
 	virtual FEnumerateAssetsEvent& OnEnumerateAssetsEvent() = 0;
 
