Fix deadlock in IAssetRegistry->WaitForCompletion when called from a PostLoad function before completion.
#rnx

diff --git a/Engine/Source/Runtime/AssetRegistry/Private/AssetRegistry.cpp b/Engine/Source/Runtime/AssetRegistry/Private/AssetRegistry.cpp
--- a/Engine/Source/Runtime/AssetRegistry/Private/AssetRegistry.cpp
+++ b/Engine/Source/Runtime/AssetRegistry/Private/AssetRegistry.cpp
@@ -2505,6 +2505,35 @@ void UAssetRegistryImpl::WaitForCompletion()
 			}
 			break;
 		}
+		else if (Status == EGatherStatus::WaitingForEvents)
+		{
+			// If we are able to Broadcast events, then we will be able to do the broadcast and move past this status
+			// to completion. But if we cannot broadcast events, then we cannot complete the gather and have to exit.
+			// We only guarantee completion of the gather when our caller is in the right state to be able to broadcast.
+			// Some callers might not expect this, so issue a warning explaining the problem.
+			if (!CanBroadcastEvents())
+			{
+				const TCHAR* Reason = nullptr;
+				if (!IsInGameThread())
+				{
+					Reason = TEXT("caller is !IsInGameThread()");
+				}
+				else if (FUObjectThreadContext::Get().IsRoutingPostLoad)
+				{
+					Reason = TEXT("caller is IsRoutingPostLoad");
+				}
+				else
+				{
+					Reason = TEXT("<unknownreason>");
+				}
+
+				UE_LOG(LogAssetRegistry, Warning,
+					TEXT("Exiting from UAssetRegistryImpl::WaitForCompletion with gather almost completed ")
+					TEXT("but missing the final broadcast of events, so IsInitialSearchCompleted is still false. ")
+					TEXT("We cannot execute the broadcast because %s."), Reason);
+				break;
+			} 
+		}
 
 		FThreadHeartBeat::Get().HeartBeat();
 		if (Status == EGatherStatus::TickActiveGatherActive && bAsyncGathering)
@@ -10365,10 +10394,15 @@ void UAssetRegistryImpl::ReadLockEnumerateAllTagToAssetDatas(TFunctionRef<bool(F
 	GuardedData.GetState().EnumerateTagToAssetDatas(Callback);
 }
 
+bool UAssetRegistryImpl::CanBroadcastEvents() const
+{
+	return IsInGameThread() && !FUObjectThreadContext::Get().IsRoutingPostLoad;
+}
+
 void UAssetRegistryImpl::Broadcast(UE::AssetRegistry::Impl::FEventContext& EventContext, bool bAllowFileLoadedEvent)
 {
 	using namespace UE::AssetRegistry::Impl;
-	if (!IsInGameThread() || FUObjectThreadContext::Get().IsRoutingPostLoad)
+	if (!CanBroadcastEvents())
 	{
 		// By contract events (and PackageLoads) can only be sent on the game thread; some legacy systems depend on 
 		// this and are not threadsafe. If we're not in the game thread, defer all events in the EventContext
diff --git a/Engine/Source/Runtime/AssetRegistry/Private/AssetRegistry.h b/Engine/Source/Runtime/AssetRegistry/Private/AssetRegistry.h
--- a/Engine/Source/Runtime/AssetRegistry/Private/AssetRegistry.h
+++ b/Engine/Source/Runtime/AssetRegistry/Private/AssetRegistry.h
@@ -237,6 +237,7 @@ private:
 #endif
 	void InitializeEvents(UE::AssetRegistry::Impl::FInitializeContext& Context);
 	void Broadcast(UE::AssetRegistry::Impl::FEventContext& EventContext, bool bAllowFileLoadedEvent = false);
+	bool CanBroadcastEvents() const;
 
 	bool OnResolveRedirect(const FString& InPackageName, FString& OutPackageName);
 
