Add a bool to scan the paths passed in, which has been the default behaviour but can be very slow and a root scan + disable may be much faster
For now this can save a ton of time as each scan, even if already done ends up hitting at least 2 syscalls (GetFilenameOnDisk/GetStatData)

diff --git a/Engine/Source/Runtime/AssetRegistry/Private/AssetRegistryHelpers.cpp b/Engine/Source/Runtime/AssetRegistry/Private/AssetRegistryHelpers.cpp
--- a/Engine/Source/Runtime/AssetRegistry/Private/AssetRegistryHelpers.cpp
+++ b/Engine/Source/Runtime/AssetRegistry/Private/AssetRegistryHelpers.cpp
@@ -417,7 +417,7 @@ void UAssetRegistryHelpers::FixupRedirectedAssetPath(FName& InOutAssetPath)
 }
 
 #if WITH_EDITOR
-TArray<FAssetData> UAssetRegistryHelpers::GetAssetsWithOuterForPaths(const TArray<FName>& InPackagePaths, FName InOuterPath, bool bInRecursivePaths, bool bInIncludeOnlyOnDiskAssets, bool bInExactOuter)
+TArray<FAssetData> UAssetRegistryHelpers::GetAssetsWithOuterForPaths(const TArray<FName>& InPackagePaths, FName InOuterPath, bool bInRecursivePaths, bool bInIncludeOnlyOnDiskAssets, bool bInExactOuter, bool bScanPaths)
 {
 	FARFilter Filter;
 	Filter.PackagePaths = InPackagePaths;
@@ -430,7 +430,10 @@ TArray<FAssetData> UAssetRegistryHelpers::GetAssetsWithOuterForPaths(const TArra
 
 	TArray<FAssetData> Assets;
 	IAssetRegistry& AssetRegistry = IAssetRegistry::GetChecked();
-	AssetRegistry.ScanSynchronous(PathsToScan, TArray<FString>());
+	if (bScanPaths)
+	{
+		AssetRegistry.ScanSynchronous(PathsToScan, TArray<FString>());
+	}
 	AssetRegistry.GetAssets(Filter, Assets);
 
 	Assets.SetNum(
diff --git a/Engine/Source/Runtime/AssetRegistry/Public/AssetRegistry/AssetRegistryHelpers.h b/Engine/Source/Runtime/AssetRegistry/Public/AssetRegistry/AssetRegistryHelpers.h
--- a/Engine/Source/Runtime/AssetRegistry/Public/AssetRegistry/AssetRegistryHelpers.h
+++ b/Engine/Source/Runtime/AssetRegistry/Public/AssetRegistry/AssetRegistryHelpers.h
@@ -198,7 +198,8 @@ public:
 	 * @param	bInRecursivePaths If we should scan recursively or not.
 	 * @param	bInIncludeOnlyOnDiskAsset If we should only include on disk assets.
 	 * @param	bInExactOuter If we should only return assets directly outered to the provided outer.
+	 * @param	bScanPaths Scan the AssetRegistry for each path passed in
 	 */
-	 static ASSETREGISTRY_API TArray<FAssetData> GetAssetsWithOuterForPaths(const TArray<FName>& InPackagePaths, FName InOuterPath, bool bInRecursivePaths = false, bool bInIncludeOnlyOnDiskAsset = false, bool bInExactOuter = true);
+	 static ASSETREGISTRY_API TArray<FAssetData> GetAssetsWithOuterForPaths(const TArray<FName>& InPackagePaths, FName InOuterPath, bool bInRecursivePaths = false, bool bInIncludeOnlyOnDiskAsset = false, bool bInExactOuter = true, bool bScanPaths = true);
 #endif
 };
