Add an enum style function for GetAssetsWithOuterForPaths and deprecate the old way bool way
Remove from the old way the new bScanPaths, and only add it to the new enum style to avoid any breakage of ABI

diff --git a/Engine/Source/Editor/UnrealEd/Private/Commandlets/WorldPartitionConvertCommandlet.cpp b/Engine/Source/Editor/UnrealEd/Private/Commandlets/WorldPartitionConvertCommandlet.cpp
--- a/Engine/Source/Editor/UnrealEd/Private/Commandlets/WorldPartitionConvertCommandlet.cpp
+++ b/Engine/Source/Editor/UnrealEd/Private/Commandlets/WorldPartitionConvertCommandlet.cpp
@@ -806,7 +806,8 @@ int32 UWorldPartitionConvertCommandlet::Main(const FString& Params)
 			}
 		}
 
-		if (TArray<FAssetData> EntityAssets = UAssetRegistryHelpers::GetAssetsWithOuterForPaths({ *OldLevelName }, *OldLevelName, true, true, false); !EntityAssets.IsEmpty())
+		EAssetsWithOuterForPathsFlags OuterForPathsFlags = (EAssetsWithOuterForPathsFlags::RecursivePaths | EAssetsWithOuterForPathsFlags::IncludeOnlyOnDiskAsset | EAssetsWithOuterForPathsFlags::ScanPaths);
+		if (TArray<FAssetData> EntityAssets = UAssetRegistryHelpers::GetAssetsWithOuterForPaths({ *OldLevelName }, *OldLevelName, OuterForPathsFlags); !EntityAssets.IsEmpty())
 		{
 			for (const FAssetData& EntityAsset : EntityAssets)
 			{
diff --git a/Engine/Source/Runtime/AssetRegistry/Private/AssetRegistryHelpers.cpp b/Engine/Source/Runtime/AssetRegistry/Private/AssetRegistryHelpers.cpp
--- a/Engine/Source/Runtime/AssetRegistry/Private/AssetRegistryHelpers.cpp
+++ b/Engine/Source/Runtime/AssetRegistry/Private/AssetRegistryHelpers.cpp
@@ -417,12 +417,34 @@ void UAssetRegistryHelpers::FixupRedirectedAssetPath(FName& InOutAssetPath)
 }
 
 #if WITH_EDITOR
-TArray<FAssetData> UAssetRegistryHelpers::GetAssetsWithOuterForPaths(const TArray<FName>& InPackagePaths, FName InOuterPath, bool bInRecursivePaths, bool bInIncludeOnlyOnDiskAssets, bool bInExactOuter, bool bScanPaths)
+TArray<FAssetData> UAssetRegistryHelpers::GetAssetsWithOuterForPaths(const TArray<FName>& InPackagePaths, FName InOuterPath, bool bInRecursivePaths, bool bInIncludeOnlyOnDiskAssets, bool bInExactOuter)
+{
+	EAssetsWithOuterForPathsFlags Flags = EAssetsWithOuterForPathsFlags::ScanPaths;
+	if (bInRecursivePaths)
+	{
+		Flags |= EAssetsWithOuterForPathsFlags::RecursivePaths;
+	}
+
+	if (bInIncludeOnlyOnDiskAssets)
+	{
+		Flags |= EAssetsWithOuterForPathsFlags::IncludeOnlyOnDiskAsset;
+	}
+
+	if (bInExactOuter)
+	{
+		Flags |= EAssetsWithOuterForPathsFlags::ExactOuter;
+	}
+
+	return GetAssetsWithOuterForPaths(InPackagePaths, InOuterPath, Flags);
+}
+
+
+TArray<FAssetData> UAssetRegistryHelpers::GetAssetsWithOuterForPaths(const TArray<FName>& InPackagePaths, FName InOuterPath, EAssetsWithOuterForPathsFlags InFlags)
 {
 	FARFilter Filter;
 	Filter.PackagePaths = InPackagePaths;
-	Filter.bRecursivePaths = bInRecursivePaths;
-	Filter.bIncludeOnlyOnDiskAssets = bInIncludeOnlyOnDiskAssets;
+	Filter.bRecursivePaths = EnumHasAnyFlags(InFlags, EAssetsWithOuterForPathsFlags::RecursivePaths);
+	Filter.bIncludeOnlyOnDiskAssets = EnumHasAnyFlags(InFlags, EAssetsWithOuterForPathsFlags::IncludeOnlyOnDiskAsset);
 
 	TArray<FString> PathsToScan;
 	PathsToScan.Reserve(Filter.PackagePaths.Num());
@@ -430,12 +452,13 @@ TArray<FAssetData> UAssetRegistryHelpers::GetAssetsWithOuterForPaths(const TArra
 
 	TArray<FAssetData> Assets;
 	IAssetRegistry& AssetRegistry = IAssetRegistry::GetChecked();
-	if (bScanPaths)
+	if (EnumHasAnyFlags(InFlags, EAssetsWithOuterForPathsFlags::ScanPaths))
 	{
 		AssetRegistry.ScanSynchronous(PathsToScan, TArray<FString>());
 	}
 	AssetRegistry.GetAssets(Filter, Assets);
 
+	bool bInExactOuter = EnumHasAnyFlags(InFlags, EAssetsWithOuterForPathsFlags::ExactOuter);
 	Assets.SetNum(
 		Algo::RemoveIf(Assets, [InOuterPath, bInExactOuter](const FAssetData& InAssetData)
 		{
diff --git a/Engine/Source/Runtime/AssetRegistry/Public/AssetRegistry/AssetRegistryHelpers.h b/Engine/Source/Runtime/AssetRegistry/Public/AssetRegistry/AssetRegistryHelpers.h
--- a/Engine/Source/Runtime/AssetRegistry/Public/AssetRegistry/AssetRegistryHelpers.h
+++ b/Engine/Source/Runtime/AssetRegistry/Public/AssetRegistry/AssetRegistryHelpers.h
@@ -42,6 +42,20 @@ enum class EAssetRegistrySortOrder : uint8
 	Descending
 };
 
+enum class EAssetsWithOuterForPathsFlags
+{
+	None					= 0,
+	// If we should scan recursively or not
+	RecursivePaths			= (1 << 0),
+	// If we should only include on disk assets
+	IncludeOnlyOnDiskAsset	= (1 << 1),
+	// If we should only return assets directly outered to the provided outer
+	ExactOuter				= (1 << 2),
+	// Scan the AssetRegistry for each path passed in
+	ScanPaths				= (1 << 3)
+};
+ENUM_CLASS_FLAGS(EAssetsWithOuterForPathsFlags);
+
 UCLASS(transient)
 class UAssetRegistryHelpers : public UObject
 {
@@ -200,6 +214,15 @@ public:
 	 * @param	bInExactOuter If we should only return assets directly outered to the provided outer.
 	 * @param	bScanPaths Scan the AssetRegistry for each path passed in
 	 */
-	 static ASSETREGISTRY_API TArray<FAssetData> GetAssetsWithOuterForPaths(const TArray<FName>& InPackagePaths, FName InOuterPath, bool bInRecursivePaths = false, bool bInIncludeOnlyOnDiskAsset = false, bool bInExactOuter = true, bool bScanPaths = true);
+	UE_DEPRECATED(5.7, "Please use GetAssetsWithOuterForPaths with enum flags")
+	static ASSETREGISTRY_API TArray<FAssetData> GetAssetsWithOuterForPaths(const TArray<FName>& InPackagePaths, FName InOuterPath, bool bInRecursivePaths = false, bool bInIncludeOnlyOnDiskAsset = false, bool bInExactOuter = true);
+
+	/**
+	 * Returns all the corresponding assets from the provided paths that are outered to the provided outer name.
+	 * @param	InPackagePaths List of package paths to scan for assets.
+	 * @param	InOuterPath The outer path to filter asset with.
+	 * @param	bInFlags The flags to control the behavior of this function, see EAssetsWithOuterForPathsFlags
+	 */
+	static ASSETREGISTRY_API TArray<FAssetData> GetAssetsWithOuterForPaths(const TArray<FName>& InPackagePaths, FName InOuterPath, EAssetsWithOuterForPathsFlags Flags = (EAssetsWithOuterForPathsFlags::ExactOuter | EAssetsWithOuterForPathsFlags::ScanPaths));
 #endif
 };
