Remove doxygen param that was removed
Match doxygen param with comment

diff --git a/Engine/Source/Runtime/AssetRegistry/Public/AssetRegistry/AssetRegistryHelpers.h b/Engine/Source/Runtime/AssetRegistry/Public/AssetRegistry/AssetRegistryHelpers.h
--- a/Engine/Source/Runtime/AssetRegistry/Public/AssetRegistry/AssetRegistryHelpers.h
+++ b/Engine/Source/Runtime/AssetRegistry/Public/AssetRegistry/AssetRegistryHelpers.h
@@ -212,7 +212,6 @@ public:
 	 * @param	bInRecursivePaths If we should scan recursively or not.
 	 * @param	bInIncludeOnlyOnDiskAsset If we should only include on disk assets.
 	 * @param	bInExactOuter If we should only return assets directly outered to the provided outer.
-	 * @param	bScanPaths Scan the AssetRegistry for each path passed in
 	 */
 	UE_DEPRECATED(5.7, "Please use GetAssetsWithOuterForPaths with enum flags")
 	static ASSETREGISTRY_API TArray<FAssetData> GetAssetsWithOuterForPaths(const TArray<FName>& InPackagePaths, FName InOuterPath, bool bInRecursivePaths = false, bool bInIncludeOnlyOnDiskAsset = false, bool bInExactOuter = true);
@@ -221,8 +220,8 @@ public:
 	 * Returns all the corresponding assets from the provided paths that are outered to the provided outer name.
 	 * @param	InPackagePaths List of package paths to scan for assets.
 	 * @param	InOuterPath The outer path to filter asset with.
-	 * @param	bInFlags The flags to control the behavior of this function, see EAssetsWithOuterForPathsFlags
+	 * @param	InFlags The flags to control the behavior of this function, see EAssetsWithOuterForPathsFlags
 	 */
-	static ASSETREGISTRY_API TArray<FAssetData> GetAssetsWithOuterForPaths(const TArray<FName>& InPackagePaths, FName InOuterPath, EAssetsWithOuterForPathsFlags Flags = (EAssetsWithOuterForPathsFlags::ExactOuter | EAssetsWithOuterForPathsFlags::ScanPaths));
+	static ASSETREGISTRY_API TArray<FAssetData> GetAssetsWithOuterForPaths(const TArray<FName>& InPackagePaths, FName InOuterPath, EAssetsWithOuterForPathsFlags InFlags = (EAssetsWithOuterForPathsFlags::ExactOuter | EAssetsWithOuterForPathsFlags::ScanPaths));
 #endif
 };
