AssetRegistry: Fix infinite loop in startup gather in PostLoadAssetRegistryTags when there is a cycle in ObjectRedirectors.
#rnx

diff --git a/Engine/Source/Runtime/AssetRegistry/Private/AssetRegistry.cpp b/Engine/Source/Runtime/AssetRegistry/Private/AssetRegistry.cpp
--- a/Engine/Source/Runtime/AssetRegistry/Private/AssetRegistry.cpp
+++ b/Engine/Source/Runtime/AssetRegistry/Private/AssetRegistry.cpp
@@ -7062,9 +7062,42 @@ bool FAssetRegistryImpl::TryPostLoadAssetRegistryTags(FAssetData* AssetData)
 	FTopLevelAssetPath AssetClassPath;
 	AssetClassPath = AssetData->AssetClassPath;
 	AssetClass = FindObject<UClass>(AssetClassPath, EFindObjectFlags::ExactClass);
-
+	bool bInvalidAssetClass = false;
+
+	int32 CycleCount = 0;
+	// We want to avoid the expense of detecting the cycle for the common case.
+	// CycleCountBeforeCycleDetection is Tuned for the maximum depth encountered on a large project.
+	constexpr int32 CycleCountBeforeCycleDetection = 10;
+	TOptional<TSet<FTopLevelAssetPath>> CycleSet;
+	TOptional<TArray<FTopLevelAssetPath>> CycleArray;
 	while (!AssetClass)
 	{
+		if (++CycleCount > CycleCountBeforeCycleDetection)
+		{
+			if (!CycleSet)
+			{
+				CycleSet.Emplace();
+				CycleArray.Emplace();
+			}
+			bool bExists;
+			CycleSet->Add(AssetClassPath, &bExists);
+			CycleArray->Add(AssetClassPath);
+			if (bExists)
+			{
+				TStringBuilder<256> CycleText;
+				for (const FTopLevelAssetPath& CyclePath : *CycleArray)
+				{
+					CycleText << CyclePath.ToString() << TEXT(", ");
+				}
+				UE_LOG(LogAssetRegistry, Warning,
+					TEXT("Cycle detected in ParentClasses, ObjectRedirectors, or CoreRedirects, starting from AssetClassPath %s. Skipping AssetRegistry normalization of assets with that class.")
+					TEXT("\n\tCycle: { %s }"),
+					*AssetData->AssetClassPath.ToString(), *CycleText);
+				bInvalidAssetClass = true;
+				bCouldPostLoadAssetRegistryTags = true; // Nothing to load because the class is invalid, so mark that we're done with the postload
+				break;
+			}
+		}
 		// this is probably a blueprint that has not yet been loaded, try to find its native base class
 		const FTopLevelAssetPath* ParentClassPath = CachedBPInheritanceMap.Find(AssetClassPath);
 		if (ParentClassPath && !ParentClassPath->IsNull())
@@ -7118,7 +7151,7 @@ bool FAssetRegistryImpl::TryPostLoadAssetRegistryTags(FAssetData* AssetData)
 		// track down the classes for fixup or just give up
 		bMakeFinalChecks = true;
 	}
-	if (!AssetClass && bForceCompletionEvenIfPostLoadsFail)
+	if (!AssetClass && !bInvalidAssetClass && bForceCompletionEvenIfPostLoadsFail)
 	{
 		if (bMakeFinalChecks)
 		{
