AssetRegistry: Fix AssetData update logic for unloaded assets still tracked in AssetDataObjectPathsUpdatedOnLoad
* Previously, asset data entries could be incorrectly skipped from background updates if their FSoftObjectPath was still present in AssetDataObjectPathsUpdatedOnLoad, even though the asset itself was no longer loaded in memory. This caused stale cached AssetData to persist.

diff --git a/Engine/Source/Runtime/AssetRegistry/Private/AssetRegistry.cpp b/Engine/Source/Runtime/AssetRegistry/Private/AssetRegistry.cpp
--- a/Engine/Source/Runtime/AssetRegistry/Private/AssetRegistry.cpp
+++ b/Engine/Source/Runtime/AssetRegistry/Private/AssetRegistry.cpp
@@ -7463,7 +7463,7 @@ void FAssetRegistryImpl::AssetSearchDataGathered(Impl::FEventContext& EventConte
 		if (ExistingAssetData)
 		{
 #if WITH_EDITOR
-			if (AssetDataObjectPathsUpdatedOnLoad.Contains(BackgroundResult->GetSoftObjectPath()))
+			if (IsLoadedAndTrackedAsset(BackgroundResult->GetSoftObjectPath()))
 			{
 				// If the current AssetData came from a loaded asset, don't overwrite it with the new one from disk
 				// The loaded asset is more authoritative because it has run the postload steps.
@@ -8891,6 +8891,23 @@ void FAssetRegistryImpl::PushProcessLoadedAssetsBatch(Impl::FEventContext& Event
 	}
 }
 
+bool FAssetRegistryImpl::IsLoadedAndTrackedAsset(const FSoftObjectPath& InObjectPath) const
+{
+	// Must be in our "updated on load" set
+	if (!AssetDataObjectPathsUpdatedOnLoad.Contains(InObjectPath))
+	{
+		return false;
+	}
+
+	// Must also be loaded in memory
+	if (InObjectPath.ResolveObject() == nullptr)
+	{
+		return false;
+	}
+
+	return true;
+}
+
 }
 
 #endif // WITH_EDITOR
diff --git a/Engine/Source/Runtime/AssetRegistry/Private/AssetRegistryImpl.h b/Engine/Source/Runtime/AssetRegistry/Private/AssetRegistryImpl.h
--- a/Engine/Source/Runtime/AssetRegistry/Private/AssetRegistryImpl.h
+++ b/Engine/Source/Runtime/AssetRegistry/Private/AssetRegistryImpl.h
@@ -485,6 +485,7 @@ public:
 		bool bUpdateDiskCacheAfterLoad);
 	void PushProcessLoadedAssetsBatch(Impl::FEventContext& EventContext,
 		TArrayView<FAssetData> LoadedAssetDatas, TArrayView<const UObject*> UnprocessedFromBatch);
+	bool IsLoadedAndTrackedAsset(const FSoftObjectPath& InObjectPath) const;
 	/** Call LoadCalculatedDependencies on each Package updated after the last LoadCalculatedDependencies. */
 	void LoadCalculatedDependencies(TArray<FName>* AssetPackageNamessToCalculate, Impl::FClassInheritanceContext& InheritanceContext, 
 		TSet<FName>* InPackagesNeedingDependencyCalculation, Impl::FInterruptionContext& InOutInterruptionContext);
