Additional trace markers for AssetRegistry merging

diff --git a/Engine/Source/Runtime/AssetRegistry/Private/AssetRegistry.cpp b/Engine/Source/Runtime/AssetRegistry/Private/AssetRegistry.cpp
--- a/Engine/Source/Runtime/AssetRegistry/Private/AssetRegistry.cpp
+++ b/Engine/Source/Runtime/AssetRegistry/Private/AssetRegistry.cpp
@@ -6330,6 +6330,7 @@ void FAssetRegistryImpl::Serialize(FArchive& Ar, Impl::FEventContext& EventConte
 /** Append the assets from the incoming state into our own */
 void UAssetRegistryImpl::AppendState(const FAssetRegistryState& InState, UE::AssetRegistry::EAppendMode AppendMode)
 {
+	TRACE_CPUPROFILER_EVENT_SCOPE_STR("UAssetRegistryImpl::AppendState");
 	UE::AssetRegistry::Impl::FEventContext EventContext;
 	{
 		LLM_SCOPE(ELLMTag::AssetRegistry);
diff --git a/Engine/Source/Runtime/AssetRegistry/Private/AssetRegistryState.cpp b/Engine/Source/Runtime/AssetRegistry/Private/AssetRegistryState.cpp
--- a/Engine/Source/Runtime/AssetRegistry/Private/AssetRegistryState.cpp
+++ b/Engine/Source/Runtime/AssetRegistry/Private/AssetRegistryState.cpp
@@ -2410,6 +2410,7 @@ bool FAssetRegistryState::Load(FArchive& OriginalAr, const FAssetRegistryLoadOpt
 /* static */ bool FAssetRegistryState::LoadFromDisk(const TCHAR* InPath, const FAssetRegistryLoadOptions& InOptions,
 	FAssetRegistryState& OutState, FAssetRegistryVersion::Type* OutVersion)
 {
+	TRACE_CPUPROFILER_EVENT_SCOPE_STR("FAssetRegistryState::LoadFromDisk");
 	check(InPath);
 
 	TUniquePtr<FArchive> FileReader(IFileManager::Get().CreateFileReader(InPath));
