Plugin load optimization: if using Pak files, avoid falling back to loose files when checking if AssetRegistry.bin files exist. These are frequently not present, and the fallback can add seconds to the load time when we have thousands of plugins active.

#DisabledValidations Preflight, checkCopyright, checkBalancedMacros, CheckUnacceptableWords, CrossChangelist

diff --git a/Engine/Plugins/Runtime/GameFeatures/Source/GameFeatures/Private/GameFeaturePluginStateMachine.cpp b/Engine/Plugins/Runtime/GameFeatures/Source/GameFeatures/Private/GameFeaturePluginStateMachine.cpp
--- a/Engine/Plugins/Runtime/GameFeatures/Source/GameFeatures/Private/GameFeaturePluginStateMachine.cpp
+++ b/Engine/Plugins/Runtime/GameFeatures/Source/GameFeatures/Private/GameFeaturePluginStateMachine.cpp
@@ -7,6 +7,7 @@
 #include "GameFeatureData.h"
 #include "GameplayTagsManager.h"
 #include "Engine/AssetManager.h"
+#include "HAL/PlatformFileManager.h"
 #include "IPlatformFilePak.h"
 #include "InstallBundleUtils.h"
 #include "BundlePrereqCombinedStatusHelper.h"
@@ -2623,13 +2624,16 @@ struct FGameFeaturePluginState_Mounting : public FGameFeaturePluginState
 		{
 			const FString PluginFolder = FPaths::GetPath(StateProperties.PluginInstalledFilename);
 			TArray<FString> PluginAssetRegistrySearchPaths;
+			FPakPlatformFile* PakPlatformFile = (FPakPlatformFile*)(FPlatformFileManager::Get().FindPlatformFile(FPakPlatformFile::GetTypeName()));
 			// For GFPs cooked as DLC
 			PluginAssetRegistrySearchPaths.Add(PluginFolder / TEXT("AssetRegistry.bin"));
 			// For GFPs with a unique chunk
 			PluginAssetRegistrySearchPaths.Add(FPaths::ProjectDir() / FString::Printf(TEXT("AssetRegistry_GFP_%s.bin"), *StateProperties.PluginName));
 			for (FString& Path : PluginAssetRegistrySearchPaths)
 			{
-				if (IFileManager::Get().FileExists(*Path))
+				// Optimization: if we're using pak files then only search paks (avoid unnecessary fallback to loose)
+				bool bFileExists = PakPlatformFile != nullptr ? PakPlatformFile->FindFileInPakFiles(*Path) : IFileManager::Get().FileExists(*Path);
+				if (bFileExists)
 				{
 					PluginAssetRegistry = MoveTemp(Path);
 					break;
diff --git a/Engine/Source/Runtime/AssetRegistry/AssetRegistry.Build.cs b/Engine/Source/Runtime/AssetRegistry/AssetRegistry.Build.cs
--- a/Engine/Source/Runtime/AssetRegistry/AssetRegistry.Build.cs
+++ b/Engine/Source/Runtime/AssetRegistry/AssetRegistry.Build.cs
@@ -13,7 +13,8 @@ public class AssetRegistry : ModuleRules
 				"TraceLog",
 				"ApplicationCore",
 				"Projects",
-				"TelemetryUtils"
+				"TelemetryUtils",
+				"PakFile"
 			}
 			);
 
diff --git a/Engine/Source/Runtime/AssetRegistry/Private/AssetRegistry.cpp b/Engine/Source/Runtime/AssetRegistry/Private/AssetRegistry.cpp
--- a/Engine/Source/Runtime/AssetRegistry/Private/AssetRegistry.cpp
+++ b/Engine/Source/Runtime/AssetRegistry/Private/AssetRegistry.cpp
@@ -22,9 +22,11 @@
 #include "GenericPlatform/GenericPlatformChunkInstall.h"
 #include "GenericPlatform/GenericPlatformFile.h"
 #include "GenericPlatform/GenericPlatformMisc.h"
+#include "HAL/PlatformFileManager.h"
 #include "HAL/PlatformMisc.h"
 #include "HAL/ThreadHeartBeat.h"
 #include "Interfaces/IPluginManager.h"
+#include "IPlatformFilePak.h"
 #include "Math/UnrealMathUtility.h"
 #include "Misc/CommandLine.h"
 #include "Misc/ConfigCacheIni.h"
@@ -1014,6 +1016,7 @@ void FAssetRegistryImpl::LoadPremadeAssetRegistry(Impl::FEventContext& EventCont
 
 	{
 		SCOPED_BOOT_TIMING("LoadPremadeAssetRegistry_Plugins");
+		FPakPlatformFile* PakPlatformFile = (FPakPlatformFile*)(FPlatformFileManager::Get().FindPlatformFile(FPakPlatformFile::GetTypeName()));
 		TArray<TSharedRef<IPlugin>> ContentPlugins = IPluginManager::Get().GetEnabledPluginsWithContent();
 		for (const TSharedRef<IPlugin>& ContentPlugin : ContentPlugins)
 		{
@@ -1021,7 +1024,10 @@ void FAssetRegistryImpl::LoadPremadeAssetRegistry(Impl::FEventContext& EventCont
 			{
 				FArrayReader SerializedAssetData;
 				FString PluginAssetRegistry = ContentPlugin->GetBaseDir() / TEXT("AssetRegistry.bin");
-				if (IFileManager::Get().FileExists(*PluginAssetRegistry) && FFileHelper::LoadFileToArray(SerializedAssetData, *PluginAssetRegistry))
+
+				// Optimization: if we're using pak files then only search paks (avoid unnecessary fallback to loose)
+				bool bFileExists = PakPlatformFile != nullptr ? PakPlatformFile->FindFileInPakFiles(*PluginAssetRegistry) : IFileManager::Get().FileExists(*PluginAssetRegistry);
+				if (bFileExists && FFileHelper::LoadFileToArray(SerializedAssetData, *PluginAssetRegistry))
 				{
 					SerializedAssetData.Seek(0);
 					FAssetRegistryState PluginState;
diff --git a/Engine/Source/Runtime/PakFile/Public/IPlatformFilePak.h b/Engine/Source/Runtime/PakFile/Public/IPlatformFilePak.h
--- a/Engine/Source/Runtime/PakFile/Public/IPlatformFilePak.h
+++ b/Engine/Source/Runtime/PakFile/Public/IPlatformFilePak.h
@@ -2314,7 +2314,7 @@ public:
 	 * @param Paks Pak files to find the file in.
 	 * @param Filename File to find in pak files.
 	 * @param OutPakFile Optional pointer to a pak file where the filename was found.
-	 * @return Pointer to pak entry if the file was found, NULL otherwise.
+	 * @return true if the file was found, false otherwise.
 	 */
 	PAKFILE_API static bool FindFileInPakFiles(TArray<FPakListEntry>& Paks, const TCHAR* Filename,
 		TRefCountPtr<FPakFile>* OutPakFile, FPakEntry* OutEntry = nullptr);
@@ -2324,9 +2324,9 @@ public:
 	 *
 	 * @param Filename File to find in pak files.
 	 * @param OutPakFile Optional pointer to a pak file where the filename was found.
-	 * @return Pointer to pak entry if the file was found, NULL otherwise.
+	 * @return true if the file was found, false otherwise.
 	 */
-	bool FindFileInPakFiles(const TCHAR* Filename, TRefCountPtr<FPakFile>* OutPakFile = nullptr,
+	PAKFILE_API bool FindFileInPakFiles(const TCHAR* Filename, TRefCountPtr<FPakFile>* OutPakFile = nullptr,
 		FPakEntry* OutEntry = nullptr);
 
 	//~ Begin IPlatformFile Interface
