AssetRegistry: When ticking the AssetRegistry from the background gather thread, do not allow ThreadedPostLoadAssetRegistryTags to create SparseClassData. SparseClassData, like most class data, must only be written on the game thread. Writing to it from the background gather thread can cause data races or other problems of unsynchronized multithreadiing, for system-specific classes that rely on their class data only being written on the game thread.
#rnx

#FailedValidations Preflight

diff --git a/Engine/Source/Runtime/AssetRegistry/Private/AssetRegistry.cpp b/Engine/Source/Runtime/AssetRegistry/Private/AssetRegistry.cpp
--- a/Engine/Source/Runtime/AssetRegistry/Private/AssetRegistry.cpp
+++ b/Engine/Source/Runtime/AssetRegistry/Private/AssetRegistry.cpp
@@ -5234,9 +5234,12 @@ void UAssetRegistryImpl::Tick(float DeltaTime)
 		}
 		else
 		{
-			FScopeLock DeferredEventsLock(&DeferredEventsCriticalSection);
-			EventContext.Append(MoveTemp(DeferredEvents));
-			DeferredEvents.Clear();
+			{
+				FScopeLock DeferredEventsLock(&DeferredEventsCriticalSection);
+				EventContext.Append(MoveTemp(DeferredEvents));
+				DeferredEvents.Clear();
+			}
+			GuardedData.ProcessGameThreadRequestedClasses();
 		}
 
 #if WITH_EDITOR
@@ -5410,6 +5413,8 @@ Impl::EGatherStatus FAssetRegistryImpl::TickGatherer(Impl::FTickContext& TickCon
 		return TickContext.UpdateIntermediateStatus();
 	}
 
+	ConditionalRetryGatherAssetsForGameThread(TickContext);
+
 	bool bRetryAssetGathering = true;
 	int32 OriginalDeferredAssetsCount = 0;
 	int32 NumRetries = 0;
@@ -5431,6 +5436,9 @@ Impl::EGatherStatus FAssetRegistryImpl::TickGatherer(Impl::FTickContext& TickCon
 			{
 				return TickContext.UpdateIntermediateStatus();
 			}
+
+			// If we cleared out AssetsForGameThread, mark that we no longer need GameThreadRequestedClasses.
+			ConditionalClearGameThreadRequestedClasses(TickContext);
 		}
 
 		TSet<FName>* PackagesNeedingDependencyCalculationPointer = nullptr;
@@ -5472,9 +5480,9 @@ Impl::EGatherStatus FAssetRegistryImpl::TickGatherer(Impl::FTickContext& TickCon
 				}
 				else
 				{
-					if (!bForceCompletionEvenIfPostLoadsFail && bPreloadingComplete && IsEngineStartupModuleLoadingComplete())
+					if (!bForceCompletionEvenIfClassNotLoaded && bPreloadingComplete && IsEngineStartupModuleLoadingComplete())
 					{
-						bForceCompletionEvenIfPostLoadsFail = true;
+						bForceCompletionEvenIfClassNotLoaded = true;
 						bRetryAssetGathering = true;
 					}
 				}
@@ -5639,6 +5647,73 @@ Impl::EGatherStatus FAssetRegistryImpl::TickGatherer(Impl::FTickContext& TickCon
 	return EGatherStatus::Complete;
 }
 
+void FAssetRegistryImpl::ProcessGameThreadRequestedClasses()
+{
+	if (!bGameThreadRequestedClassesShouldProcess.load(std::memory_order_relaxed))
+	{
+		return;
+	}
+
+	TSet<FTopLevelAssetPath> ClassesToProcess;
+	{
+		FScopeLock GameThreadRequestedClassesScopeLock(&GameThreadRequestedClassesLock);
+		ClassesToProcess = MoveTemp(GameThreadRequestedClasses);
+		GameThreadRequestedClasses.Empty();
+		bGameThreadRequestedClassesShouldProcess.store(false, std::memory_order_relaxed);
+	}
+
+	for (const FTopLevelAssetPath& AssetClassPath : ClassesToProcess)
+	{
+		UClass* AssetClass = FindObject<UClass>(AssetClassPath, EFindObjectFlags::ExactClass);
+		if (!AssetClass)
+		{
+			continue;
+		}
+
+		// Process each of the possible properties we are waiting on for the GameThreadRequestedClass
+		if (AssetClass->GetSparseClassDataStruct())
+		{
+			AssetClass->GetSparseClassData(EGetSparseClassDataMethod::ArchetypeIfNull);
+		}
+	}
+	bGameThreadRequestedClassesShouldRetry.store(true, std::memory_order_release);
+}
+
+void FAssetRegistryImpl::ConditionalRetryGatherAssetsForGameThread(Impl::FTickContext& TickContext)
+{
+	if (!TickContext.bIsInGameThread && bGameThreadRequestedClassesShouldRetry.load(std::memory_order_relaxed))
+	{
+		// When we find we were instructed to retry, block on an acquire memory fence so that the
+		// modifications the game thread made to the classes are transferred from the other thread before
+		// we continue on.
+		std::atomic_thread_fence(std::memory_order_acquire);
+		bGameThreadRequestedClassesShouldRetry.store(false, std::memory_order_release);
+		for (TPair<FName, TUniquePtr<FAssetData>>& GameThreadPair : BackgroundResults.AssetsForGameThread)
+		{
+			BackgroundResults.Assets.Add(GameThreadPair.Key, MoveTemp(GameThreadPair.Value));
+		}
+		BackgroundResults.AssetsForGameThread.Empty();
+	}
+}
+
+void FAssetRegistryImpl::ConditionalClearGameThreadRequestedClasses(Impl::FTickContext& TickContext)
+{
+	if (BackgroundResults.AssetsForGameThread.IsEmpty()
+		&& bGameThreadRequestedClassesShouldProcess.load(std::memory_order_relaxed))
+	{
+		FScopeLock GameThreadRequestedClassesScopeLock(&GameThreadRequestedClassesLock);
+		GameThreadRequestedClasses.Empty();
+		bGameThreadRequestedClassesShouldProcess.store(false, std::memory_order_relaxed);
+	}
+}
+
+void FAssetRegistryImpl::RequestGameThreadProcessClass(const FTopLevelAssetPath& RequestedGameThreadClass)
+{
+	FScopeLock GameThreadRequestedClassesScopeLock(&GameThreadRequestedClassesLock);
+	GameThreadRequestedClasses.Add(RequestedGameThreadClass);
+	bGameThreadRequestedClassesShouldProcess.store(true, std::memory_order_relaxed);
+}
+
 namespace Impl // UE::AssetRegistry::Impl
 {
 
@@ -5799,8 +5874,28 @@ void FTickContext::RunAssetSearchDataGathered(TMultiMap<FName, TUniquePtr<FAsset
 	{
 		LazyStartTimer();
 
+		EDeferFlag DeferFlag;
+		if (bIsInGameThread)
+		{
+			DeferFlag = EDeferFlag::CanExecuteGameThread;
+		}
+		else if (!GuardedData.bForceCompletionEvenIfClassNotLoaded)
+		{
+			// If we are still able to wait on missing classes, then defer rather than kicking to game thread.
+			// The Gamethread might end up creating the required data for us if we continue waiting.
+			DeferFlag = EDeferFlag::RetryLater;
+		}
+		else if (!GuardedData.bForceCompletionEvenIfNotOnGameThread)
+		{
+			DeferFlag = EDeferFlag::QueueToGameThread;
+		}
+		else
+		{
+			DeferFlag = EDeferFlag::BypassNow;
+		}
+
 		GuardedData.AssetSearchDataGathered(EventContext, InAssetResults, OutDeferredAssetResults,
-			InterruptionContext, MountPointsForVerifyAfterGather);
+			InterruptionContext, MountPointsForVerifyAfterGather, DeferFlag);
 	}
 };
 
@@ -6025,6 +6120,8 @@ void FAssetRegistryImpl::TickGatherPackage(UE::AssetRegistry::FInterfaceWriteSco
 		return;
 	}
 
+	bool bIsInGameThread = IsInGameThread();
+
 	// Release the lock while performing expensive calls that don't need the lock
 	ScopeLock.Lock.WriteUnlock();
 	GlobalGatherer->WaitOnPath(LocalPath);
@@ -6091,7 +6188,8 @@ void FAssetRegistryImpl::TickGatherPackage(UE::AssetRegistry::FInterfaceWriteSco
 		TMultiMap<FName, TUniquePtr<FAssetData>> DeferredPackageAssetsMap;
 		UE::AssetRegistry::Impl::FInterruptionContext InterruptionContext;
 		AssetSearchDataGathered(EventContext, PackageAssetsMap, DeferredPackageAssetsMap,
-			InterruptionContext, MountPointsForVerifyAfterGather);
+			InterruptionContext, MountPointsForVerifyAfterGather,
+			bIsInGameThread ? EDeferFlag::CanExecuteGameThread : EDeferFlag::RetryLater);
 		if (DeferredPackageAssetsMap.Num())
 		{
 			UE_LOG(LogAssetRegistry, Warning, TEXT("Attempted to add package '%s' to the registry before its UClass was available. \
@@ -6939,24 +7037,28 @@ void FAssetRegistryImpl::ScanPathsSynchronous(FInterfaceWriteScopeLock* ScopeLoc
 				Iter.RemoveCurrent();
 			}
 		}
-		// Force AssetSearchDataGathered to process these assets, skipping the PostLoadAssetRegistryTags if needed
-		const bool bOldForceCompletionEvenIfPostLoadsFail = bForceCompletionEvenIfPostLoadsFail;
-		bForceCompletionEvenIfPostLoadsFail = true;
 
 		int32 OriginalNumDeferredAssetsForGameThread = DeferredAssetsForGameThread.Num();
 		TOptional<TSet<FString>> MountPointsForVerifyAfterGather;
 
-		// We don't call AssetsFoundCallback here because even for deferred assets it will already have been called.
-		// We pass DeferredAssetsForGameThread as the OutDeferred parameter, but we expect nothing will be deferred.
-		AssetSearchDataGathered(Context.EventContext, CollectedDeferredAssets, DeferredAssetsForGameThread,
-			TickContext.InterruptionContext, MountPointsForVerifyAfterGather);
+		// We don't call AssetsFoundCallback because even for deferred assets it will already have been called.
+		// We pass DeferredAssetsForGameThread as the OutDeferred parameter, but we expect nothing is deferred.
+		bool bIsInGameThread = IsInGameThread();
+		{
+			// Force AssetSearchDataGathered to process these assets, skipping the PostLoadAssetRegistryTags if needed
+			TGuardValue<bool> ScopedForceCompletionEventIfClassNotLoaded(bForceCompletionEvenIfClassNotLoaded, true);
+			TGuardValue<bool> ScopedForceCompletionEvenIfNotOnGameThread(bForceCompletionEvenIfNotOnGameThread, true);
+
+			AssetSearchDataGathered(Context.EventContext, CollectedDeferredAssets, DeferredAssetsForGameThread,
+				TickContext.InterruptionContext, MountPointsForVerifyAfterGather,
+				bIsInGameThread ? EDeferFlag::CanExecuteGameThread : EDeferFlag::BypassNow);
+		}
 		// All of the assets we collected should have been processed or deferred.
 		ensure(CollectedDeferredAssets.Num() == 0);
 
-		// We should not have deferred any new assets because we set bForceCompletionEvenIfPostLoadsFail=true
+		// We should not have deferred any new assets because we set the bForceCompletionEvenIf flags to true.
 		ensure(DeferredAssetsForGameThread.Num() <= OriginalNumDeferredAssetsForGameThread);
 
-		bForceCompletionEvenIfPostLoadsFail = bOldForceCompletionEvenIfPostLoadsFail;
 		// Tick to perform any subsequent processing required for these assets beyond AssetSearchDataGathered
 		Impl::FTickContext AssetTickContext(*this, Context.EventContext, Context.InheritanceContext);
 		Context.Status = TickGatherer(AssetTickContext);
@@ -7056,20 +7158,24 @@ FAssetData* FAssetRegistryImpl::ResolveAssetIdCollision(FAssetData& A, FAssetDat
 	return Keep;
 }
 
-bool FAssetRegistryImpl::TryPostLoadAssetRegistryTags(FAssetData* AssetData)
+bool FAssetRegistryImpl::TryPostLoadAssetRegistryTags(FAssetData* AssetData, EDeferFlag DeferFlag,
+	FTopLevelAssetPath& OutRequestedGameThreadClass)
 {
 	check(AssetData);
+	OutRequestedGameThreadClass = FTopLevelAssetPath();
 	if (!AssetData->TagsAndValues.Num())
 	{
 		return true;
 	}
 
-	bool bCouldPostLoadAssetRegistryTags = true;
+	// bForceCompletionEvenIfClassNotLoaded can be set in some cases even before engine startup is complete, for code
+	// that needs to access the assets immediately (e.g. ScanPathsSynchronous). Silently complete those cases. But if
+	// we are forcing completion because engine loading is complete, then log a warning for them.
+	bool bLogMissingClassWarnings = bPreloadingComplete && IsEngineStartupModuleLoadingComplete();
 	UClass* AssetClass = nullptr;
 	FTopLevelAssetPath AssetClassPath;
 	AssetClassPath = AssetData->AssetClassPath;
 	AssetClass = FindObject<UClass>(AssetClassPath, EFindObjectFlags::ExactClass);
-	bool bInvalidAssetClass = false;
 
 	int32 CycleCount = 0;
 	// We want to avoid the expense of detecting the cycle for the common case.
@@ -7100,9 +7206,10 @@ bool FAssetRegistryImpl::TryPostLoadAssetRegistryTags(FAssetData* AssetData)
 					TEXT("Cycle detected in ParentClasses, ObjectRedirectors, or CoreRedirects, starting from AssetClassPath %s. Skipping AssetRegistry normalization of assets with that class.")
 					TEXT("\n\tCycle: { %s }"),
 					*AssetData->AssetClassPath.ToString(), *CycleText);
-				bInvalidAssetClass = true;
-				bCouldPostLoadAssetRegistryTags = true; // Nothing to load because the class is invalid, so mark that we're done with the postload
-				break;
+
+				// We have given the relevant warning, there is nothing further we need to do with this class,
+				// so mark the postload completed.
+				return true;
 			}
 		}
 		// this is probably a blueprint that has not yet been loaded, try to find its native base class
@@ -7136,31 +7243,21 @@ bool FAssetRegistryImpl::TryPostLoadAssetRegistryTags(FAssetData* AssetData)
 			}
 			else
 			{
-				bCouldPostLoadAssetRegistryTags = false;
 				break;
 			}
 		}
 	}
 
-	// Now identify the most derived native class in the class hierarchy
-	if (AssetClass)
+	if (!AssetClass)
 	{
-		while (!AssetClass->HasAnyClassFlags(CLASS_Native))
+		if (!bForceCompletionEvenIfClassNotLoaded)
 		{
-			AssetClass = AssetClass->GetSuperClass();
+			// We need to wait for the class to load, return false until we are forced to complete.
+			return false;
 		}
-	}
 
-	bool bMakeFinalChecks = false;
-	if (bForceCompletionEvenIfPostLoadsFail && bPreloadingComplete && IsEngineStartupModuleLoadingComplete())
-	{
-		// Okay, we think we're done loading and now we need to make some expensive final checks to try to either
-		// track down the classes for fixup or just give up
-		bMakeFinalChecks = true;
-	}
-	if (!AssetClass && !bInvalidAssetClass && bForceCompletionEvenIfPostLoadsFail)
-	{
-		if (bMakeFinalChecks)
+		// Log a warning and mark the postload completed.
+		if (bLogMissingClassWarnings)
 		{
 			FString Reason;
 			if (AssetClassPath.ToString().StartsWith(TEXT("/Script/")))
@@ -7179,74 +7276,123 @@ bool FAssetRegistryImpl::TryPostLoadAssetRegistryTags(FAssetData* AssetData)
 				TEXT("Unable to PostLoadAssetRegistryTags for '%s' because ancestor class '%s' cannot be found. %s"), 
 				*AssetData->GetObjectPathString(), *AssetClassPath.ToString(), *Reason);
 		}
+		return true;
+	}
 
-		// Force this so that we can move on
-		bCouldPostLoadAssetRegistryTags = true;
+	// Now identify the most derived native class in the class hierarchy
+	while (!AssetClass->HasAnyClassFlags(CLASS_Native))
+	{
+		AssetClass = AssetClass->GetSuperClass();
 	}
 
-	if (AssetClass)
+	// Use the class's CDO to run the virtual functions in ThreadedPostLoadAssetRegistryTags
+	UObject* ClassDefaultObject = AssetClass->GetDefaultObject(false);
+	if (ClassDefaultObject)
 	{
-		bool bCompleted = false;
+		// It is only safe to read flags from the CDO if we have a barrier between the time where the default object
+		// pointer is published, and the time we read the flag. Otherwise, we could end up reading a stale flag that
+		// hasn't yet been set to RF_NeedInitialization. It needs to be paired with a release barrier when the default
+		// object is set.
+		std::atomic_thread_fence(std::memory_order_acquire);
+		TSAN_AFTER(ClassDefaultObject); // TSAN doesn't understand fence yet, help it understand what's going on
 
-		UObject* ClassDefaultObject = AssetClass->GetDefaultObject(false);
-		if (ClassDefaultObject)
+		if (ClassDefaultObject->HasAnyFlags(RF_NeedInitialization))
 		{
-			// The flag is only safe to read if we have a barrier between the time where the default object pointer is 
-			// published, and the time we read the flag. Otherwise, we could end up reading a stale flag that hasn't yet
-			// been set to RF_NeedInitialization. It needs to be paired with a release barrier when the default
-			// object is set.
-			std::atomic_thread_fence(std::memory_order_acquire);
-			TSAN_AFTER(ClassDefaultObject); // TSAN doesn't understand fence yet, help it understand what's going on
+			ClassDefaultObject = nullptr;
+		}
+	}
 
-			if (!ClassDefaultObject->HasAnyFlags(RF_NeedInitialization))
-			{
-				// We are using RF_NeedInitialization to guarantee that ClassDefaultObject is fully initialized
-				// potentially on another thread. For weakly ordered memory platforms, we need to 
-				// ensure that our read of the vtable ptr isn't performed prior to the read of the class flags
-				// otherwise we might see a stale vtable despite seeing RF_NeedInit clear.
-				std::atomic_thread_fence(std::memory_order_acquire);
-				TSAN_AFTER(ClassDefaultObject); // TSAN doesn't understand fence yet, help it understand what's going on
+	if (!ClassDefaultObject)
+	{
+		if (!bForceCompletionEvenIfClassNotLoaded)
+		{
+			// Continue waiting until the CDO is available
+			return false;
+		}
+		
+		// Log a warning and mark the postload completed.
+		if (bLogMissingClassWarnings)
+		{
+			ensureMsgf(false,
+				TEXT("Unable to PostLoadAssetRegistryTags for '%s' because the CDO for ancestor class '%s' could not be found or was not ready."),
+				*AssetData->GetObjectPathString(), *AssetClassPath.ToString());
+		}
 
-				TArray<UObject::FAssetRegistryTag> TagsToModify;
-				UObject::FPostLoadAssetRegistryTagsContext Context(*AssetData, TagsToModify);
-				ClassDefaultObject->ThreadedPostLoadAssetRegistryTags(Context);
-				if (TagsToModify.Num())
-				{
-					FAssetDataTagMap TagsAndValues = AssetData->TagsAndValues.CopyMap();
-					for (const UObject::FAssetRegistryTag& Tag : TagsToModify)
-					{
-						if (!Tag.Value.IsEmpty())
-						{
-							FString& Value = TagsAndValues.FindOrAdd(Tag.Name);
-							Value = Tag.Value;
-						}
-						else
-						{
-							TagsAndValues.Remove(Tag.Name);
-						}
-					}
-					AssetData->TagsAndValues = FAssetDataTagMapSharedView(MoveTemp(TagsAndValues));
-				}
+		return true;
+	}
+
+	// We are using RF_NeedInitialization to guarantee that ClassDefaultObject is fully initialized
+	// potentially on another thread. For weakly ordered memory platforms, we need to 
+	// ensure that our read of the vtable ptr isn't performed prior to the read of the class flags
+	// otherwise we might see a stale vtable despite seeing RF_NeedInit clear.
+	std::atomic_thread_fence(std::memory_order_acquire);
+	TSAN_AFTER(ClassDefaultObject); // TSAN doesn't understand fence yet, help it understand what's going on
+
+	TArray<UObject::FAssetRegistryTag> TagsToModify;
+	bool bCanCreateObjects = DeferFlag == EDeferFlag::CanExecuteGameThread;
+	UObject::FPostLoadAssetRegistryTagsContext Context(*AssetData, TagsToModify, bCanCreateObjects);
+	ClassDefaultObject->ThreadedPostLoadAssetRegistryTags(Context);
 
-				bCompleted = true;
+	if (!Context.PostLoadCouldComplete())
+	{
+		// If we are still able to wait on a missing class, then defer this asset before checking whether we need
+		// to kick it to game thread. The Gamethread might end up creating the required data for us if we
+		// continue waiting.
+		if (!bForceCompletionEvenIfClassNotLoaded)
+		{
+			// Continue waiting until the CDO is available
+			return false;
+		}
+
+		if (!bCanCreateObjects)
+		{
+			// Assume that we could not complete because we need bCanCreateObjects=true.
+			if (!bForceCompletionEvenIfNotOnGameThread && DeferFlag != EDeferFlag::BypassNow)
+			{
+				// Request completion on the game thread.
+				OutRequestedGameThreadClass = AssetClass->GetClassPathName();
+				return false;
 			}
+
+			// Log a warning and complete the PostLoad
+			UE_LOG(LogAssetRegistry, Display,
+				TEXT("Unable to PostLoadAssetRegistryTags for '%s' because class '%s' ThreadedPostLoadAssetRegistryTags requires GameThread to load its CDO, but we need it immediately when off the game thread."),
+				*AssetData->GetObjectPathString(), *AssetClass->GetPathName());
+			return true;
 		}
+		else
+		{
+			// We're on the gamethread, so we can't complete because of something else. Treat it the same as CDO not ready.
+			// Log a warning and complete the postload.
+			if (bLogMissingClassWarnings)
+			{
+				UE_LOG(LogAssetRegistry, Display,
+					TEXT("Unable to PostLoadAssetRegistryTags for '%s' because class '%s' ThreadedPostLoadAssetRegistryTags returned !PostLoadCouldComplete."),
+					*AssetData->GetObjectPathString(), *AssetClass->GetPathName());
+			}
 
-		if (!bCompleted)
+			return true;
+		}
+	}
+
+	if (TagsToModify.Num())
+	{
+		FAssetDataTagMap TagsAndValues = AssetData->TagsAndValues.CopyMap();
+		for (const UObject::FAssetRegistryTag& Tag : TagsToModify)
 		{
-			if (!bForceCompletionEvenIfPostLoadsFail)
+			if (!Tag.Value.IsEmpty())
 			{
-				bCouldPostLoadAssetRegistryTags = false;
+				FString& Value = TagsAndValues.FindOrAdd(Tag.Name);
+				Value = Tag.Value;
 			}
-			else 
+			else
 			{
-				ensureMsgf(!bMakeFinalChecks,
-					TEXT("Unable to PostLoadAssetRegistryTags for '%s' because the CDO for ancestor class '%s' could not be found or was not ready."),
-					*AssetData->GetObjectPathString(), *AssetClassPath.ToString());
+				TagsAndValues.Remove(Tag.Name);
 			}
 		}
+		AssetData->TagsAndValues = FAssetDataTagMapSharedView(MoveTemp(TagsAndValues));
 	}
-	return bCouldPostLoadAssetRegistryTags;
+	return true;
 }
 #endif
 
@@ -7366,7 +7512,8 @@ void FAssetRegistryImpl::AssetSearchDataGathered(Impl::FEventContext& EventConte
 	TMultiMap<FName, TUniquePtr<FAssetData>>& AssetResults,
 	TMultiMap<FName, TUniquePtr<FAssetData>>& OutDeferredAssetResults,
 	Impl::FInterruptionContext& InOutInterruptionContext,
-	TOptional<TSet<FString>>& MountPointsForVerifyAfterGather)
+	TOptional<TSet<FString>>& MountPointsForVerifyAfterGather,
+	EDeferFlag DeferFlag)
 {
 	TRACE_CPUPROFILER_EVENT_SCOPE(AssetSearchDataGathered);
 
@@ -7432,10 +7579,19 @@ void FAssetRegistryImpl::AssetSearchDataGathered(Impl::FEventContext& EventConte
 
 #if WITH_EDITOR
 		// Postload assets based on their declared class. Queue them for for later retry if their class has not yet loaded.
-		bool CouldPostLoad = TryPostLoadAssetRegistryTags(BackgroundResult.Get());
-		if (!CouldPostLoad)
+		FTopLevelAssetPath RequestedGameThreadClass;
+		bool bCouldPostLoad = TryPostLoadAssetRegistryTags(BackgroundResult.Get(), DeferFlag, RequestedGameThreadClass);
+		if (!bCouldPostLoad)
 		{
-			OutDeferredAssetResults.Add(BackgroundAssetPackageName, MoveTemp(BackgroundResult));
+			if (RequestedGameThreadClass.IsValid() && DeferFlag == EDeferFlag::QueueToGameThread)
+			{
+				BackgroundResults.AssetsForGameThread.Add(BackgroundAssetPackageName, MoveTemp(BackgroundResult));
+				RequestGameThreadProcessClass(RequestedGameThreadClass);
+			}
+			else
+			{
+				OutDeferredAssetResults.Add(BackgroundAssetPackageName, MoveTemp(BackgroundResult));
+			}
 			continue;
 		}
 #endif
diff --git a/Engine/Source/Runtime/AssetRegistry/Private/AssetRegistryImpl.h b/Engine/Source/Runtime/AssetRegistry/Private/AssetRegistryImpl.h
--- a/Engine/Source/Runtime/AssetRegistry/Private/AssetRegistryImpl.h
+++ b/Engine/Source/Runtime/AssetRegistry/Private/AssetRegistryImpl.h
@@ -355,6 +355,23 @@ struct FFilterTagRules
 	bool bContainsAnyUnresolvedRule = false;
 };
 } // namespace Impl within namespace UE::AssetRegistry
+
+/** Options for how AssetRegistry functions should handle operations that must only be run on the gamethread. */
+enum class EDeferFlag
+{
+	/** The function is running on the gamethread currently, it can execute the operation. */
+	CanExecuteGameThread,
+	/** The function should mark itself temporarily failed and the caller will try again later. */
+	RetryLater,
+	/** The function should send a message to the game thread to execute the operation there. */
+	QueueToGameThread,
+	/**
+	 * The function can not execute, can not defer, cannot queue; it has to complete its task now.
+	 * Skip the operation, and log warnings or errors if the skipping will cause problems.
+	 */
+	BypassNow,
+};
+
 } // namespace UE::AssetRegistry
 
 /** Returns true if ANY work remains to be done. This work might require the game thread. */
@@ -463,7 +480,20 @@ public:
 	void WaitForGathererIdle(float TimeoutSeconds);
 	/** Consume any results from the gatherer and return its status */
 	Impl::EGatherStatus TickGatherer(Impl::FTickContext& TickContext);
-
+	/**
+	 * During gamethread tick, if we find that the background thread is active in TickGatherer, process any classes
+	 * that the background requested from the gamethread tick.
+	 */
+	void ProcessGameThreadRequestedClasses();
+	/**
+	 * If TickContext is on a background thread and the game thread has executed ProcessGameThreadRequestedClasses,
+	 * retry all assets in AssetsForGameThread by moving them to the any-thread Assets queue.
+	 */
+	void ConditionalRetryGatherAssetsForGameThread(Impl::FTickContext& TickContext);
+	/** Clear GameThreadRequestedClasses if AssetsForGameThread is empty. */
+	void ConditionalClearGameThreadRequestedClasses(Impl::FTickContext& TickContext);
+	/** Add a class to GameThreadRequestedClasses and set the sentry flag to request processing. */
+	void RequestGameThreadProcessClass(const FTopLevelAssetPath& RequestedGameThreadClass);
 	/**
 	 * At some times during editor startup, depending on config settings, we might want to block
 	 * on the gather, and not tick the rest of the editor. When we want that to happen, we set
@@ -613,13 +643,14 @@ private:
 	/**
 	 * Called every tick to when data is retrieved by the background asset search.
 	 * If TickStartTime is < 0, the entire list of gathered assets will be cached. Also used in sychronous searches
-	 * The DeferredResults array contains assets that were not ready for processing due to required UClass's not being loaded
+	 * The DeferredResults array receives assets that were not ready for processing due to required UClass's not being loaded.
 	 */
 	void AssetSearchDataGathered(Impl::FEventContext& EventContext,
 		TMultiMap<FName, TUniquePtr<FAssetData>>& AssetResults, 
-		TMultiMap<FName, TUniquePtr<FAssetData>>& OutDeferredResults,\
+		TMultiMap<FName, TUniquePtr<FAssetData>>& OutDeferredResults,
 		Impl::FInterruptionContext& InOutInterruptionContext,
-		TOptional<TSet<FString>>& MountPointsForVerifyAfterGather);
+		TOptional<TSet<FString>>& MountPointsForVerifyAfterGather,
+		EDeferFlag DeferFlag);
 	/** Validate assets gathered from disk before adding them to the AssetRegistry. */
 	bool ShouldSkipGatheredAsset(FAssetData& AssetData);
 
@@ -666,9 +697,12 @@ private:
 	/**
 	 * Calls PostLoadAssetRegistryTags on the CDO of the asset class this data represents.
 	 * @param AssetData Existing asset data
-	 * @return Returns false if the required parent UClass is not yet available
+	 * @param OutRequestedGameThreadClass If the postload fails because a class needs processing on the gamethread, this
+	 *        variable receives the classname. Otherwise it is empty.
+	 * @return Returns false if required UClass data is not yet available and we are not forced to complete.
 	 */
-	bool TryPostLoadAssetRegistryTags(FAssetData* AssetData);
+	bool TryPostLoadAssetRegistryTags(FAssetData* AssetData, EDeferFlag DeferFlag,
+		FTopLevelAssetPath& OutRequestedGameThreadClass);
 
 	/** If new gatherers were registered, this call will refresh the mapping */
 	void RebuildAssetDependencyGathererMapIfNeeded();
@@ -764,6 +798,12 @@ private:
 	TMultiMap<FName, TUniquePtr<FAssetData>> DeferredAssetsForGameThread;
 	TMultiMap<FName, FPackageDependencyData> DeferredDependencies;
 	TMultiMap<FName, FPackageDependencyData> DeferredDependenciesForGameThread;
+	/** Classes that some of the AssetsForGameThread need processed on the gamethread before a retry. */
+	TSet<FTopLevelAssetPath> GameThreadRequestedClasses;
+	FCriticalSection GameThreadRequestedClassesLock;
+	/** Sentry values read by gamethread and gatherthread to performantly coordinate on GameThreadRequestedClasses. */
+	std::atomic<bool> bGameThreadRequestedClassesShouldRetry{ false };
+	std::atomic<bool> bGameThreadRequestedClassesShouldProcess{ false };
 
 #if !NO_LOGGING
 	/** Memory profiling information: How much memory is being used by the tags for each class. */
@@ -818,7 +858,8 @@ private:
 
 	bool bVerboseLogging;
 
-	bool bForceCompletionEvenIfPostLoadsFail = false;
+	bool bForceCompletionEvenIfClassNotLoaded = false;
+	bool bForceCompletionEvenIfNotOnGameThread = false;
 	bool bProcessedAnyAssetsAfterRetryDeferred = true;
 	float MaxSecondsPerFrame = 0.04f;
 
diff --git a/Engine/Source/Runtime/CoreUObject/Private/UObject/Class.cpp b/Engine/Source/Runtime/CoreUObject/Private/UObject/Class.cpp
--- a/Engine/Source/Runtime/CoreUObject/Private/UObject/Class.cpp
+++ b/Engine/Source/Runtime/CoreUObject/Private/UObject/Class.cpp
@@ -6372,10 +6372,10 @@ FArchive& operator<<(FArchive& Ar, FImplementedInterface& A)
 	return Ar;
 }
 
-const void* UClass::GetArchetypeForSparseClassData() const
+const void* UClass::GetArchetypeForSparseClassData(EGetSparseClassDataMethod GetMethod) const
 {
 	UClass* SuperClass = GetSuperClass();
-	return SuperClass ? SuperClass->GetSparseClassData(EGetSparseClassDataMethod::ArchetypeIfNull) : nullptr;
+	return SuperClass ? SuperClass->GetSparseClassData(GetMethod) : nullptr;
 }
 
 UScriptStruct* UClass::GetSparseClassDataArchetypeStruct() const
@@ -6737,14 +6737,22 @@ const void* UClass::GetSparseClassData(const EGetSparseClassDataMethod GetMethod
 	case EGetSparseClassDataMethod::ArchetypeIfNull:
 		// Use the archetype data only when it's the expected type, otherwise the result may be cast to an incorrect type by the caller
 		return SparseClassDataStruct == GetSparseClassDataArchetypeStruct()
-			? GetArchetypeForSparseClassData()
+			? GetArchetypeForSparseClassData(GetMethod)
 			: CreateSparseClassData();
 
+	case EGetSparseClassDataMethod::DeferIfNull:
+		// Use the archetype data only when it's the expected type, otherwise the result may be cast to an incorrect type by the caller
+		return SparseClassDataStruct == GetSparseClassDataArchetypeStruct()
+			? GetArchetypeForSparseClassData(GetMethod)
+			: nullptr;
+
+	case EGetSparseClassDataMethod::ReturnIfNull:
+		return nullptr;
+
 	default:
-		break;
+		checkNoEntry();
+		return nullptr;
 	}
-
-	return nullptr;
 }
 
 UScriptStruct* UClass::GetSparseClassDataStruct() const
diff --git a/Engine/Source/Runtime/CoreUObject/Private/UObject/Obj.cpp b/Engine/Source/Runtime/CoreUObject/Private/UObject/Obj.cpp
--- a/Engine/Source/Runtime/CoreUObject/Private/UObject/Obj.cpp
+++ b/Engine/Source/Runtime/CoreUObject/Private/UObject/Obj.cpp
@@ -2616,19 +2616,37 @@ void UObject::ThreadedPostLoadAssetRegistryTags(FPostLoadAssetRegistryTagsContex
 	ensureMsgf(GetClass()->HasAnyClassFlags(CLASS_Native), TEXT("ThreadedPostLoadAssetRegistryTags should not be called on non-native types. Detected a call on type '%s'"),
 		*GetClass()->GetName());
 
+	// Before executing any postloads, check whether we will be able to complete.
+	bool bHasAssetRegistrySearchableProperties = GetClass()->HasAssetRegistrySearchableProperties();
+	UScriptStruct* SparseClassDataStruct = nullptr;
+	const void* SparseClassData = nullptr;
+	if (bHasAssetRegistrySearchableProperties)
+	{
+		SparseClassDataStruct = GetClass()->GetSparseClassDataStruct();
+		if (SparseClassDataStruct)
+		{
+			SparseClassData = GetClass()->GetSparseClassData(
+				Context.CanCreateObjects() ? EGetSparseClassDataMethod::ArchetypeIfNull
+				: EGetSparseClassDataMethod::DeferIfNull);
+			if (!SparseClassData)
+			{
+				Context.SetPostLoadCouldComplete(false);
+				return;
+			}
+		}
+	}
+
 	ThreadedPostLoadAssetRegistryTagsOverride(Context);
 
-	if (GetClass()->HasAssetRegistrySearchableProperties())
+	if (bHasAssetRegistrySearchableProperties)
 	{
 		for (TFieldIterator<FProperty> PropertyIt(GetClass()); PropertyIt; ++PropertyIt)
 		{
 			PostLoadAssetRegistryTagProperty(*PropertyIt, Context);
 		}
 
-		UScriptStruct* SparseClassDataStruct = GetClass()->GetSparseClassDataStruct();
 		if (SparseClassDataStruct)
 		{
-			const void* SparseClassData = GetClass()->GetSparseClassData(EGetSparseClassDataMethod::ArchetypeIfNull);
 			for (TFieldIterator<FProperty> PropertyIt(SparseClassDataStruct); PropertyIt; ++PropertyIt)
 			{
 				PostLoadAssetRegistryTagProperty(*PropertyIt, Context);
diff --git a/Engine/Source/Runtime/CoreUObject/Public/UObject/Class.h b/Engine/Source/Runtime/CoreUObject/Public/UObject/Class.h
--- a/Engine/Source/Runtime/CoreUObject/Public/UObject/Class.h
+++ b/Engine/Source/Runtime/CoreUObject/Public/UObject/Class.h
@@ -3744,14 +3744,39 @@ struct FClassFunctionLinkInfo
 };
 
 
+/**
+ * Flag describing desired behavior of GetSparseClassData in the case where the SparseClassData has not been created.
+ * The class's instance is guaranteed already created if the class has sparse data of its own, so the only differences
+ * in behavior are when the class does not have sparse data of its own.
+ */
 enum class EGetSparseClassDataMethod : uint8
 {
-	/** Create a new instance when this class doesn't have any sparse data of its own */
+	/**
+	 * If not yet created, store and return a new instance, whether or not the class can use its archetype's instance.
+	 */
 	CreateIfNull,
-	/** Use the archetype instance (if possible) when this class doesn't have any sparse data of its own */
+	/**
+	 * If not yet created, return the archetype instance (if possible to use for this class). If not possible to use
+	 * the archetype's instance, store and return a new instance.
+	 */
 	ArchetypeIfNull,
-	/** Return null when this class doesn't have any sparse data of its own */
+	/**
+	 * This is an advanced flag and should usually be avoided. Callers should instead call with ArchetypeIfNull if they
+	 * are currently allowed to create an instance, or with DeferIfNull if not. Callers that want to know whether the
+	 * class has sparse data of its own beyond what the archetype has should compare
+	 * GetSparseClassDataStruct() != GetSparseClassDataArchetypeStruct().
+	 * 
+	 * If not yet created, return null. An instance might still be created later when called with ArchetypeIfNull.
+	 * This call never creates an instance.
+	 */
 	ReturnIfNull,
+	/**
+	 * If not yet created, return the archetype's instance (if possible to use for this class). If not possible to use
+	 * the archetype's instance, return null. User should call later with ArchetypeIfNull when they are allowed to
+	 * create an instance.
+	 * This call never creates an instance.
+	 */
+	DeferIfNull,
 };
 
 template<class T>
@@ -4616,7 +4641,8 @@ public:
 	COREUOBJECT_API virtual UObject* GetArchetypeForCDO() const;
 
 	/** Returns archetype for sparse class data */
-	COREUOBJECT_API const void* GetArchetypeForSparseClassData() const;
+	COREUOBJECT_API const void* GetArchetypeForSparseClassData(
+		const EGetSparseClassDataMethod GetMethod = EGetSparseClassDataMethod::ArchetypeIfNull) const;
 
 	/** Returns the struct used by the sparse class data archetype */
 	COREUOBJECT_API UScriptStruct* GetSparseClassDataArchetypeStruct() const;
diff --git a/Engine/Source/Runtime/CoreUObject/Public/UObject/Object.h b/Engine/Source/Runtime/CoreUObject/Public/UObject/Object.h
--- a/Engine/Source/Runtime/CoreUObject/Public/UObject/Object.h
+++ b/Engine/Source/Runtime/CoreUObject/Public/UObject/Object.h
@@ -935,17 +935,23 @@ public:
 	/** Contains an AssetData referencing the asset to be processed and an array of tags and values to be updated on that asset */
 	struct FPostLoadAssetRegistryTagsContext
 	{
-		FPostLoadAssetRegistryTagsContext(const FAssetData& InAssetData, TArray<FAssetRegistryTag>& InTagsAndValuesToUpdate)
-		: AssetData(InAssetData), TagsAndValuesToUpdate(InTagsAndValuesToUpdate) {}
+		FPostLoadAssetRegistryTagsContext(const FAssetData& InAssetData, TArray<FAssetRegistryTag>& InTagsAndValuesToUpdate,
+			bool bInCanCreateObjects)
+		: AssetData(InAssetData), TagsAndValuesToUpdate(InTagsAndValuesToUpdate), bCanCreateObjects(bInCanCreateObjects) { }
 
 		const FAssetData& GetAssetData() const { return AssetData; }
 		void AddTagToUpdate(const FAssetRegistryTag& Tag) { TagsAndValuesToUpdate.Add(Tag); }
 
+		bool CanCreateObjects() const { return bCanCreateObjects; }
+		bool PostLoadCouldComplete() const { return bPostLoadCouldComplete; }
+		void SetPostLoadCouldComplete(bool bValue) { bPostLoadCouldComplete = bValue; }
 	private:
 		/** The AssetData being considered */
 		const FAssetData& AssetData;
 		/** The asset registry tags to be updated for this asset data */
 		TArray<FAssetRegistryTag>& TagsAndValuesToUpdate;
+		bool bCanCreateObjects = true;
+		bool bPostLoadCouldComplete = true;
 	};
 
 protected:
diff --git a/Engine/Source/Runtime/Engine/Private/Blueprint.cpp b/Engine/Source/Runtime/Engine/Private/Blueprint.cpp
--- a/Engine/Source/Runtime/Engine/Private/Blueprint.cpp
+++ b/Engine/Source/Runtime/Engine/Private/Blueprint.cpp
@@ -1221,7 +1221,7 @@ void UBlueprint::ThreadedPostLoadAssetRegistryTagsOverride(FPostLoadAssetRegistr
 
 void UBlueprint::PostLoadBlueprintAssetRegistryTags(const FAssetData& InAssetData, TArray<FAssetRegistryTag>& OutTagsAndValuesToUpdate)
 {
-	FPostLoadAssetRegistryTagsContext Context(InAssetData, OutTagsAndValuesToUpdate);
+	FPostLoadAssetRegistryTagsContext Context(InAssetData, OutTagsAndValuesToUpdate, IsInGameThread());
 	PostLoadBlueprintAssetRegistryTags(Context);
 }
 
