Exposed the ability to include or exclude editor-only references from script asset registry queries

diff --git a/Engine/Source/Runtime/AssetRegistry/Private/AssetRegistry.cpp b/Engine/Source/Runtime/AssetRegistry/Private/AssetRegistry.cpp
--- a/Engine/Source/Runtime/AssetRegistry/Private/AssetRegistry.cpp
+++ b/Engine/Source/Runtime/AssetRegistry/Private/AssetRegistry.cpp
@@ -10376,13 +10376,21 @@ bool FAssetRegistryImpl::SetPrimaryAssetIdForObjectPath(Impl::FEventContext& Eve
 bool FAssetRegistryDependencyOptions::GetPackageQuery(UE::AssetRegistry::FDependencyQuery& Flags) const
 {
 	Flags = UE::AssetRegistry::FDependencyQuery();
+
+	bool bHasFilter = false;
 	if (bIncludeSoftPackageReferences || bIncludeHardPackageReferences)
 	{
 		if (!bIncludeSoftPackageReferences) Flags.Required |= UE::AssetRegistry::EDependencyProperty::Hard;
 		if (!bIncludeHardPackageReferences) Flags.Excluded |= UE::AssetRegistry::EDependencyProperty::Hard;
-		return true;
+		bHasFilter = true;
 	}
-	return false;
+	if (bIncludeGamePackageReferences || bIncludeEditorOnlyPackageReferences)
+	{
+		if (!bIncludeGamePackageReferences) Flags.Excluded |= UE::AssetRegistry::EDependencyProperty::Game;
+		if (!bIncludeEditorOnlyPackageReferences) Flags.Required |= UE::AssetRegistry::EDependencyProperty::Game;
+		bHasFilter = true;
+	}
+	return bHasFilter;
 }
 
 bool FAssetRegistryDependencyOptions::GetSearchableNameQuery(UE::AssetRegistry::FDependencyQuery& Flags) const
diff --git a/Engine/Source/Runtime/AssetRegistry/Public/AssetRegistry/IAssetRegistry.h b/Engine/Source/Runtime/AssetRegistry/Public/AssetRegistry/IAssetRegistry.h
--- a/Engine/Source/Runtime/AssetRegistry/Public/AssetRegistry/IAssetRegistry.h
+++ b/Engine/Source/Runtime/AssetRegistry/Public/AssetRegistry/IAssetRegistry.h
@@ -157,6 +157,14 @@ struct FAssetRegistryDependencyOptions
 	UPROPERTY(EditAnywhere, BlueprintReadWrite, Category = "AssetRegistry")
 	bool bIncludeHardPackageReferences = true;
 
+	/** Dependencies which are needed in the game as well as the editor. */
+	UPROPERTY(EditAnywhere, BlueprintReadWrite, Category = "AssetRegistry")
+	bool bIncludeGamePackageReferences = true;
+
+	/** Dependencies which are only needed in the editor. */
+	UPROPERTY(EditAnywhere, BlueprintReadWrite, Category = "AssetRegistry")
+	bool bIncludeEditorOnlyPackageReferences = true;
+
 	/** References to specific SearchableNames inside a package */
 	UPROPERTY(EditAnywhere, BlueprintReadWrite, Category = "AssetRegistry")
 	bool bIncludeSearchableNames = false;
