[AssetRegistry]

When SearchingAllAssets do not treat this call as starting an additional mount scan if we aren't gathering already. OnContentMounted will properly flag additional mount scanning. This prevents us from issuing additional cache saves for content that might have already been scanned and saved.

#rnx
#FailedValidations Preflight

diff --git a/Engine/Source/Runtime/AssetRegistry/Private/AssetRegistry.cpp b/Engine/Source/Runtime/AssetRegistry/Private/AssetRegistry.cpp
--- a/Engine/Source/Runtime/AssetRegistry/Private/AssetRegistry.cpp
+++ b/Engine/Source/Runtime/AssetRegistry/Private/AssetRegistry.cpp
@@ -2349,16 +2349,6 @@ void FAssetRegistryImpl::SearchAllAssets(Impl::FEventContext& EventContext,
 #endif
 		UpdateMaxSecondsPerFrame(); // MaxSecondsPerFrame only depends on initial search
 	}
-	else if (!IsGathering())
-	{
-		TRACE_BEGIN_REGION(TEXT("Asset Registry - Additional Mount Search")); // Matching TRACE_END_REGION in OnAdditionalMountSearchCompleted
-		AdditionalMountSearchStartTime = FPlatformTime::Seconds();
-		bAdditionalMountSearchInProgress.store(true, std::memory_order_relaxed);
-		GlobalGatherer->SetIsAdditionalMountSearchInProgress(true);
-#if WITH_EDITOR
-		bGatherNeedsApplyObjectRedirects = true;
-#endif
-	}
 
 	FAssetDataGatherer& Gatherer = *GlobalGatherer;
 	if (!Gatherer.IsAsyncEnabled())
