Disable false positive in warning in static analysis

diff --git a/Engine/Source/Runtime/AssetRegistry/Private/AssetRegistry.cpp b/Engine/Source/Runtime/AssetRegistry/Private/AssetRegistry.cpp
--- a/Engine/Source/Runtime/AssetRegistry/Private/AssetRegistry.cpp
+++ b/Engine/Source/Runtime/AssetRegistry/Private/AssetRegistry.cpp
@@ -6311,7 +6311,7 @@ void FAssetRegistryImpl::LoadCalculatedDependencies(FName PackageName,
 				bool bWasSourceDependencyListSorted = SourceNode->IsDependencyListSorted(EDependencyCategory::Package);
 				SourceNode->SetIsDependencyListSorted(EDependencyCategory::Package, false);
 
-				for (const IAssetDependencyGatherer::FGathereredDependency& GatheredDepencency : GatheredDependencies)
+				for (const IAssetDependencyGatherer::FGathereredDependency& GatheredDepencency : GatheredDependencies) //-V1078 
 				{
 					FDependsNode* TargetNode = State.CreateOrFindDependsNode(FAssetIdentifier(GatheredDepencency.PackageName));
 					EDependencyProperty DependencyProperties = GatheredDepencency.Property;
