[AssetRegistry]

Change FAssetDataGatherer::OnAdditionalMountSearchCompleted to no longer save the AR cache immediately but instead start a timer to save the cache the next time the gatherer is ticked _and_ FGatherSettings::MinSecondsToElapseBeforeCacheWrite have elapsed.

This is to minimize the amount of time spend saving the cache. Depending on when the AR initial gather completes, if multiple mounts are added, the AR would save the cache once for each mount causing noticable stalls and longer walltime.

#rnx
#FailedValidations Preflight

diff --git a/Engine/Source/Runtime/AssetRegistry/Private/AssetDataGatherer.cpp b/Engine/Source/Runtime/AssetRegistry/Private/AssetDataGatherer.cpp
--- a/Engine/Source/Runtime/AssetRegistry/Private/AssetDataGatherer.cpp
+++ b/Engine/Source/Runtime/AssetRegistry/Private/AssetDataGatherer.cpp
@@ -3792,6 +3792,7 @@ FAssetDataGatherer::FAssetDataGatherer(UAssetRegistryImpl& InRegistryImpl)
 	, IsStopped(0)
 	, IsGatheringPaused(0)
 	, IsProcessingPaused(0)
+	, DelayedSaveAsyncCacheStartTime(0.0)
 	, bSaveAsyncCacheTriggered(false)
 	, CurrentSearchTime(0.)
 	, LastCacheWriteTime(0.0)
@@ -3915,8 +3916,14 @@ void FAssetDataGatherer::OnAdditionalMountSearchCompleted()
 	{
 		Discovery->OnAdditionalMountSearchCompleted();
 	}
-
-	RequestAsyncCacheSave();
+	// If we haven't done so already, request a delayed save of the cache.
+	// Since additional mounts may be added in quick succession we want to save the cache but only
+	// after a brief period of time (i.e. FGatherSettings::MinSecondsToElapseBeforeCacheWrite) 
+	// to avoid saving the entire cache multiples times in a row.
+	if (DelayedSaveAsyncCacheStartTime.load(std::memory_order_relaxed) == 0.0)
+	{
+		DelayedSaveAsyncCacheStartTime.store(FPlatformTime::Seconds(), std::memory_order_relaxed);
+	}
 	bIsAdditionalMountSearchInProgress.store(false, std::memory_order_relaxed);
 	FileReadScheduler->FreeReadTaskLoaders();
 }
@@ -5106,13 +5113,22 @@ void FAssetDataGatherer::ConsumeCacheFiles(TArray<UE::AssetDataGather::Private::
 
 void FAssetDataGatherer::TryReserveSaveCache(bool& bOutShouldSave, TArray<TPair<FName,FDiskCachedAssetData*>>& AssetsToSave)
 {
+	using namespace UE::AssetDataGather::Private;
+
 	bOutShouldSave = false;
 	if (IsStopped)
 	{
 		return;
 	}
-	if (!bSaveAsyncCacheTriggered || bIsSavingAsyncCache)
+	double LocalDelayedSaveAsyncCacheStartTime = DelayedSaveAsyncCacheStartTime.load(std::memory_order_relaxed);
+	if ((!bSaveAsyncCacheTriggered && LocalDelayedSaveAsyncCacheStartTime == 0.0) || bIsSavingAsyncCache)
+	{
+		// No save requests have been issued or we are already saving
+		return;
+	}
+	if (FPlatformTime::Seconds() - LocalDelayedSaveAsyncCacheStartTime < FGatherSettings::MinSecondsToElapseBeforeCacheWrite)
 	{
+		// A delayed save request was made but we haven't waited long enough to start the save yet
 		return;
 	}
 	CHECK_IS_LOCKED_CURRENT_THREAD(TickLock);
@@ -5129,6 +5145,7 @@ void FAssetDataGatherer::TryReserveSaveCache(bool& bOutShouldSave, TArray<TPair<
 		LastCacheSaveNumUncachedAssetFiles = LocalNumUncachedAssetFiles;
 	}
 	bSaveAsyncCacheTriggered = false;
+	DelayedSaveAsyncCacheStartTime.store(0.0, std::memory_order_relaxed);
 }
 
 void FAssetDataGatherer::GetAssetsToSave(TArrayView<const FString> SaveCacheLongPackageNameDirs, TArray<TPair<FName,FDiskCachedAssetData*>>& OutAssetsToSave)
diff --git a/Engine/Source/Runtime/AssetRegistry/Private/AssetDataGatherer.h b/Engine/Source/Runtime/AssetRegistry/Private/AssetDataGatherer.h
--- a/Engine/Source/Runtime/AssetRegistry/Private/AssetDataGatherer.h
+++ b/Engine/Source/Runtime/AssetRegistry/Private/AssetDataGatherer.h
@@ -498,6 +498,8 @@ private:
 	 * Pointer is constant during threading. Object pointed to internally provides threadsafety.
 	 */
 	TUniquePtr<UE::AssetDataGather::Private::FAssetDataDiscovery> Discovery;
+	/** Start time a delayed request to save the async cache was made. */
+	std::atomic<double> DelayedSaveAsyncCacheStartTime;
 	/** True when TickInternal requests periodic or final save of the async cache. */
 	std::atomic<bool> bSaveAsyncCacheTriggered;
 	/** True if the current process allows reading of AssetDataGatherer cache files. */
