Add insights scope for EnumerateAllAssets

diff --git a/Engine/Source/Runtime/AssetRegistry/Private/AssetRegistry.cpp b/Engine/Source/Runtime/AssetRegistry/Private/AssetRegistry.cpp
--- a/Engine/Source/Runtime/AssetRegistry/Private/AssetRegistry.cpp
+++ b/Engine/Source/Runtime/AssetRegistry/Private/AssetRegistry.cpp
@@ -45,6 +45,7 @@
 #include "Misc/TransactionallySafeRWLock.h"
 #include "AssetRegistry/PackageReader.h"
 #include "Modules/ModuleManager.h"
+#include "ProfilingDebugging/CpuProfilerTrace.h"
 #include "ProfilingDebugging/MiscTrace.h"
 #include "Serialization/ArrayReader.h"
 #include "Serialization/CompactBinarySerialization.h"
@@ -3300,6 +3301,7 @@ bool UAssetRegistryImpl::GetAllAssets(TArray<FAssetData>& OutAssetData, bool bIn
 bool UAssetRegistryImpl::EnumerateAllAssets(TFunctionRef<bool(const FAssetData&)> Callback, UE::AssetRegistry::EEnumerateAssetsFlags InEnumerateFlags) const
 {
 	using namespace UE::AssetRegistry;
+	TRACE_CPUPROFILER_EVENT_SCOPE_STR("UAssetRegistryImpl::EnumerateAllAssets");
 
 	TSet<FName> PackageNamesToSkip;
 	// All in memory assets
