[AssetRegistry]

Ensure the asset registry cache is saved before shutdown if a request for save has been made. Previously no attempt was made to ensure the cache was written out completely before exiting.

- Removes bSaveAsyncCacheTriggered as we can indicate if a save was trigged via the DelayedSaveAsyncCacheStartTime time being non-zero and force the start time to be a negative time to trigger saves immediately

#rnx

diff --git a/Engine/Source/Runtime/AssetRegistry/Private/AssetDataGatherer.cpp b/Engine/Source/Runtime/AssetRegistry/Private/AssetDataGatherer.cpp
--- a/Engine/Source/Runtime/AssetRegistry/Private/AssetDataGatherer.cpp
+++ b/Engine/Source/Runtime/AssetRegistry/Private/AssetDataGatherer.cpp
@@ -3793,7 +3793,6 @@ FAssetDataGatherer::FAssetDataGatherer(UAssetRegistryImpl& InRegistryImpl)
 	, IsGatheringPaused(0)
 	, IsProcessingPaused(0)
 	, DelayedSaveAsyncCacheStartTime(0.0)
-	, bSaveAsyncCacheTriggered(false)
 	, CurrentSearchTime(0.)
 	, LastCacheWriteTime(0.0)
 	, bHasLoadedCache(false)
@@ -3920,10 +3919,7 @@ void FAssetDataGatherer::OnAdditionalMountSearchCompleted()
 	// Since additional mounts may be added in quick succession we want to save the cache but only
 	// after a brief period of time (i.e. FGatherSettings::MinSecondsToElapseBeforeCacheWrite) 
 	// to avoid saving the entire cache multiples times in a row.
-	if (DelayedSaveAsyncCacheStartTime.load(std::memory_order_relaxed) == 0.0)
-	{
-		DelayedSaveAsyncCacheStartTime.store(FPlatformTime::Seconds(), std::memory_order_relaxed);
-	}
+	RequestDelayedAsyncCacheSave();
 	bIsAdditionalMountSearchInProgress.store(false, std::memory_order_relaxed);
 	FileReadScheduler->FreeReadTaskLoaders();
 }
@@ -4088,7 +4084,7 @@ bool FAssetDataGatherer::ShouldBackgroundGatherThreadPauseGathering(bool& bOutGa
 {
 	FGathererScopeLock ResultsScopeLock(&ResultsLock); // bIsIdle requires the lock
 	bOutGathererIsIdle = bIsIdle;
-	return !IsStopped && !bSaveAsyncCacheTriggered && (IsGatheringPaused || bIsIdle);
+	return !IsStopped && !IsSaveRequested() && (IsGatheringPaused || bIsIdle);
 }
 
 UE::AssetRegistry::Impl::EGatherStatus FAssetDataGatherer::TryTickOnBackgroundThread()
@@ -5116,19 +5112,22 @@ void FAssetDataGatherer::TryReserveSaveCache(bool& bOutShouldSave, TArray<TPair<
 	using namespace UE::AssetDataGather::Private;
 
 	bOutShouldSave = false;
-	if (IsStopped)
+	double SaveAsyncCacheStartTime = 0.0;
+	bool bIsSaveRequested = IsSaveRequested(&SaveAsyncCacheStartTime);
+	int32 LocalIsStopped = IsStopped.load(std::memory_order_relaxed);
+	if (LocalIsStopped && !bIsSaveRequested)
 	{
 		return;
 	}
-	double LocalDelayedSaveAsyncCacheStartTime = DelayedSaveAsyncCacheStartTime.load(std::memory_order_relaxed);
-	if ((!bSaveAsyncCacheTriggered && LocalDelayedSaveAsyncCacheStartTime == 0.0) || bIsSavingAsyncCache)
+	if (bIsSavingAsyncCache)
 	{
-		// No save requests have been issued or we are already saving
+		// we are already saving
 		return;
 	}
-	if (FPlatformTime::Seconds() - LocalDelayedSaveAsyncCacheStartTime < FGatherSettings::MinSecondsToElapseBeforeCacheWrite)
+	if (!LocalIsStopped && (FPlatformTime::Seconds() - SaveAsyncCacheStartTime) < FGatherSettings::MinSecondsToElapseBeforeCacheWrite)
 	{
-		// A delayed save request was made but we haven't waited long enough to start the save yet
+		// A delayed save request was made but we haven't waited long enough to start the save yet.
+		// Note, we ignore the delayed save timer and handle the save request immediately if the gatherer is requested to be stopped.
 		return;
 	}
 	CHECK_IS_LOCKED_CURRENT_THREAD(TickLock);
@@ -5144,8 +5143,8 @@ void FAssetDataGatherer::TryReserveSaveCache(bool& bOutShouldSave, TArray<TPair<
 		bIsSavingAsyncCache = true;
 		LastCacheSaveNumUncachedAssetFiles = LocalNumUncachedAssetFiles;
 	}
-	bSaveAsyncCacheTriggered = false;
-	DelayedSaveAsyncCacheStartTime.store(0.0, std::memory_order_relaxed);
+
+	ClearRequestAsyncCacheSave();
 }
 
 void FAssetDataGatherer::GetAssetsToSave(TArrayView<const FString> SaveCacheLongPackageNameDirs, TArray<TPair<FName,FDiskCachedAssetData*>>& OutAssetsToSave)
diff --git a/Engine/Source/Runtime/AssetRegistry/Private/AssetDataGatherer.h b/Engine/Source/Runtime/AssetRegistry/Private/AssetDataGatherer.h
--- a/Engine/Source/Runtime/AssetRegistry/Private/AssetDataGatherer.h
+++ b/Engine/Source/Runtime/AssetRegistry/Private/AssetDataGatherer.h
@@ -184,7 +184,18 @@ public:
 	void UpdateCacheForSaving();
 	void OnInitialSearchCompleted();
 	void OnAdditionalMountSearchCompleted();
+	/** Requests a save of the cache as soon as possible. */
 	void RequestAsyncCacheSave();
+	/** Requests a delayed save of the cache. If a delayed save has already been requested, the delayed save timer is not reset. */
+	void RequestDelayedAsyncCacheSave();
+	/** Removes any existing request to save the cache. */
+	void ClearRequestAsyncCacheSave();
+	/** 
+	* Returns true if an immediate or delayed save has been requested. 
+	* Sets OutSaveRequestTime to the time the save was requested (value may 
+	* be negative to indicate a save should occur immediately), zero otherwise.
+	*/
+	bool IsSaveRequested(double* OutSaveRequestTime = nullptr) const;
 
 	// Controlling Async behavior
 
@@ -498,10 +509,12 @@ private:
 	 * Pointer is constant during threading. Object pointed to internally provides threadsafety.
 	 */
 	TUniquePtr<UE::AssetDataGather::Private::FAssetDataDiscovery> Discovery;
-	/** Start time a delayed request to save the async cache was made. */
+	/** 
+	* Start time a delayed request to save the async cache was made. 
+	* A save will be triggered after FGatherSettings::MinSecondsToElapseBeforeCacheWrite have elapsed. 
+	* When an immediate save request is made this time will be negative.
+	*/
 	std::atomic<double> DelayedSaveAsyncCacheStartTime;
-	/** True when TickInternal requests periodic or final save of the async cache. */
-	std::atomic<bool> bSaveAsyncCacheTriggered;
 	/** True if the current process allows reading of AssetDataGatherer cache files. */
 	std::atomic<bool> bCacheReadEnabled;
 	/** True if the current process allows writing of AssetDataGatherer cache files. */
@@ -612,9 +625,32 @@ private:
 // Inline implementations
 ///////////////////////////////////////////////////////
 
+inline bool FAssetDataGatherer::IsSaveRequested(double* OutSaveRequestTime) const
+{
+	double LocalSaveAsyncCacheStartTime = DelayedSaveAsyncCacheStartTime.load(std::memory_order_relaxed);
+	if (OutSaveRequestTime)
+	{
+		*OutSaveRequestTime = LocalSaveAsyncCacheStartTime;
+	}
+	return LocalSaveAsyncCacheStartTime != 0.0;
+}
+
 inline void FAssetDataGatherer::RequestAsyncCacheSave()
 {
-	bSaveAsyncCacheTriggered = true;
+	DelayedSaveAsyncCacheStartTime.store(-UE_DOUBLE_BIG_NUMBER, std::memory_order_relaxed);
+}
+
+inline void FAssetDataGatherer::RequestDelayedAsyncCacheSave()
+{
+	if (DelayedSaveAsyncCacheStartTime.load(std::memory_order_relaxed) == 0.0)
+	{
+		DelayedSaveAsyncCacheStartTime.store(FPlatformTime::Seconds(), std::memory_order_relaxed);
+	}
+}
+
+inline void FAssetDataGatherer::ClearRequestAsyncCacheSave()
+{
+	DelayedSaveAsyncCacheStartTime.store(0.0, std::memory_order_relaxed);
 }
 
 inline FGatheredResultsEvent& FAssetDataGatherer::GetGatheredResultsEvent()
