[AssetRegistry]

Writing to the AssetRegistry cache will now attempt to take a system-wide lock to prevent parallel processes writing the same AssetRegistry CacheDirectory simultaneously when using an editor process. If a writing is underway by another process, cache saving will be requeued for the next AR tick.

#rnx

diff --git a/Engine/Source/Runtime/AssetRegistry/Private/AssetDataGatherer.cpp b/Engine/Source/Runtime/AssetRegistry/Private/AssetDataGatherer.cpp
--- a/Engine/Source/Runtime/AssetRegistry/Private/AssetDataGatherer.cpp
+++ b/Engine/Source/Runtime/AssetRegistry/Private/AssetDataGatherer.cpp
@@ -18,6 +18,7 @@
 #include "Containers/BinaryHeap.h"
 #include "HAL/LowLevelMemTracker.h"
 #include "HAL/PlatformFileManager.h"
+#include "HAL/PlatformMutex.h"
 #include "HAL/PlatformProcess.h"
 #include "HAL/RunnableThread.h"
 #include "Hash/xxhash.h"
@@ -337,6 +338,18 @@ TArray<FString> FGatherSettings::FindShardedCacheFiles() const
 
 FGatherSettings GGatherSettings;
 
+#if WITH_EDITOR
+// Attempts to lock a system-wide mutex guarding the writing of AssetRegistry cache files. If already locked, no waiting occurs.
+// Users should verify the lock state after calling this function.
+[[nodiscard]] UE::FPlatformSystemWideMutex TryLockSystemWideCacheFilesWrites()
+{
+	check(GGatherSettings.IsInitialized());
+	TStringBuilder<256> Builder;
+	Builder.Appendf(TEXT("AssetRegistryCacheLock-%s"), *GGatherSettings.GetAssetRegistryCacheRootFolder());
+	return UE::FPlatformSystemWideMutex(Builder.ToString());
+}
+#endif
+
 /** A structure to hold serialized cache data from async loads before adding it to the Gatherer's main cache. */
 struct FCachePayload
 {
@@ -4150,6 +4163,17 @@ void FAssetDataGatherer::SaveCacheFile(const TArray<TPair<FName,FDiskCachedAsset
 		FPlatformProcess::Sleep(static_cast<float>(LastCacheWriteTime + WaitTimeBeforeReopen - CurrentTime));
 	}
 
+	// Attempt to get exclusive write access to the cache files. If we cannot, skip writing
+	// and try again next tick. Failing to get the lock here means another process 
+	// is currently writing so we avoid the possibility of never writing cache files 
+#if WITH_EDITOR
+	UE::FPlatformSystemWideMutex CacheFilesLock = TryLockSystemWideCacheFilesWrites();
+	if (!CacheFilesLock.IsValid())
+	{
+		return;
+	}
+#endif
+
 	ParallelFor(CacheShards, [this, &AssetsToSave, &TotalCacheSize, &DataPerShard, ShardMask](int32 Shard) 
 	{
 		TArray<TPair<FName, FDiskCachedAssetData*>>& ShardData = DataPerShard[Shard];
@@ -5122,10 +5146,10 @@ void FAssetDataGatherer::TryReserveSaveCache(bool& bOutShouldSave, TArray<TPair<
 		// we are already saving
 		return;
 	}
-	if ((FPlatformTime::Seconds() - SaveAsyncCacheStartTime) < FGatherSettings::MinSecondsToElapseBeforeCacheWrite)
+	if (!IsStopped.load(std::memory_order_relaxed) && (FPlatformTime::Seconds() - SaveAsyncCacheStartTime) < FGatherSettings::MinSecondsToElapseBeforeCacheWrite)
 	{
-		// A delayed save request was made but we haven't waited long enough to start the save yet.
-		// Note, we ignore the delayed save timer and handle the save request immediately if the gatherer is requested to be stopped.
+		// A delayed save request was made but we haven't waited long enough to start the save yet. 
+		// However if we are stopped, ignore the timer and handle the save request immediately
 		return;
 	}
 	CHECK_IS_LOCKED_CURRENT_THREAD(TickLock);
diff --git a/Engine/Source/Runtime/AssetRegistry/Private/AssetDataGathererPrivate.h b/Engine/Source/Runtime/AssetRegistry/Private/AssetDataGathererPrivate.h
--- a/Engine/Source/Runtime/AssetRegistry/Private/AssetDataGathererPrivate.h
+++ b/Engine/Source/Runtime/AssetRegistry/Private/AssetDataGathererPrivate.h
@@ -988,6 +988,7 @@ public:
 	static inline int32 GARGatherCacheParallelism = 16;
 
 	void Initialize();
+	bool IsInitialized() const;
 	bool IsAsyncEnabled() const;
 	bool IsGatherCacheReadEnabled() const;
 	bool IsGatherCacheWriteEnabled() const;
@@ -1029,6 +1030,11 @@ bool DoesPathContainInvalidCharacters(const EGatherableFileType FileType, FStrin
 // Inline implementations
 ///////////////////////////////////////////////////////
 
+inline bool FGatherSettings::IsInitialized() const
+{
+	return bInitialized;
+}
+
 inline bool FGatherSettings::IsAsyncEnabled() const
 {
 	return bAsyncEnabled;
