Fix telemetry for Editor.AssetRegistry.InitialScan to only be sent for initial scan, and otherwise change event name fo Editor.AssetRegistry.AdditionalScan

diff --git a/Engine/Plugins/Experimental/EditorTelemetry/Source/Private/EditorTelemetry.cpp b/Engine/Plugins/Experimental/EditorTelemetry/Source/Private/EditorTelemetry.cpp
--- a/Engine/Plugins/Experimental/EditorTelemetry/Source/Private/EditorTelemetry.cpp
+++ b/Engine/Plugins/Experimental/EditorTelemetry/Source/Private/EditorTelemetry.cpp
@@ -1090,7 +1090,8 @@ void FEditorTelemetry::StartSession()
 			});
 		});
 		Router.OnTelemetry<FGatherTelemetry>([this](const FGatherTelemetry& Data){
-			FStudioTelemetry::Get().RecordEvent(TEXT("Editor.AssetRegistry.InitialScan"), 
+			const TCHAR* EventName = Data.bIsInitialScan ? TEXT("Editor.AssetRegistry.InitialScan")  : TEXT("Editor.AssetRegistry.AdditionalScan");
+			FStudioTelemetry::Get().RecordEvent(EventName, 
 			{
 				{ TEXT("SchemaVersion"), 1 },
 				{ TEXT("TotalDurationSeconds"), Data.TotalSearchDurationSeconds },
diff --git a/Engine/Source/Runtime/AssetRegistry/Private/AssetRegistry.cpp b/Engine/Source/Runtime/AssetRegistry/Private/AssetRegistry.cpp
--- a/Engine/Source/Runtime/AssetRegistry/Private/AssetRegistry.cpp
+++ b/Engine/Source/Runtime/AssetRegistry/Private/AssetRegistry.cpp
@@ -6024,7 +6024,7 @@ void FAssetRegistryImpl::OnInitialSearchCompleted(Impl::FEventContext& EventCont
 	// Handle any deferred loading operations
 	SetPerformanceMode(Impl::EPerformanceMode::MostlyStatic);
 
-	LogSearchDiagnostics(InitialSearchStartTime);
+	LogSearchDiagnostics(InitialSearchStartTime, /* bInitialScan */ true);
 	TRACE_END_REGION(TEXT("Asset Registry Scan"));
 
 	GlobalGatherer->UpdateCacheForSaving();
@@ -6047,7 +6047,7 @@ void FAssetRegistryImpl::OnAdditionalMountSearchCompleted(Impl::FEventContext& E
 	// Handle any deferred loading operations
 	SetPerformanceMode(Impl::EPerformanceMode::MostlyStatic);
 
-	LogSearchDiagnostics(AdditionalMountSearchStartTime);
+	LogSearchDiagnostics(AdditionalMountSearchStartTime, /* bInitialScan */ false);
 	TRACE_END_REGION(TEXT("Asset Registry - Additional Mount Search")); // Matching TRACE_BEGIN_REGION in OnContentPathMounted
 
 	GlobalGatherer->UpdateCacheForSaving();
@@ -6059,7 +6059,7 @@ void FAssetRegistryImpl::OnAdditionalMountSearchCompleted(Impl::FEventContext& E
 	bAdditionalMountSearchInProgress.store(false, std::memory_order_relaxed);
 }
 
-void FAssetRegistryImpl::LogSearchDiagnostics(double StartTime)
+void FAssetRegistryImpl::LogSearchDiagnostics(double StartTime, bool bInitialScan)
 {
 	FAssetGatherDiagnostics Diagnostics = GlobalGatherer->GetDiagnostics();
 	float Total = Diagnostics.DiscoveryTimeSeconds + Diagnostics.GatherTimeSeconds + StoreGatherResultsTimeSeconds;
@@ -6073,6 +6073,7 @@ void FAssetRegistryImpl::LogSearchDiagnostics(double StartTime)
 	Telemetry.NumUncachedDirectories = Diagnostics.NumUncachedDirectories;
 	Telemetry.NumCachedAssetFiles = Diagnostics.NumCachedAssetFiles;
 	Telemetry.NumUncachedAssetFiles = Diagnostics.NumUncachedAssetFiles;
+	Telemetry.bIsInitialScan = bInitialScan;
 	FTelemetryRouter::Get().ProvideTelemetry(Telemetry);
 #if !NO_LOGGING
 	TStringBuilder<256> Message;
diff --git a/Engine/Source/Runtime/AssetRegistry/Private/AssetRegistryImpl.h b/Engine/Source/Runtime/AssetRegistry/Private/AssetRegistryImpl.h
--- a/Engine/Source/Runtime/AssetRegistry/Private/AssetRegistryImpl.h
+++ b/Engine/Source/Runtime/AssetRegistry/Private/AssetRegistryImpl.h
@@ -505,7 +505,7 @@ public:
 	/** Send a log message with the search statistics. 
 	 *  StartTime is used to report wall clock search time in the case of background scan
 	 */
-	void LogSearchDiagnostics(double StartTime);
+	void LogSearchDiagnostics(double StartTime, bool bInitialScan);
 	/** Look for and load a single AssetData result from the gatherer. */
 	void TickGatherPackage(UE::AssetRegistry::FInterfaceWriteScopeLock& ScopeLock, Impl::FEventContext& EventContext, const FString& PackageName, const FString& LocalPath);
 	void ClearGathererCache();
diff --git a/Engine/Source/Runtime/AssetRegistry/Public/AssetRegistry/AssetRegistryTelemetry.h b/Engine/Source/Runtime/AssetRegistry/Public/AssetRegistry/AssetRegistryTelemetry.h
--- a/Engine/Source/Runtime/AssetRegistry/Public/AssetRegistry/AssetRegistryTelemetry.h
+++ b/Engine/Source/Runtime/AssetRegistry/Public/AssetRegistry/AssetRegistryTelemetry.h
@@ -44,7 +44,9 @@ namespace UE::Telemetry::AssetRegistry
 		bool bAdditionalMountSearchInProgress = false;
     };
 
-    // Information about the initial asset registry scan triggered when launching the process. 
+    // Information about asset registry scanning either as the initial scan when the process starts,
+    // or when an additional mount point is added. See the bIsInitialScan field to disambiguate.
+    // Note that for subsequent scans, the times other than TotalSearchDurationSeconds will be cumulative totals including previous work.
     // This is also a good time to gather information about the contents of the asset registry.
     struct FGatherTelemetry
     {
@@ -68,6 +70,9 @@ namespace UE::Telemetry::AssetRegistry
         int32 NumCachedAssetFiles = 0;
         // Number of asset files read loose from disk
         int32 NumUncachedAssetFiles = 0;
+
+        // True if this was the initial startup scan - this should be true only once per process.        
+        bool bIsInitialScan = false;
     };
     
     // Information about an asset registry update that was triggered by the directory watcher module. 
