[AR] Batch paths dismounting when multiple plugins are requested to release. Speeds up the process.

#tests local run, replay runs
#rnx

diff --git a/Engine/Plugins/Runtime/GameFeatures/Source/GameFeatures/Private/GameFeaturesSubsystem.cpp b/Engine/Plugins/Runtime/GameFeatures/Source/GameFeatures/Private/GameFeaturesSubsystem.cpp
--- a/Engine/Plugins/Runtime/GameFeatures/Source/GameFeatures/Private/GameFeaturesSubsystem.cpp
+++ b/Engine/Plugins/Runtime/GameFeatures/Source/GameFeatures/Private/GameFeaturesSubsystem.cpp
@@ -17,6 +17,7 @@
 #include "GameFeaturePluginStateMachine.h"
 #include "GameFeatureStateChangeObserver.h"
 #include "GameplayTagsManager.h"
+#include "HAL/Platform.h"
 #include "Interfaces/IPluginManager.h"
 #include "IO/IoStoreOnDemand.h"
 #include "Logging/StructuredLog.h"
@@ -2084,42 +2085,62 @@ void UGameFeaturesSubsystem::ReleaseGameFeaturePlugin(const FString& PluginURL,
 
 void UGameFeaturesSubsystem::ReleaseGameFeaturePlugin(TConstArrayView<FString> PluginURLs, const FMultipleGameFeaturePluginsReleased& CompleteDelegate)
 {
+	struct FSuspendPluginManger
+	{
+		UE_NODISCARD_CTOR FSuspendPluginManger() { IPluginManager::Get().SuppressPluginUnloadGC(); }
+		~FSuspendPluginManger() { IPluginManager::Get().ResumePluginUnloadGC(); }
+	};
+	struct FSuspendGameplayTagTreeChangedBroadcast
+	{
+		UE_NODISCARD_CTOR FSuspendGameplayTagTreeChangedBroadcast(bool bAllocateIfRequested)
+		{
+			if (UGameplayTagsManager* TagsManager = bAllocateIfRequested ? &UGameplayTagsManager::Get() : UGameplayTagsManager::GetIfAllocated() )
+			{
+				TagsManager->PushDeferOnGameplayTagTreeChangedBroadcast();
+				bWasSuspended = true;
+			}
+		}
+		~FSuspendGameplayTagTreeChangedBroadcast()
+		{
+			if (bWasSuspended)
+			{
+				UGameplayTagsManager::Get().PopDeferOnGameplayTagTreeChangedBroadcast();
+			}
+		}
+
+		bool bWasSuspended = false;
+	};
 	struct FReleaseContext
 	{
 		TMap<FString, UE::GameFeatures::FResult> Results;
 		FMultipleGameFeaturePluginsReleased CompleteDelegate;
 		int32 NumPluginsReleased = 0;
-		bool bDeferGCAfterUnload = false;
-		bool bPushedTagsBroadcast = false;
 
-		FReleaseContext()
+		TOptional<FSuspendPluginManger> SuspendPluginManager;
+		TOptional<FSuspendGameplayTagTreeChangedBroadcast> SuspendGameplayTagTreeBroadcast;
+		TOptional<IAssetRegistry::FSuspendPathsDismounting> SuspendAssetRegistry;
+
+		UE_NODISCARD_CTOR FReleaseContext()
 		{
 			if (!IsEngineExitRequested())
 			{
-				UGameplayTagsManager::Get().PushDeferOnGameplayTagTreeChangedBroadcast();
-				IPluginManager::Get().SuppressPluginUnloadGC();
-				bDeferGCAfterUnload = true;
-				bPushedTagsBroadcast = true;
-			}
-			else if(UGameplayTagsManager* TagsManager = UGameplayTagsManager::GetIfAllocated())
-			{
-				TagsManager->PushDeferOnGameplayTagTreeChangedBroadcast();
-				bPushedTagsBroadcast = true;
+				SuspendPluginManager.Emplace();
 			}
+			SuspendGameplayTagTreeBroadcast.Emplace(!IsEngineExitRequested());
+			SuspendAssetRegistry.Emplace();
 		}
 
 		~FReleaseContext()
 		{
-			if (bDeferGCAfterUnload)
-			{
-				IPluginManager::Get().ResumePluginUnloadGC();
-			}
-			if (bPushedTagsBroadcast)
-			{
-				UGameplayTagsManager::Get().PopDeferOnGameplayTagTreeChangedBroadcast();
-			}
-
 			CompleteDelegate.ExecuteIfBound(Results);
+			// Explicit order of destruction is required
+			// At the very least we need to execute PluginManger batch last
+			// as it will run GC and does the leaked packages verification.
+			// And by that point all other resources need to be properly cleared,
+			// i.e. other batches need to finish.
+			SuspendAssetRegistry.Reset();
+			SuspendGameplayTagTreeBroadcast.Reset();
+			SuspendPluginManager.Reset();
 		}
 	};
 	TSharedRef<FReleaseContext> LoadContext = MakeShared<FReleaseContext>();
diff --git a/Engine/Source/Runtime/AssetRegistry/Private/AssetRegistry.cpp b/Engine/Source/Runtime/AssetRegistry/Private/AssetRegistry.cpp
--- a/Engine/Source/Runtime/AssetRegistry/Private/AssetRegistry.cpp
+++ b/Engine/Source/Runtime/AssetRegistry/Private/AssetRegistry.cpp
@@ -52,6 +52,7 @@
 #include "Serialization/CompactBinaryWriter.h"
 #include "String/RemoveFrom.h"
 #include "Tasks/Task.h"
+#include "Templates/Tuple.h"
 #include "Templates/UnrealTemplate.h"
 #include "TelemetryRouter.h"
 #include "UObject/AssetRegistryTagsContext.h"
@@ -64,6 +65,7 @@
 #include "UObject/UObjectHash.h"
 #include "UObject/UObjectIterator.h"
 #include "UObject/UObjectThreadContext.h"
+#include <atomic>
 
 #include UE_INLINE_GENERATED_CPP_BY_NAME(IAssetRegistry)
 #include UE_INLINE_GENERATED_CPP_BY_NAME(AssetRegistry)
@@ -9178,6 +9180,7 @@ void FAssetRegistryImpl::ScanModifiedAssetFiles(FInterfaceWriteScopeLock& ScopeL
 
 void UAssetRegistryImpl::OnContentPathMounted(const FString& InAssetPath, const FString& FileSystemPath)
 {
+	TRACE_CPUPROFILER_EVENT_SCOPE(AssetRegistry_OnContentPathMounted);
 	// Sanitize
 	FString AssetPathWithTrailingSlash;
 	if (!InAssetPath.EndsWith(TEXT("/"), ESearchCase::CaseSensitive))
@@ -9280,14 +9283,19 @@ void FAssetRegistryImpl::OnContentPathMounted(UE::AssetRegistry::FInterfaceWrite
 
 void UAssetRegistryImpl::OnContentPathDismounted(const FString& InAssetPath, const FString& FileSystemPath)
 {
-	// Sanitize
-	FString AssetPathNoTrailingSlash = InAssetPath;
-	if (AssetPathNoTrailingSlash.EndsWith(TEXT("/"), ESearchCase::CaseSensitive))
+	if (SuspendDismountingCounter.load(std::memory_order_acquire) > 0)
 	{
-		// We don't want a trailing slash here as it could interfere with RemoveAssetPath
-		AssetPathNoTrailingSlash.LeftChopInline(1, EAllowShrinking::No);
+		FScopeLock Guard(&DismountingCriticalSection);
+		ListOfUnmountedPaths.Emplace(InAssetPath, FileSystemPath);
+		return;
 	}
 
+	OnMultipleContentPathsDismounted({MakeTuple(InAssetPath, FileSystemPath)});
+}
+
+void UAssetRegistryImpl::OnMultipleContentPathsDismounted(TConstArrayView<TTuple<FString, FString>> UnmountedAssetFilesystemPairs)
+{
+	TRACE_CPUPROFILER_EVENT_SCOPE(AssetRegistry_OnMultipleContentPathsDismounted);
 #if WITH_EDITOR
 	FDirectoryWatcherModule& DirectoryWatcherModule = FModuleManager::LoadModuleChecked<FDirectoryWatcherModule>(TEXT("DirectoryWatcher"));
 	IDirectoryWatcher* DirectoryWatcher = nullptr;
@@ -9301,28 +9309,27 @@ void UAssetRegistryImpl::OnContentPathDismounted(const FString& InAssetPath, con
 	{
 		LLM_SCOPE(ELLMTag::AssetRegistry);
 		UE::AssetRegistry::FInterfaceWriteScopeLock InterfaceScopeLock(InterfaceLock);
-		GuardedData.OnContentPathDismounted(EventContext, InAssetPath, AssetPathNoTrailingSlash, FileSystemPath);
+		GuardedData.OnMultipleContentPathsDismounted(EventContext, UnmountedAssetFilesystemPairs);
 
 		// Stop listening for directory changes in this content path
 #if WITH_EDITOR
-		const FString StandardFileSystemPath = UE::AssetRegistry::CreateStandardFilename(FileSystemPath);
-		if (DirectoryWatcher && !IsDirAlreadyWatchedByRootWatchers(StandardFileSystemPath))
+		for (const TTuple<FString, FString>& Pair : UnmountedAssetFilesystemPairs)
 		{
-			// Make sure OnDirectoryChangedDelegateHandles key is symmetrical with the one used in OnContentPathMounted
-			FString AssetPathWithTrailingSlash;
-			if (!InAssetPath.EndsWith(TEXT("/"), ESearchCase::CaseSensitive))
-			{
-				AssetPathWithTrailingSlash = InAssetPath + TEXT("/");
-			}
-			else
+			const FString StandardFileSystemPath = UE::AssetRegistry::CreateStandardFilename(Pair.Value);
+			if (DirectoryWatcher && !IsDirAlreadyWatchedByRootWatchers(StandardFileSystemPath))
 			{
-				AssetPathWithTrailingSlash = InAssetPath;
-			}
+				// Make sure OnDirectoryChangedDelegateHandles key is symmetrical with the one used in OnContentPathMounted
+				FString AssetPathWithTrailingSlash = Pair.Key;
+				if (!AssetPathWithTrailingSlash.EndsWith(TEXT("/"), ESearchCase::CaseSensitive))
+				{
+					AssetPathWithTrailingSlash += TEXT("/");
+				}
 
-			FDelegateHandle DirectoryChangedHandle;
-			if (ensure(OnDirectoryChangedDelegateHandles.RemoveAndCopyValue(AssetPathWithTrailingSlash, DirectoryChangedHandle)))
-			{
-				DirectoryWatcher->UnregisterDirectoryChangedCallback_Handle(StandardFileSystemPath, DirectoryChangedHandle);
+				FDelegateHandle DirectoryChangedHandle;
+				if (ensure(OnDirectoryChangedDelegateHandles.RemoveAndCopyValue(AssetPathWithTrailingSlash, DirectoryChangedHandle)))
+				{
+					DirectoryWatcher->UnregisterDirectoryChangedCallback_Handle(StandardFileSystemPath, DirectoryChangedHandle);
+				}
 			}
 		}
 #endif // WITH_EDITOR
@@ -9333,30 +9340,54 @@ void UAssetRegistryImpl::OnContentPathDismounted(const FString& InAssetPath, con
 namespace UE::AssetRegistry
 {
 
-void FAssetRegistryImpl::OnContentPathDismounted(Impl::FEventContext& EventContext, const FString& InAssetPath, 
-	const FString& AssetPathNoTrailingSlash, const FString& FileSystemPath)
+void FAssetRegistryImpl::OnMultipleContentPathsDismounted(
+	Impl::FEventContext& EventContext,
+	TConstArrayView<TTuple<FString, FString>> UnmountedAssetFilesystemPairs)
 {
-	if (GlobalGatherer.IsValid())
+	// For the GlobalGatherer we remove all input MountPoints; we don't need to keep information about their gather.
+	// But for the FAssetData stored inside our AssetRegistryState, for some mount points
+	// we are configured to keep those assets in memory even when dismounted.
+	TSet<FName> MountPointsToRemoveAssetsFrom;
+	for (const TTuple<FString, FString>& Pair : UnmountedAssetFilesystemPairs)
 	{
-		GlobalGatherer->RemoveMountPoint(FileSystemPath);
+		if (GlobalGatherer.IsValid())
+		{
+			GlobalGatherer->RemoveMountPoint(Pair.Value);
+		}
+		// The input AssetPaths of the unmounted MountPoints are unnormalized and might have a trailing slash,
+		// but we are comparing them by exact string inside our containers 
+		// and need to normalize them to not have a trailing slash.
+		FStringView AssetPathNoTrailingSlash = Pair.Key;
+		if (AssetPathNoTrailingSlash.EndsWith(TEXT("/"), ESearchCase::CaseSensitive))
+		{
+			// We don't want a trailing slash here as it could interfere with RemoveAssetPath
+			AssetPathNoTrailingSlash.LeftChopInline(1);
+		}
+		FName AssetPathNoTrailingSlashFName(AssetPathNoTrailingSlash);
+		if (PersistentMountPoints.Contains(AssetPathNoTrailingSlashFName))
+		{
+			// This path is not marked to never remove its AssetDatas.
+			continue;
+		}
+
+		MountPointsToRemoveAssetsFrom.Add(AssetPathNoTrailingSlashFName);
 	}
 
-	FName MountPoint = FName(FStringView(AssetPathNoTrailingSlash));
-	if (PersistentMountPoints.Contains(MountPoint))
+	// Collect the list of paths from given list of roots
+	TSet<FName> AssetPathsToRemoveAssetsFrom;
 	{
-		// This path is marked to never remove its AssetDatas. Skip the code below to remove it.
-		return;
+		const bool bRecurse = true;
+		for (FName AssetPathNoTrailingSlashFName : MountPointsToRemoveAssetsFrom)
+		{
+			CachedPathTree.GetSubPaths(AssetPathNoTrailingSlashFName, AssetPathsToRemoveAssetsFrom, bRecurse);
+			AssetPathsToRemoveAssetsFrom.Add(AssetPathNoTrailingSlashFName);
+		}
 	}
 
 	// Remove all cached assets and Verse files found at this location
 	{
-		FName AssetPathNoTrailingSlashFName(*AssetPathNoTrailingSlash);
 		TArray<FAssetData*> AllAssetDataToRemove;
-		TSet<FName> PathList;
-		const bool bRecurse = true;
-		CachedPathTree.GetSubPaths(AssetPathNoTrailingSlashFName, PathList, bRecurse);
-		PathList.Add(AssetPathNoTrailingSlashFName);
-		for (FName PathName : PathList)
+		for (FName PathName : AssetPathsToRemoveAssetsFrom)
 		{
 			// Gather assets
 			State.EnumerateMutableAssetsByPackagePath(PathName, [&AllAssetDataToRemove](FAssetData* AssetData)
@@ -9381,9 +9412,10 @@ void FAssetRegistryImpl::OnContentPathDismounted(Impl::FEventContext& EventConte
 	}
 
 	// Remove the root path
+	for (FName AssetPathNoTrailingSlashFName : MountPointsToRemoveAssetsFrom)
 	{
 		const bool bEvenIfAssetsStillExist = true;
-		RemoveAssetPath(EventContext, FName(*AssetPathNoTrailingSlash), bEvenIfAssetsStillExist);
+		RemoveAssetPath(EventContext, AssetPathNoTrailingSlashFName, bEvenIfAssetsStillExist);
 	}
 }
 
@@ -9641,6 +9673,22 @@ bool UAssetRegistryImpl::IsDirAlreadyWatchedByRootWatchers(const FString& Direct
 
 #endif
 
+void UAssetRegistryImpl::SuspendPathsDismounting()
+{
+	SuspendDismountingCounter.fetch_add(1, std::memory_order_acq_rel);
+}
+
+void UAssetRegistryImpl::ResumePathsDismounting()
+{
+	TArray<TTuple<FString, FString>> LocalList;
+	if (SuspendDismountingCounter.fetch_sub(1, std::memory_order_acq_rel) == 1)
+	{
+		FScopeLock Guard(&DismountingCriticalSection);
+		LocalList = MoveTemp(ListOfUnmountedPaths);
+	}
+	OnMultipleContentPathsDismounted(LocalList);
+}
+
 void UAssetRegistryImpl::RequestPauseBackgroundProcessing()
 {
 #if WITH_EDITOR
diff --git a/Engine/Source/Runtime/AssetRegistry/Private/AssetRegistry.h b/Engine/Source/Runtime/AssetRegistry/Private/AssetRegistry.h
--- a/Engine/Source/Runtime/AssetRegistry/Private/AssetRegistry.h
+++ b/Engine/Source/Runtime/AssetRegistry/Private/AssetRegistry.h
@@ -267,6 +267,13 @@ private:
 	 */
 	void OnContentPathDismounted(const FString& AssetPath, const FString& FileSystemPath);
 
+	/**
+	 * Batched processing version of the OnContentPathDismounted
+	 *
+	 * @param	UnmountedAssetFilesystemPairs	the list of pairs AssetPath - FilesystemPath
+	 */
+	void OnMultipleContentPathsDismounted(TConstArrayView<TTuple<FString, FString>> UnmountedAssetFilesystemPairs);
+
 	/** Called to refresh the native classes list, called at end of engine initialization. */
 	void OnPostEngineInit();
 
@@ -313,12 +320,16 @@ private:
 	bool IsDirAlreadyWatchedByRootWatchers(const FString& Directory) const;
 #endif
 
+	// Implementation methods for FPausePathDismountingScope
+	virtual void SuspendPathsDismounting() override;
+	virtual void ResumePathsDismounting() override;
+
 	/** Request to pause or resume background processing of scan results.
 	 *  This can be used to allow a priority thread to perform along sequence of operations
 	 *  without having to contend with the background thread for data access
 	 */
-	virtual void RequestPauseBackgroundProcessing();
-	virtual void RequestResumeBackgroundProcessing();
+	virtual void RequestPauseBackgroundProcessing() override;
+	virtual void RequestResumeBackgroundProcessing() override;
 	bool IsBackgroundProcessingPaused() const 
 	{ 
 #if WITH_EDITOR 
@@ -436,6 +447,12 @@ private:
 	FTSTicker::FDelegateHandle TickRequestHandle;
 	FCriticalSection DeferredEventsCriticalSection;
 
+	/** Handles suspension and resume of paths dismounting Used for batching many Dismount calls together. */
+	std::atomic<int32> SuspendDismountingCounter = 0;
+	FCriticalSection DismountingCriticalSection;
+	TArray<TTuple<FString, FString>> ListOfUnmountedPaths;
+
+
 	friend class FAssetDataGatherer;
 	friend class UE::AssetRegistry::FAssetRegistryImpl;
 	friend struct UE::AssetRegistry::Premade::FAsyncConsumer;
diff --git a/Engine/Source/Runtime/AssetRegistry/Private/AssetRegistryImpl.h b/Engine/Source/Runtime/AssetRegistry/Private/AssetRegistryImpl.h
--- a/Engine/Source/Runtime/AssetRegistry/Private/AssetRegistryImpl.h
+++ b/Engine/Source/Runtime/AssetRegistry/Private/AssetRegistryImpl.h
@@ -7,6 +7,7 @@
 #include "AssetRegistry/AssetRegistryState.h"
 #include "AssetRegistry/IAssetRegistry.h"
 #include "AssetRegistry/PathTree.h"
+#include "Containers/ArrayView.h"
 #include "Containers/RingBuffer.h"
 #include "Containers/Map.h"
 #include "Containers/Set.h"
@@ -453,9 +454,7 @@ public:
 #endif
 	void OnContentPathMounted(UE::AssetRegistry::FInterfaceWriteScopeLock& ScopeLock, Impl::FEventContext& EventContext, Impl::FClassInheritanceContext& InheritanceContext,
 		const FString& InAssetPath, const FString& AssetPathWithTrailingSlash, const FString& FileSystemPath);
-	void OnContentPathDismounted(Impl::FEventContext& EventContext, const FString& InAssetPath,
-		const FString& AssetPathNoTrailingSlash, const FString& FileSystemPath);
-
+	void OnMultipleContentPathsDismounted(Impl::FEventContext& EventContext, TConstArrayView<TTuple<FString, FString>> UnmountedAssetFilesystemPairs);
 
 	// Other helper functions called by UAssetRegistryImpl
 
diff --git a/Engine/Source/Runtime/AssetRegistry/Public/AssetRegistry/IAssetRegistry.h b/Engine/Source/Runtime/AssetRegistry/Public/AssetRegistry/IAssetRegistry.h
--- a/Engine/Source/Runtime/AssetRegistry/Public/AssetRegistry/IAssetRegistry.h
+++ b/Engine/Source/Runtime/AssetRegistry/Public/AssetRegistry/IAssetRegistry.h
@@ -1155,8 +1155,22 @@ public:
 	DECLARE_TS_MULTICAST_DELEGATE_TwoParams(FEnumerateAssetsEvent, const FARCompiledFilter& /*Filter*/, UE::AssetRegistry::EEnumerateAssetsFlags /*Flag*/)
 	virtual FEnumerateAssetsEvent& OnEnumerateAssetsEvent() = 0;
 
+	/** Request to pause paths dismounting.
+	 * Can be used by the calling code to temporarily pause clearing of data inside the AssetRegistry.
+	 * All unmounting requests will be executed in one batch at the end of that scope.
+	 */
+	struct FSuspendPathsDismounting
+	{
+		UE_NODISCARD_CTOR FSuspendPathsDismounting() { IAssetRegistry::GetChecked().SuspendPathsDismounting(); }
+		~FSuspendPathsDismounting() { IAssetRegistry::GetChecked().ResumePathsDismounting(); }
+	};
+
 protected:
 
+	// Implementation methods for FPausePathDismountingScope
+	virtual void SuspendPathsDismounting() {};
+	virtual void ResumePathsDismounting() {};
+
 	// Implementation methods for FPauseBackgroundProcessingScope
 	virtual void RequestPauseBackgroundProcessing() {}
 	virtual void RequestResumeBackgroundProcessing() {}
