[AR] Fix incorrect doc comments for FSuspendPathsDismounting

#rnx

diff --git a/Engine/Source/Runtime/AssetRegistry/Private/AssetRegistry.h b/Engine/Source/Runtime/AssetRegistry/Private/AssetRegistry.h
--- a/Engine/Source/Runtime/AssetRegistry/Private/AssetRegistry.h
+++ b/Engine/Source/Runtime/AssetRegistry/Private/AssetRegistry.h
@@ -320,7 +320,7 @@ private:
 	bool IsDirAlreadyWatchedByRootWatchers(const FString& Directory) const;
 #endif
 
-	// Implementation methods for FPausePathDismountingScope
+	// Implementation methods for FSuspendPathsDismounting
 	virtual void SuspendPathsDismounting() override;
 	virtual void ResumePathsDismounting() override;
 
diff --git a/Engine/Source/Runtime/AssetRegistry/Public/AssetRegistry/IAssetRegistry.h b/Engine/Source/Runtime/AssetRegistry/Public/AssetRegistry/IAssetRegistry.h
--- a/Engine/Source/Runtime/AssetRegistry/Public/AssetRegistry/IAssetRegistry.h
+++ b/Engine/Source/Runtime/AssetRegistry/Public/AssetRegistry/IAssetRegistry.h
@@ -1167,7 +1167,7 @@ public:
 
 protected:
 
-	// Implementation methods for FPausePathDismountingScope
+	// Implementation methods for FSuspendPathsDismounting
 	virtual void SuspendPathsDismounting() {};
 	virtual void ResumePathsDismounting() {};
 
