AssetRegistry: On non-windows platforms, temporarily turn off the use of the system-wide lock around the write of cache files. It is failing and hanging forever on Mac and Linux.
#rnx

diff --git a/Engine/Source/Runtime/AssetRegistry/Private/AssetDataGatherer.cpp b/Engine/Source/Runtime/AssetRegistry/Private/AssetDataGatherer.cpp
--- a/Engine/Source/Runtime/AssetRegistry/Private/AssetDataGatherer.cpp
+++ b/Engine/Source/Runtime/AssetRegistry/Private/AssetDataGatherer.cpp
@@ -338,7 +338,9 @@ TArray<FString> FGatherSettings::FindShardedCacheFiles() const
 
 FGatherSettings GGatherSettings;
 
-#if WITH_EDITOR
+#define CACHEFILE_SYSTEMWIDE_WRITELOCK WITH_EDITOR && PLATFORM_WINDOWS
+
+#if CACHEFILE_SYSTEMWIDE_WRITELOCK
 // Attempts to lock a system-wide mutex guarding the writing of AssetRegistry cache files. If already locked, no waiting occurs.
 // Users should verify the lock state after calling this function.
 [[nodiscard]] UE::FPlatformSystemWideMutex TryLockSystemWideCacheFilesWrites()
@@ -4166,7 +4168,7 @@ void FAssetDataGatherer::SaveCacheFile(const TArray<TPair<FName,FDiskCachedAsset
 	// Attempt to get exclusive write access to the cache files. If we cannot, skip writing
 	// and try again next tick. Failing to get the lock here means another process 
 	// is currently writing so we avoid the possibility of never writing cache files 
-#if WITH_EDITOR
+#if CACHEFILE_SYSTEMWIDE_WRITELOCK
 	UE::FPlatformSystemWideMutex CacheFilesLock = TryLockSystemWideCacheFilesWrites();
 	if (!CacheFilesLock.IsValid())
 	{
