[AssetRegistry]

- Removes CacheInUseCount as it's functionality was duplicated by bIsSavingAsyncCache. Clarified comments regarding to the expected state transitions for bIsSavingAsyncCache when saves are attempted
- RenamesFAssetDataGatherer::SaveCacheFile to FAssetDataGatherer::TrySaveCacheFile to more clearly indicate the save can fail. When a save fails we conditionally request another save to be handled in a future tick of the gatherer.
- FAssetDataGatherer::ClearCache will now attempt to save the cache to disk before clearing the in-memory cache. No attempt is made to ensure the cache save is successful and opts to let the cache writing be best effort.
- Re-enables the SystemWideMutex for cache writing now that the above changes have been made. Previously we could deadlock if a save was attempted, but failed incrementing CacheInUseCount and never decrementing it. This would hang ClearCache as it waits for CacheInUseCount to become zero.

#rnx

diff --git a/Engine/Source/Runtime/AssetRegistry/Private/AssetDataGatherer.cpp b/Engine/Source/Runtime/AssetRegistry/Private/AssetDataGatherer.cpp
--- a/Engine/Source/Runtime/AssetRegistry/Private/AssetDataGatherer.cpp
+++ b/Engine/Source/Runtime/AssetRegistry/Private/AssetDataGatherer.cpp
@@ -338,7 +338,7 @@ TArray<FString> FGatherSettings::FindShardedCacheFiles() const
 
 FGatherSettings GGatherSettings;
 
-#define CACHEFILE_SYSTEMWIDE_WRITELOCK WITH_EDITOR && PLATFORM_WINDOWS
+#define CACHEFILE_SYSTEMWIDE_WRITELOCK WITH_EDITOR
 
 #if CACHEFILE_SYSTEMWIDE_WRITELOCK
 // Attempts to lock a system-wide mutex guarding the writing of AssetRegistry cache files. If already locked, no waiting occurs.
@@ -3821,7 +3821,6 @@ FAssetDataGatherer::FAssetDataGatherer(UAssetRegistryImpl& InRegistryImpl)
 	, bIsAdditionalMountSearchInProgress(false)
 	, bGatherOnGameThreadOnly(false)
 	, LastCacheSaveNumUncachedAssetFiles(0)
-	, CacheInUseCount(0)
 	, bIsSavingAsyncCache(false)
 	, bFlushedRetryFiles(false)
 {
@@ -4082,15 +4081,11 @@ FAssetDataGatherer::ETickResult FAssetDataGatherer::InnerTickLoop(bool bInSynchr
 		if (bContributeToCacheSave)
 		{
 			TryReserveSaveCache(bShouldSaveCache, AssetsToSave);
-			if (bShouldSaveCache)
-			{
-				CacheInUseCount++;
-			}
 		}
 	}
 	if (bShouldSaveCache)
 	{
-		SaveCacheFile(AssetsToSave);
+		TrySaveCacheFile(AssetsToSave, true /*bRequestSaveOnFailure*/);
 	}
 	return Result;
 }
@@ -4114,7 +4109,7 @@ UE::AssetRegistry::Impl::EGatherStatus FAssetDataGatherer::TryTickOnBackgroundTh
 	return Status;
 }
 
-void FAssetDataGatherer::SaveCacheFile(const TArray<TPair<FName,FDiskCachedAssetData*>>& AssetsToSave)
+bool FAssetDataGatherer::TrySaveCacheFile(const TArray<TPair<FName,FDiskCachedAssetData*>>& AssetsToSave, bool bRequestSaveOnFailure)
 {
 	using namespace UE::AssetDataGather::Private;
 	TRACE_CPUPROFILER_EVENT_SCOPE_STR("Save Cache")
@@ -4124,6 +4119,7 @@ void FAssetDataGatherer::SaveCacheFile(const TArray<TPair<FName,FDiskCachedAsset
 	const int32 CacheShards = FMath::RoundUpToPowerOfTwo(FMath::Min((AssetsToSave.Num() + CacheShardAssetCount-1) / CacheShardAssetCount, MaxShards));
 	const int32 ShardMask = CacheShards - 1;
 	std::atomic<int64> TotalCacheSize{0};
+	bool bSuccess = true;
 
 	TArray<TArray<TPair<FName, FDiskCachedAssetData*>>> DataPerShard;
 	{
@@ -4165,57 +4161,66 @@ void FAssetDataGatherer::SaveCacheFile(const TArray<TPair<FName,FDiskCachedAsset
 		FPlatformProcess::Sleep(static_cast<float>(LastCacheWriteTime + WaitTimeBeforeReopen - CurrentTime));
 	}
 
+#if CACHEFILE_SYSTEMWIDE_WRITELOCK
 	// Attempt to get exclusive write access to the cache files. If we cannot, skip writing
 	// and try again next tick. Failing to get the lock here means another process 
 	// is currently writing so we avoid the possibility of never writing cache files 
-#if CACHEFILE_SYSTEMWIDE_WRITELOCK
 	UE::FPlatformSystemWideMutex CacheFilesLock = TryLockSystemWideCacheFilesWrites();
 	if (!CacheFilesLock.IsValid())
 	{
-		return;
+		bSuccess = false;
+		if (bRequestSaveOnFailure)
+		{
+			RequestAsyncCacheSave();
+		}
+		// Intentionally fall-through to perform necessary cleanup
 	}
+	else
 #endif
+	{
+		ParallelFor(CacheShards, [this, &AssetsToSave, &TotalCacheSize, &DataPerShard, ShardMask](int32 Shard)
+		{
+			TArray<TPair<FName, FDiskCachedAssetData*>>& ShardData = DataPerShard[Shard];
+			FString Filename = FString::Printf(TEXT("%s_%d.bin"), *GGatherSettings.GetCacheBaseFilename(), Shard);
+			int64 CacheSize = SaveCacheFileInternal(Filename, ShardData);
+			TotalCacheSize += CacheSize;
+		}, EParallelForFlags::BackgroundPriority);
 
-	ParallelFor(CacheShards, [this, &AssetsToSave, &TotalCacheSize, &DataPerShard, ShardMask](int32 Shard) 
-	{
-		TArray<TPair<FName, FDiskCachedAssetData*>>& ShardData = DataPerShard[Shard];
-		FString Filename = FString::Printf(TEXT("%s_%d.bin"), *GGatherSettings.GetCacheBaseFilename(), Shard);
-		int64 CacheSize = SaveCacheFileInternal(Filename, ShardData);
-		TotalCacheSize += CacheSize;
-	}, EParallelForFlags::BackgroundPriority);
-
-	UE_LOG(LogAssetRegistry, Display, TEXT("Asset registry cache written as %.1f MiB to %s_*.bin"), 
-		static_cast<float>(TotalCacheSize)/1024.f/1024.f, *GGatherSettings.GetCacheBaseFilename());
-	// Delete old name of monolithic cache file and old non-monolithic cache directory if they exist
-	IFileManager::Get().Delete(*GGatherSettings.GetLegacyCacheFilename(), false /* bRequireExists */,
-		true /* EvenReadOnly */, true /* Quiet */);
-	IFileManager::Get().DeleteDirectory(*GGatherSettings.GetLegacyNonMonolithicCacheDirectory(),
-		false /* bRequireExists */, true /* Tree */);
-
-	// Delete any other shards if number of shards was reduced
-	TArray<FString> CacheFiles = GGatherSettings.FindShardedCacheFiles();
-	for (const FString& CacheFile : CacheFiles)
-	{
-		FStringView BaseName = FPathViews::GetBaseFilename(CacheFile);
-		if (int32 Index; BaseName.FindLastChar('_', Index))
-		{
-			BaseName.MidInline(Index+1);
-			int32 Suffix = 0;
-			LexFromString(Suffix, BaseName);
-			if (Suffix >= CacheShards)
+		UE_LOG(LogAssetRegistry, Display, TEXT("Asset registry cache written as %.1f MiB to %s_*.bin"),
+			static_cast<float>(TotalCacheSize) / 1024.f / 1024.f, *GGatherSettings.GetCacheBaseFilename());
+		// Delete old name of monolithic cache file and old non-monolithic cache directory if they exist
+		IFileManager::Get().Delete(*GGatherSettings.GetLegacyCacheFilename(), false /* bRequireExists */,
+			true /* EvenReadOnly */, true /* Quiet */);
+		IFileManager::Get().DeleteDirectory(*GGatherSettings.GetLegacyNonMonolithicCacheDirectory(),
+			false /* bRequireExists */, true /* Tree */);
+
+		// Delete any other shards if number of shards was reduced
+		TArray<FString> CacheFiles = GGatherSettings.FindShardedCacheFiles();
+		for (const FString& CacheFile : CacheFiles)
+		{
+			FStringView BaseName = FPathViews::GetBaseFilename(CacheFile);
+			if (int32 Index; BaseName.FindLastChar('_', Index))
 			{
-				IFileManager::Get().Delete(*CacheFile);
+				BaseName.MidInline(Index + 1);
+				int32 Suffix = 0;
+				LexFromString(Suffix, BaseName);
+				if (Suffix >= CacheShards)
+				{
+					IFileManager::Get().Delete(*CacheFile);
+				}
 			}
 		}
+
+		FGathererScopeLock ResultsScopeLock(&ResultsLock);
+		LastCacheWriteTime = FPlatformTime::Seconds();
 	}
 
 	FScopedGatheringPause ScopedPause(*this);
 	FGathererScopeLock TickScopeLock(&TickLock);
+	check(bIsSavingAsyncCache);
 	bIsSavingAsyncCache = false;
-	check(CacheInUseCount > 0);
-	CacheInUseCount--;
-	FGathererScopeLock ResultsScopeLock(&ResultsLock);
-	LastCacheWriteTime = FPlatformTime::Seconds();
+
+	return bSuccess;
 }
 
 FAssetDataGatherer::FScopedGatheringPause::FScopedGatheringPause(const FAssetDataGatherer& InOwner)
@@ -4894,14 +4899,19 @@ void FAssetDataGatherer::WaitOnPath(FStringView InPath)
 
 void FAssetDataGatherer::ClearCache()
 {
-	bool bCacheIsInUseOnOtherThread = false;
+	bool bOtherThreadIsSavingAsyncCache = false;
 	bool bWasCacheEnabled = false;
+	bool bShouldSaveCache = false;
+	TArray<TPair<FName, FDiskCachedAssetData*>> AssetsToSave;
 	{
 		FGathererScopeLock TickScopeLock(&TickLock);
+		// Must be done before we disable writing
+		TryReserveSaveCache(bShouldSaveCache, AssetsToSave);
+
 		bWasCacheEnabled = bCacheWriteEnabled || bCacheReadEnabled;
 		bCacheWriteEnabled = false;
 		bCacheReadEnabled = false;
-		bCacheIsInUseOnOtherThread = CacheInUseCount > 0;
+		bOtherThreadIsSavingAsyncCache = bIsSavingAsyncCache && !bShouldSaveCache;
 	}
 
 	if (!bWasCacheEnabled)
@@ -4910,14 +4920,20 @@ void FAssetDataGatherer::ClearCache()
 	}
 
 	// Wait for any cache saves to complete because saves read the cache data we are about to delete.
-	// Saves are executed outside of the lock, but they indicate they are in progress by incrementing CacheInUseCount.
-	// CacheInUseCount is no longer incremented after bCacheEnabled=false which we set above, so starvation should not be possible.
-	while (bCacheIsInUseOnOtherThread)
+	// Saves are executed outside of the lock, but they indicate they are in progress by setting bIsSavingAsyncCache=true.
+	// bIsSavingAsyncCache cannot be set to true after bCacheWriteEnabled=false which we set above, so starvation should not be possible.
+	while (bOtherThreadIsSavingAsyncCache)
 	{
 		constexpr float WaitForSaveCompleteTime = .001f;
 		FPlatformProcess::Sleep(WaitForSaveCompleteTime);
 		FGathererScopeLock TickScopeLock(&TickLock);
-		bCacheIsInUseOnOtherThread = CacheInUseCount > 0;
+		bOtherThreadIsSavingAsyncCache = bIsSavingAsyncCache;
+	}
+
+	if (bShouldSaveCache)
+	{
+		// If this save fails, we don't re-attempt it and accept we couldn't write out the latest state
+		TrySaveCacheFile(AssetsToSave);
 	}
 
 	{
diff --git a/Engine/Source/Runtime/AssetRegistry/Private/AssetDataGatherer.h b/Engine/Source/Runtime/AssetRegistry/Private/AssetDataGatherer.h
--- a/Engine/Source/Runtime/AssetRegistry/Private/AssetDataGatherer.h
+++ b/Engine/Source/Runtime/AssetRegistry/Private/AssetDataGatherer.h
@@ -229,8 +229,9 @@ public:
 	 */
 	void WaitOnPath(FStringView LocalPath);
 	/**
-	 * Empty the cache read from disk and the cache used to write to disk. Disable further caching.
-	 * Used to save memory when cooking after the scan is complete.
+	 * Saves the current cache state to disk, then empties the caches holding data read from disk and new data to be 
+	 * written to disk. Disables further caching. Used to save memory when cooking after the scan is complete.
+	 * Saving is best-effort and may fail if the cache files are open for read on the system.
 	*/
 	void ClearCache();
 
@@ -407,8 +408,12 @@ private:
 	 * should be saved, for later saving outside of the critical section.
 	 */
 	void TryReserveSaveCache(bool& bOutShouldSave, TArray<TPair<FName,FDiskCachedAssetData*>>& AssetsToSave);
-	/** Save cache file for the assetdatas read from package headers, possibly sharded into multiple files. */
-	void SaveCacheFile(const TArray<TPair<FName,FDiskCachedAssetData*>>& AssetsToSave);
+	/** 
+	* Tries to save the cache file for the assetdatas read from package headers, possibly sharded into multiple files. 
+	* If the save fails, bRequestSaveOnFailure specifies whether a save will be requested again to be handled in a future gatherer tick.
+	* Returns whether the save was successful or not.
+	*/
+	bool TrySaveCacheFile(const TArray<TPair<FName,FDiskCachedAssetData*>>& AssetsToSave, bool bRequestSaveOnFailure = false);
 	/**
 	 * If the CacheFilename/AssetsToSave are non empty, save the cache file. 
 	 * This function reads the read-only-after-creation data from each FDiskCachedAssetData*, but otherwise does not use
@@ -600,17 +605,15 @@ private:
 	TMap<FName, FDiskCachedAssetData*> NewCachedAssetDataMap;
 	/** How many uncached asset files had been discovered at the last async cache save */
 	int32 LastCacheSaveNumUncachedAssetFiles;
-	/**
-	 * Incremented when a thread is in the middle of saving any cache and therefore the cache cannot be deleted,
-	 * decremented when the thread is done. Only incremented when bCacheEnabled has been recently confirmed to be true.
-	 */
-	int32 CacheInUseCount;
 	/**
 	 * True if the current TickInternal is synchronous, which may be because !IsSynchronous or because the game thread has
 	 * taken over the tick for a synchronous function.
 	 */
 	bool bSynchronousTick;
-	/** True when a thread is saving an async cache and so another save of the cache should not be triggered. */
+	/** 
+	* True when a thread is saving an async cache and therefore the cache cannot be deleted.
+	* This value resets after the save attempt has completed regardless of whether it was successful or not. 
+	*/
 	bool bIsSavingAsyncCache;
 	/** Packages can be marked for retry up until bInitialPluginsLoaded is set. After it is set, we retry them once. */
 	bool bFlushedRetryFiles;
