AssetRegistry: move contended IsMonitored test to a later point.

#Synced-CL 48301749

diff --git a/Engine/Source/Runtime/AssetRegistry/Private/AssetRegistry.cpp b/Engine/Source/Runtime/AssetRegistry/Private/AssetRegistry.cpp
--- a/Engine/Source/Runtime/AssetRegistry/Private/AssetRegistry.cpp
+++ b/Engine/Source/Runtime/AssetRegistry/Private/AssetRegistry.cpp
@@ -8966,11 +8966,6 @@ void FAssetRegistryImpl::AddLoadedAssetToProcess(const UObject& AssetLoaded)
 		{
 			return;
 		}
-
-		if (!GlobalGatherer->IsMonitored(LocalPath))
-		{
-			return;
-		}
 	}
 
 	LoadedAssetsToProcess.Add(&AssetLoaded);
@@ -8995,6 +8990,17 @@ void FAssetRegistryImpl::GetProcessLoadedAssetsBatch(TArray<const UObject*>& Out
 			continue;
 		}
 
+		FString LocalPath;
+		if (UNLIKELY(!FPackageName::TryConvertLongPackageNameToFilename(LoadedAsset->GetPackage()->GetName(), LocalPath)))
+		{
+			continue;
+		}
+
+		if (!GlobalGatherer->IsMonitored(LocalPath))
+		{
+			continue;
+		}
+
 		// Take a new snapshot of the asset's data every time it loads or saves
 
 		UPackage* InMemoryPackage = LoadedAsset->GetOutermost();
