Asset registry should handle junctions better

#Synced-CL 48597399

diff --git a/Engine/Source/Runtime/AssetRegistry/Private/AssetDataGatherer.cpp b/Engine/Source/Runtime/AssetRegistry/Private/AssetDataGatherer.cpp
--- a/Engine/Source/Runtime/AssetRegistry/Private/AssetDataGatherer.cpp
+++ b/Engine/Source/Runtime/AssetRegistry/Private/AssetDataGatherer.cpp
@@ -2444,6 +2444,44 @@ bool FPathExistence::IsReadOnly()
 	return bIsReadOnly;
 }
 
+bool TryFixSymlinkCapitalization(FString& LocalAbsPath, FString& CorrectedCapitalization)
+{
+	// try to find the last junction / symlink 
+	FString Path = FString(FPathViews::GetPath(LocalAbsPath));
+	while (!Path.IsEmpty() && !IFileManager::Get().IsSymlink(*Path))
+	{
+		Path = FString(FPathViews::GetPath(Path));
+	}
+
+	// We couldn't find a symlink on the path, we don't know why the capitalization is wrong.
+	if (Path.IsEmpty())
+	{
+		return false;
+	}
+
+	FStringView LocalAbsPathAfterJunction(LocalAbsPath);
+	LocalAbsPathAfterJunction = LocalAbsPathAfterJunction.RightChop(Path.Len());
+
+	FStringView CorrectedCapitalizationAfterJunction(CorrectedCapitalization);
+	CorrectedCapitalizationAfterJunction = CorrectedCapitalizationAfterJunction.RightChop(CorrectedCapitalizationAfterJunction.Len() - LocalAbsPathAfterJunction.Len());
+
+	// we expect something like: [junction]/foo/bar.uasset <-> [fullpath]/Foo/BAR.uasset
+	if (!LocalAbsPathAfterJunction.Equals(CorrectedCapitalizationAfterJunction, ESearchCase::IgnoreCase))
+	{
+		//unexpected but no need for a separate error
+		return false;
+	}
+
+	int LocalAbsPathIt = LocalAbsPath.Len() - 1;
+	int CorrectedCapitalizationIt = CorrectedCapitalization.Len() - 1;
+	for (int i = 0; i < LocalAbsPathAfterJunction.Len(); ++i)
+	{
+		LocalAbsPath[LocalAbsPathIt--] = CorrectedCapitalization[CorrectedCapitalizationIt--];
+	}
+
+	return true;
+}
+
 void FPathExistence::LoadExistenceData()
 {
 	if (bHasExistenceData)
@@ -2458,11 +2496,14 @@ void FPathExistence::LoadExistenceData()
 		{
 			LocalAbsPath = MoveTemp(CorrectedCapitalization);
 		}
-		else
+		else // if the case insensitive comparison above failed, there might be junction / symlink in the path
 		{
-			UE_LOG(LogAssetRegistry, Warning,
-				TEXT("FPathExistence failed to gather correct capitalization from disk for %s, because GetFilenameOnDisk returned non-matching filename '%s'."),
-				*LocalAbsPath, *CorrectedCapitalization);
+			if (!TryFixSymlinkCapitalization(LocalAbsPath, CorrectedCapitalization))
+			{
+				UE_LOG(LogAssetRegistry, Warning,
+					TEXT("FPathExistence failed to gather correct capitalization from disk for %s, because GetFilenameOnDisk returned non-matching filename '%s'."),
+					*LocalAbsPath, *CorrectedCapitalization);
+			}
 		}
 
 		ModificationTime = StatData.ModificationTime;
@@ -2483,11 +2524,14 @@ void FPathExistence::LoadExistenceData()
 			{
 				LocalAbsPath = MoveTemp(CorrectedCapitalization);
 			}
-			else
+			else // if the case insensitive comparison above failed, there might be junction / symlink in the path
 			{
-				UE_LOG(LogAssetRegistry, Warning,
-					TEXT("FPathExistence failed to gather correct capitalization from disk for %s, because GetFilenameOnDisk returned non-matching filename '%s'."),
-					*LocalAbsPath, *CorrectedCapitalization);
+				if (!TryFixSymlinkCapitalization(LocalAbsPath, CorrectedCapitalization))
+				{
+					UE_LOG(LogAssetRegistry, Warning,
+						TEXT("FPathExistence failed to gather correct capitalization from disk for %s, because GetFilenameOnDisk returned non-matching filename '%s'."),
+						*LocalAbsPath, *CorrectedCapitalization);
+				}
 			}
 			PathType = EType::MissingButDirExists;
 		}
diff --git a/Engine/Source/Runtime/Core/Private/HAL/FileManagerGeneric.cpp b/Engine/Source/Runtime/Core/Private/HAL/FileManagerGeneric.cpp
--- a/Engine/Source/Runtime/Core/Private/HAL/FileManagerGeneric.cpp
+++ b/Engine/Source/Runtime/Core/Private/HAL/FileManagerGeneric.cpp
@@ -585,6 +585,11 @@ FString FFileManagerGeneric::GetFilenameOnDisk(const TCHAR* Filename)
 	return GetLowLevel().GetFilenameOnDisk(Filename);
 }
 
+bool FFileManagerGeneric::IsSymlink(const TCHAR* Filename)
+{
+	return GetLowLevel().IsSymlink(Filename) == ESymlinkResult::Symlink;
+}
+
 FString FFileManagerGeneric::DefaultConvertToRelativePath( const TCHAR* Filename )
 {
 	//default to the full absolute path of this file
diff --git a/Engine/Source/Runtime/Core/Public/GenericPlatform/GenericPlatformFile.h b/Engine/Source/Runtime/Core/Public/GenericPlatform/GenericPlatformFile.h
--- a/Engine/Source/Runtime/Core/Public/GenericPlatform/GenericPlatformFile.h
+++ b/Engine/Source/Runtime/Core/Public/GenericPlatform/GenericPlatformFile.h
@@ -464,8 +464,15 @@ public:
 	/** For case insensitive filesystems, returns the full path of the file with the same case as in the filesystem */
 	virtual FString GetFilenameOnDisk(const TCHAR* Filename) = 0;
 
-	/** Return true if the file is a symbolic link */
-	virtual ESymlinkResult IsSymlink(const TCHAR* Filename) { return ESymlinkResult::Unimplemented; }
+	/** Return ESymlinkResult::Symlink if the file is a symbolic link, if unimplemented, try to go down a level */
+	virtual ESymlinkResult IsSymlink(const TCHAR* Filename) 
+	{ 
+		if (GetLowerLevel())
+		{
+			return GetLowerLevel()->IsSymlink(Filename);
+		}
+		return ESymlinkResult::Unimplemented; 
+	}
 
 	/**
 	 * Determine if the file has been downloaded from a web browser, based on platform-specific metadata.
diff --git a/Engine/Source/Runtime/Core/Public/HAL/FileManager.h b/Engine/Source/Runtime/Core/Public/HAL/FileManager.h
--- a/Engine/Source/Runtime/Core/Public/HAL/FileManager.h
+++ b/Engine/Source/Runtime/Core/Public/HAL/FileManager.h
@@ -288,4 +288,13 @@ public:
 	* @return	Filename with the same case as in the filesystem.
 	*/
 	virtual FString GetFilenameOnDisk(const TCHAR* Filename) = 0;
+
+	/**
+	* Returns if the given path is a symlink. Depending on the OS this can be a soft link, hard link or a junction.
+	*
+	* @param Path	Path to check if it's a symlink
+	*
+	* @return	True if symlink, false if not or the symlink check is not implemented. 
+	*/
+	virtual bool IsSymlink(const TCHAR* Path) = 0;
 };
diff --git a/Engine/Source/Runtime/Core/Public/HAL/FileManagerGeneric.h b/Engine/Source/Runtime/Core/Public/HAL/FileManagerGeneric.h
--- a/Engine/Source/Runtime/Core/Public/HAL/FileManagerGeneric.h
+++ b/Engine/Source/Runtime/Core/Public/HAL/FileManagerGeneric.h
@@ -99,6 +99,7 @@ public:
 	CORE_API void GetTimeStampPair(const TCHAR* PathA, const TCHAR* PathB, FDateTime& OutTimeStampA, FDateTime& OutTimeStampB);
 	CORE_API bool SetTimeStamp( const TCHAR* Filename, FDateTime Timestamp ) override;
 	CORE_API virtual FString GetFilenameOnDisk(const TCHAR* Filename) override;
+	CORE_API virtual bool IsSymlink(const TCHAR* Filename) override;
 
 	CORE_API virtual uint32	Copy( const TCHAR* Dest, const TCHAR* Src, bool Replace = 1, bool EvenIfReadOnly = 0, bool Attributes = 0, FCopyProgress* Progress = nullptr, EFileRead ReadFlags = FILEREAD_None, EFileWrite WriteFlags = FILEWRITE_None ) override;
 	CORE_API virtual bool	MakeDirectory( const TCHAR* Path, bool Tree=0 ) override;
