Add AssetRegistry preload timing regions
#tests tested on console
#Synced-CL 0

diff --git a/Engine/Source/Runtime/AssetRegistry/Private/AssetRegistry.cpp b/Engine/Source/Runtime/AssetRegistry/Private/AssetRegistry.cpp
--- a/Engine/Source/Runtime/AssetRegistry/Private/AssetRegistry.cpp
+++ b/Engine/Source/Runtime/AssetRegistry/Private/AssetRegistry.cpp
@@ -548,6 +548,7 @@ private:
 
 	void KickPreload()
 	{
+		CSV_TRACE_REGION_BEGIN(TEXT("AssetRegistry Preload"))
 		TRACE_CPUPROFILER_EVENT_SCOPE(FCookedAssetRegistryPreloader::KickPreload);
 		// Called from Within the Lock
 		check(LoadState == EState::NotFound && !ARPath.IsEmpty());
@@ -659,6 +660,7 @@ private:
 		}
 		ARPath.Reset();
 		Payload.Reset();
+		CSV_TRACE_REGION_END(TEXT("AssetRegistry Preload"))
 	}
 
 	/** simple way to trigger a callback at a specific time that TaskGraph is usable. */
