Use a category for AssetRegistry timing regions
#tests tested on console
#rnx
#Synced-CL 0

diff --git a/Engine/Source/Runtime/AssetRegistry/Private/AssetRegistry.cpp b/Engine/Source/Runtime/AssetRegistry/Private/AssetRegistry.cpp
--- a/Engine/Source/Runtime/AssetRegistry/Private/AssetRegistry.cpp
+++ b/Engine/Source/Runtime/AssetRegistry/Private/AssetRegistry.cpp
@@ -548,7 +548,7 @@ private:
 
 	void KickPreload()
 	{
-		CSV_TRACE_REGION_BEGIN(TEXT("AssetRegistry Preload"))
+		CSV_TRACE_REGION_BEGIN(TEXT("Preload"), TEXT("AssetRegistry"))
 		TRACE_CPUPROFILER_EVENT_SCOPE(FCookedAssetRegistryPreloader::KickPreload);
 		// Called from Within the Lock
 		check(LoadState == EState::NotFound && !ARPath.IsEmpty());
@@ -660,7 +660,7 @@ private:
 		}
 		ARPath.Reset();
 		Payload.Reset();
-		CSV_TRACE_REGION_END(TEXT("AssetRegistry Preload"))
+		CSV_TRACE_REGION_END(TEXT("Preload"), TEXT("AssetRegistry"))
 	}
 
 	/** simple way to trigger a callback at a specific time that TaskGraph is usable. */
@@ -2346,7 +2346,7 @@ void FAssetRegistryImpl::SearchAllAssets(Impl::FEventContext& EventContext,
 
 	if (!bInitialSearchStarted)
 	{
-		CSV_TRACE_REGION_BEGIN(TEXT("Asset Registry Scan"));
+		CSV_TRACE_REGION_BEGIN(TEXT("Scan"), TEXT("AssetRegistry"));
 		InitialSearchStartTime = FPlatformTime::Seconds();
 		bInitialSearchStarted = true;
 		bInitialSearchCompleted.store(false, std::memory_order_relaxed);
@@ -6030,7 +6030,7 @@ void FAssetRegistryImpl::OnInitialSearchCompleted(Impl::FEventContext& EventCont
 	SetPerformanceMode(Impl::EPerformanceMode::MostlyStatic);
 
 	LogSearchDiagnostics(InitialSearchStartTime, /* bInitialScan */ true);
-	CSV_TRACE_REGION_END(TEXT("Asset Registry Scan"));
+	CSV_TRACE_REGION_END(TEXT("Scan"), TEXT("AssetRegistry"));
 
 	GlobalGatherer->UpdateCacheForSaving();
 	GlobalGatherer->OnInitialSearchCompleted();
@@ -6053,7 +6053,7 @@ void FAssetRegistryImpl::OnAdditionalMountSearchCompleted(Impl::FEventContext& E
 	SetPerformanceMode(Impl::EPerformanceMode::MostlyStatic);
 
 	LogSearchDiagnostics(AdditionalMountSearchStartTime, /* bInitialScan */ false);
-	CSV_TRACE_REGION_END(TEXT("Asset Registry - Additional Mount Search")); // Matching CSV_TRACE_REGION_BEGIN in OnContentPathMounted
+	CSV_TRACE_REGION_END(TEXT("Additional Mount Search"), TEXT("AssetRegistry")); // Matching CSV_TRACE_REGION_BEGIN in OnContentPathMounted
 
 	GlobalGatherer->UpdateCacheForSaving();
 	GlobalGatherer->OnAdditionalMountSearchCompleted();
@@ -9266,7 +9266,7 @@ void FAssetRegistryImpl::OnContentPathMounted(UE::AssetRegistry::FInterfaceWrite
 		{
 			if (!IsGathering())
 			{
-				CSV_TRACE_REGION_BEGIN(TEXT("Asset Registry - Additional Mount Search")); // Matching CSV_TRACE_REGION_END in OnAdditionalMountSearchCompleted
+				CSV_TRACE_REGION_BEGIN(TEXT("Additional Mount Search"), TEXT("AssetRegistry")); // Matching CSV_TRACE_REGION_END in OnAdditionalMountSearchCompleted
 				AdditionalMountSearchStartTime = FPlatformTime::Seconds();
 				bAdditionalMountSearchInProgress.store(true, std::memory_order_relaxed);
 				GlobalGatherer->SetIsAdditionalMountSearchInProgress(true);
