[AssetRegistry]

Allow the gather thread to locally determine if it should tick the AssetRegistry using stack memory rather than a member since that member can be mutated by multiple threads

#rnx
#Synced-CL 48712420

diff --git a/Engine/Source/Runtime/AssetRegistry/Private/AssetDataGatherer.cpp b/Engine/Source/Runtime/AssetRegistry/Private/AssetDataGatherer.cpp
--- a/Engine/Source/Runtime/AssetRegistry/Private/AssetDataGatherer.cpp
+++ b/Engine/Source/Runtime/AssetRegistry/Private/AssetDataGatherer.cpp
@@ -3857,7 +3857,6 @@ FAssetDataGatherer::FAssetDataGatherer(UAssetRegistryImpl& InRegistryImpl)
 	, bHasLoadedCache(false)
 	, bDiscoveryIsComplete(false)
 	, bIsComplete(false)
-	, bRequestAssetRegistryTick(false)
 	, bIsIdle(false)
 	, bFirstTickAfterIdle(true)
 	, bFinishedInitialDiscovery(false)
@@ -4006,12 +4005,12 @@ uint32 FAssetDataGatherer::Run()
 
 	while (!IsStopped)
 	{
+		bool bShouldRequestAssetRegistryTick = false;
 		ETickResult TickResult = InnerTickLoop(false /* bInSynchronousTick */, true /* bContributeToCacheSave */,
-			-1. /* EndTimeSeconds */);
+			-1. /* EndTimeSeconds */, &bShouldRequestAssetRegistryTick);
 
-		if (bRequestAssetRegistryTick)
+		if (bShouldRequestAssetRegistryTick)
 		{
-			bRequestAssetRegistryTick = false;
 			TryTickOnBackgroundThread();
 		}
 
@@ -4050,7 +4049,7 @@ uint32 FAssetDataGatherer::Run()
 }
 
 FAssetDataGatherer::ETickResult FAssetDataGatherer::InnerTickLoop(bool bInSynchronousTick,
-	bool bContributeToCacheSave, double EndTimeSeconds)
+	bool bContributeToCacheSave, double EndTimeSeconds, bool* bShouldRequestAssetRegistryTick)
 {
 	using namespace UE::AssetDataGather::Private;
 	ETickResult Result = ETickResult::KeepTicking;
@@ -4072,7 +4071,7 @@ FAssetDataGatherer::ETickResult FAssetDataGatherer::InnerTickLoop(bool bInSynchr
 
 		for (;;)
 		{
-			ETickResult TickResult = TickInternal(TickStartTime, bPollDiscovery);
+			ETickResult TickResult = TickInternal(TickStartTime, bPollDiscovery, bShouldRequestAssetRegistryTick);
 			if (EndTimeSeconds > 0. && FPlatformTime::Seconds() > EndTimeSeconds)
 			{
 				Result = ETickResult::Interrupt;
@@ -4319,7 +4318,7 @@ void FAssetDataGatherer::EnsureCompletion()
 	}
 }
 
-FAssetDataGatherer::ETickResult FAssetDataGatherer::TickInternal(double& TickStartTime, bool bPollDiscovery)
+FAssetDataGatherer::ETickResult FAssetDataGatherer::TickInternal(double& TickStartTime, bool bPollDiscovery, bool* bShouldRequestAssetRegistryTick)
 {
 	using namespace UE::AssetDataGather::Private;
 	using namespace UE::Tasks;
@@ -4341,6 +4340,10 @@ FAssetDataGatherer::ETickResult FAssetDataGatherer::TickInternal(double& TickSta
 	bool bLocalIsCacheWriteEnabled = false;
 	double LocalLastCacheWriteTime = 0.0;
 	ETickResult TickResult = ETickResult::KeepTicking;
+	if (bShouldRequestAssetRegistryTick)
+	{
+		*bShouldRequestAssetRegistryTick = false;
+	}
 
 	{
 		FGathererScopeLock ResultsScopeLock(&ResultsLock);
@@ -4572,7 +4575,10 @@ FAssetDataGatherer::ETickResult FAssetDataGatherer::TickInternal(double& TickSta
 
 	// If we have finished any async reads outside of a wait batch, signal that we want to 
 	// kick the gatherer to help avoid gather work from running on the game thread
-	bRequestAssetRegistryTick = !bInWaitBatch && !CompletedReadContexts.IsEmpty();
+	if (bShouldRequestAssetRegistryTick)
+	{
+		*bShouldRequestAssetRegistryTick = !bInWaitBatch && !CompletedReadContexts.IsEmpty();
+	}
 
 	// Accumulate the results
 	for (TUniquePtr<FReadContext>& ReadContext : CompletedReadContexts)
diff --git a/Engine/Source/Runtime/AssetRegistry/Private/AssetDataGatherer.h b/Engine/Source/Runtime/AssetRegistry/Private/AssetDataGatherer.h
--- a/Engine/Source/Runtime/AssetRegistry/Private/AssetDataGatherer.h
+++ b/Engine/Source/Runtime/AssetRegistry/Private/AssetDataGatherer.h
@@ -348,12 +348,12 @@ private:
 	 * Helper function to run the tick in a loop-within-a-loop to minimize critical section entry, and to move expensive
 	 * operations out of the critical section
 	 */
-	ETickResult InnerTickLoop(bool bInSynchronousTick, bool bContributeToCacheSave, double EndTimeSeconds);
+	ETickResult InnerTickLoop(bool bInSynchronousTick, bool bContributeToCacheSave, double EndTimeSeconds, bool* bShouldRequestAssetRegistryTick = nullptr);
 	/**
 	 * Tick function to pump scanning and push results into the search results structure. May be called from devoted
 	 * thread or inline from synchronous functions on other threads.
 	 */
-	ETickResult TickInternal(double& TickStartTime, bool bPollDiscovery);
+	ETickResult TickInternal(double& TickStartTime, bool bPollDiscovery, bool* bShouldRequestAssetRegistryTick = nullptr);
 	/**
 	* Conditionally Ticks the AssetRegistry from the background thread returning the result of gathering. Return the new status of the
 	* AssetRegistry's gather. If the background thread should not currently perfom ticking, EGatherStatus::Complete is returned.
@@ -574,8 +574,6 @@ private:
 	bool bDiscoveryIsComplete;
 	/** Track whether this Gather has gone idle and a caller has read all search data from it. */
 	bool bIsComplete;
-	/** Track whether the Gatherers background thread should interrupt reading from disk to tick the AssetRegistry. */
-	bool bRequestAssetRegistryTick;
 	/** Track whether this Gatherer has gone idle, either it has no more work or it's blocked on external events. */
 	bool bIsIdle;
 	/** Track the first tick after idle to set up e.g. timing data. */
