[AssetRegistry]

- Treat serialized directories with invalid journal handles as invalid so they must be rescanned since we won't be able to determine if the directory has changed from the USN records. We have seen discovery cache directories serialized as bCacheValid=true, JournalHandle=Invalid, which makes the directories seem like a known directory with no changes and is never scanned potentially missing changes. This is possible if journal iteration fails; we will proceed to iterate the directory without using the journal API which results in valid directories with invalid journal handles.
- Prune all directories from the cached volume that are invalid or have invalid journal handles.
- Fix FAssetDataDiscovery::AddMountPoint to properly set bOutAlreadyExisted

#rnx

#Synced-CL 48998224

diff --git a/Engine/Source/Runtime/AssetRegistry/Private/AssetDataGatherer.cpp b/Engine/Source/Runtime/AssetRegistry/Private/AssetDataGatherer.cpp
--- a/Engine/Source/Runtime/AssetRegistry/Private/AssetDataGatherer.cpp
+++ b/Engine/Source/Runtime/AssetRegistry/Private/AssetDataGatherer.cpp
@@ -3024,9 +3024,8 @@ void FAssetDataDiscovery::AddMountPoint(const FString& LocalAbsPath, FStringView
 {
 	CHECK_IS_NOT_LOCKED_CURRENT_THREAD(ResultsLock);
 	FGathererScopeLock TreeScopeLock(&TreeLock);
-	bool bAlreadyExisted = false;
-	AddMountPointInternal(LocalAbsPath, LongPackageName, bAlreadyExisted);
-	if (!bAlreadyExisted)
+	AddMountPointInternal(LocalAbsPath, LongPackageName, bOutAlreadyExisted);
+	if (!bOutAlreadyExisted)
 	{
 		SetIsIdle(false);
 	}
diff --git a/Engine/Source/Runtime/AssetRegistry/Private/AssetDataGathererDiscoveryCache.cpp b/Engine/Source/Runtime/AssetRegistry/Private/AssetDataGathererDiscoveryCache.cpp
--- a/Engine/Source/Runtime/AssetRegistry/Private/AssetDataGathererDiscoveryCache.cpp
+++ b/Engine/Source/Runtime/AssetRegistry/Private/AssetDataGathererDiscoveryCache.cpp
@@ -133,9 +133,19 @@ void FAssetDataDiscoveryCache::LoadAndUpdateCache()
 			FCachedDirScanDir& DirData = DirIter->Value;
 			// If not valid, we cannot remove the DirData yet because we will need it to get the list of
 			// its child directories that exist in the tree, at the point when we rescan it
-			if (DirData.bCacheValid && DirData.JournalHandle != FileJournalFileHandleInvalid)
+			if (DirData.bCacheValid)
 			{
-				KnownDirectories.Add(DirData.JournalHandle, DirName);
+				if (DirData.JournalHandle != FileJournalFileHandleInvalid)
+				{
+					KnownDirectories.Add(DirData.JournalHandle, DirName);
+				}
+				else
+				{
+					// We have serialized DirData that doesn't have a valid handle and thus we can't check
+					// if the directory has been modified. Mark it as invalid so it can be rescanned (same 
+					// effect as appending to ModifiedDirectories but avoids the overhead)
+					DirData.bCacheValid = false;
+				}
 			}
 		}
 
@@ -556,9 +566,10 @@ void FAssetDataDiscoveryCache::QueueConsume()
 		// These values get set in FindOrAddIsReparsePoint, make sure they dont get lost here
 		ScanDir.bIsReparsePoint = Existing.bIsReparsePoint;
 		ScanDir.bIsInsideReparsePoint = Existing.bIsInsideReparsePoint;
+		// Confirm our handle is valid otherwise mark the cache as invalid
+		ScanDir.bCacheValid = ScanDir.JournalHandle != FileJournalFileHandleInvalid;
 
 		Existing = MoveTemp(ScanDir);
-		Existing.bCacheValid = true;
 	}
 }
 
@@ -593,6 +604,7 @@ void FCachedVolumeInfo::ConditionalConstruct(const FString& InVolumeName)
 void FCachedVolumeInfo::PreSave()
 {
 	RemoveDirs(MoveTemp(DirsToRemove));
+	RemoveInvalidDirs();
 }
 
 void FCachedVolumeInfo::InitializePlatformData()
@@ -657,6 +669,18 @@ void FCachedVolumeInfo::RemoveDirs(TArray<FString>&& InPaths)
 	InPaths.Empty(); // Free allocated memory
 }
 
+void FCachedVolumeInfo::RemoveInvalidDirs()
+{
+	for (TMap<FString, FCachedDirScanDir>::TIterator It = Dirs.CreateIterator(); It; ++It)
+	{
+		const FCachedDirScanDir& ScanDir = It->Value;
+		if (!ScanDir.bCacheValid || ScanDir.JournalHandle == FileJournalFileHandleInvalid)
+		{
+			It.RemoveCurrent();
+		}
+	}
+}
+
 FCachedDirScanDir* FCachedVolumeInfo::FindDir(FStringView Path)
 {
 	uint32 PathHash = GetTypeHash(Path);
diff --git a/Engine/Source/Runtime/AssetRegistry/Private/AssetDataGathererDiscoveryCache.h b/Engine/Source/Runtime/AssetRegistry/Private/AssetDataGathererDiscoveryCache.h
--- a/Engine/Source/Runtime/AssetRegistry/Private/AssetDataGathererDiscoveryCache.h
+++ b/Engine/Source/Runtime/AssetRegistry/Private/AssetDataGathererDiscoveryCache.h
@@ -98,7 +98,7 @@ public:
 	FCachedDirScanDir& FindOrAddDir(FStringView InPath);
 	FCachedDirScanDir& FindOrAddDir(FStringView InPath, bool& bOutAdded);
 	void RemoveDirs(TArray<FString>&& InPaths);
-
+	void RemoveInvalidDirs();
 	/** Assumes Path has already been normalized. */
 	FCachedDirScanDir* FindDir(FStringView Path);
 
