[AssetRegistry]

Remove implementation of FAssetPackageNameSet::Shrink to be a noop as the underlying set cannot be shrunk to reduce its memory footprint while trying to maintain its load balance.

#rnx

diff --git a/Engine/Source/Runtime/AssetRegistry/Private/AssetDataMap.cpp b/Engine/Source/Runtime/AssetRegistry/Private/AssetDataMap.cpp
--- a/Engine/Source/Runtime/AssetRegistry/Private/AssetDataMap.cpp
+++ b/Engine/Source/Runtime/AssetRegistry/Private/AssetDataMap.cpp
@@ -637,7 +637,8 @@ void FAssetPackageNameMap::Empty(int32 ReservedSize)
 
 void FAssetPackageNameMap::Shrink()
 {
-	AssetOrArrayByPackageName->ResizeToTargetSize();
+	// There is no means to shrink FAssetPackageNameSet.
+	// Do nothing as any change without reducing the HashSize will only cause growth
 }
 
 void FAssetPackageNameMap::Add(FName PackageName, FAssetDataPtrIndex AssetIndex)
