Remove overzealous ensure, this invariant cannot be met in some editor contexts with cooked data

#rnx
#Synced-CL 49181000

diff --git a/Engine/Source/Runtime/AssetRegistry/Private/AssetRegistryHelpers.cpp b/Engine/Source/Runtime/AssetRegistry/Private/AssetRegistryHelpers.cpp
--- a/Engine/Source/Runtime/AssetRegistry/Private/AssetRegistryHelpers.cpp
+++ b/Engine/Source/Runtime/AssetRegistry/Private/AssetRegistryHelpers.cpp
@@ -325,10 +325,6 @@ void UAssetRegistryHelpers::GetDerivedClassAssetData(const TArray<FTopLevelAsset
 				FSoftObjectPath{FTopLevelAssetPath{PackageName, FName(*AssetName)}, FString()}, 
 				bIncludeOnlyOnDiskAssets, 
 				bSkipARFilteredAssets);
-
-			ensureMsgf(!AssetData.PackageName.IsNone(), 
-				TEXT("Failed to find asset data for package: %s.%s"), 
-				*PackageName.ToString(), *AssetName);
 		}
 
 		// AssetData very likely is set, no known cases where we fail to find asset data
