[SEW][CI][AssetRegistry]

- No longer copy AR caches from other streams when running SEW tests as that can expose case mismatches between two different streams for the same asset paths leading to non-determinism
- Invalidate AssetRegistry cache to reset caches which might have incorrect asset path casing for the current stream.

#rnx

#Synced-CL 49541019

diff --git a/Engine/Source/Runtime/AssetRegistry/Private/AssetDataGathererPrivate.h b/Engine/Source/Runtime/AssetRegistry/Private/AssetDataGathererPrivate.h
--- a/Engine/Source/Runtime/AssetRegistry/Private/AssetDataGathererPrivate.h
+++ b/Engine/Source/Runtime/AssetRegistry/Private/AssetDataGathererPrivate.h
@@ -954,7 +954,7 @@ public:
 	static constexpr int32 MinSecondsToElapseBeforeCacheWrite = 300;
 	static constexpr int32 MinAssetReadsBeforeCacheWrite = 100000;
 	static constexpr int32 CacheShardAssetCount = 250000;
-	static constexpr uint32 CacheSerializationMagic = 0x29FE48B0; // Versioning and integrity checking
+	static constexpr uint32 CacheSerializationMagic = 0xC2832E98; // Versioning and integrity checking
 	static constexpr uint64 CurrentVersion = FAssetRegistryVersion::LatestVersion | (uint64(CacheSerializationMagic) << 32);
 	static inline const FTimespan CachePruneAge = FTimespan(7, 0, 0, 0);
 
