[AssetRegistry]

Ensure we have shrink the DirectoryReadsToSchedule map while scheduling work during gathering to reduce any costs with creating iterators for a large sparsely allocated map.

#rnx
#Synced-CL 49587959

diff --git a/Engine/Source/Runtime/AssetRegistry/Private/AssetDataGatherer.cpp b/Engine/Source/Runtime/AssetRegistry/Private/AssetDataGatherer.cpp
--- a/Engine/Source/Runtime/AssetRegistry/Private/AssetDataGatherer.cpp
+++ b/Engine/Source/Runtime/AssetRegistry/Private/AssetDataGatherer.cpp
@@ -6237,16 +6237,13 @@ void FFileReadScheduler::Schedule(const TArray<FString>& PriorityDirectories)
 	}
 
 	// If there are any remaining directories, schedule them for background processing
-	for (TMap<FString, TArray<TUniquePtr<FReadContext>>>::TIterator It = DirectoryReadsToSchedule.CreateIterator(); It; ++It)
+	for (TPair<FString, TArray<TUniquePtr<FReadContext>>>& Pair : DirectoryReadsToSchedule)
 	{
-		const FString& Directory = It->Key;
-		TArray<TUniquePtr<FReadContext>> ReadContexts = MoveTemp(It->Value);
+		const FString& Directory = Pair.Key;
+		TArray<TUniquePtr<FReadContext>> ReadContexts = MoveTemp(Pair.Value);
 		ScheduleDirectory(ReadConcurrencyLimiter, Directory, MoveTemp(ReadContexts), false /*bUnblockImmediately*/);
-		It.RemoveCurrent();
 	}
-
-	// Because of single ownership of FReadContexts, we must be certain we scheduled all directories with ReadContexts by this point
-	check(DirectoryReadsToSchedule.IsEmpty());
+	DirectoryReadsToSchedule.Empty();
 }
 
 void FFileReadScheduler::CollectCompletedReads(TArray<TUniquePtr<FReadContext>>& OutCompletedReadContexts)
@@ -6266,6 +6263,7 @@ void FFileReadScheduler::RemoveScheduledDirectory(const FString& Directory)
 void FFileReadScheduler::Shrink()
 {
 	ScheduledDirectories->Shrink();
+	DirectoryReadsToSchedule.Shrink();
 }
 
 void FFilesToSearch::AddPriorityFile(FGatheredPathData&& FilePath)
