Remove temporary code from AR.

#Synced-CL 49818686

diff --git a/Engine/Source/Runtime/AssetRegistry/Private/AssetRegistryState.cpp b/Engine/Source/Runtime/AssetRegistry/Private/AssetRegistryState.cpp
--- a/Engine/Source/Runtime/AssetRegistry/Private/AssetRegistryState.cpp
+++ b/Engine/Source/Runtime/AssetRegistry/Private/AssetRegistryState.cpp
@@ -2202,20 +2202,6 @@ TArray<FAssetData*> FAssetRegistryState::GetMutableAssetsSortedByObjectPath()
 			return false;
 		};
 	Algo::ParallelSort(SortDatas, CompareSortData);
-	// TODO: Temporary validation of ParallelSort, Remove after a few days
-	for (int32 Index = 0; Index < SortDatas.Num() - 1; ++Index)
-	{
-		if (CompareSortData(SortDatas[Index+1], SortDatas[Index]))
-		{
-			UE_LOG(LogAssetRegistry, Warning,
-				TEXT("Algo::ParallelSort failed to sort correctly. Falling back to Algo::Sort.\n")
-				TEXT("Failed at %d: '%s' > '%s'"),
-				Index, *SortDatas[Index].AssetData->GetSoftObjectPath().ToString(),
-				 *SortDatas[Index+1].AssetData->GetSoftObjectPath().ToString());
-			Algo::Sort(SortDatas, CompareSortData);
-			break;
-		}
-	}	
 
 	TArray<FAssetData*> Result;
 	Result.Reserve(SortDatas.Num());
