[AssetRegistry]

- Don't keep that gatherer thread ticking if a delayed save has been requested as this can result in the gather thread busy waiting until a save is completed when new mounts are added after our last cache save.
- Attempt to save the cache before the gather thread shutsdown if a save is still pending.

#rnx

#Synced-CL 49840101

diff --git a/Engine/Source/Runtime/AssetRegistry/Private/AssetDataGatherer.cpp b/Engine/Source/Runtime/AssetRegistry/Private/AssetDataGatherer.cpp
--- a/Engine/Source/Runtime/AssetRegistry/Private/AssetDataGatherer.cpp
+++ b/Engine/Source/Runtime/AssetRegistry/Private/AssetDataGatherer.cpp
@@ -4002,7 +4002,7 @@ uint32 FAssetDataGatherer::Run()
 	constexpr float PausedSleepTime = 0.005f;
 	LLM_SCOPE(ELLMTag::AssetRegistry);
 
-	while (!IsStopped)
+	while (!IsStopped || IsCacheSaveReady())
 	{
 		bool bShouldRequestAssetRegistryTick = false;
 		ETickResult TickResult = InnerTickLoop(false /* bInSynchronousTick */, true /* bContributeToCacheSave */,
@@ -4136,7 +4136,7 @@ bool FAssetDataGatherer::ShouldBackgroundGatherThreadPauseGathering(bool& bOutGa
 {
 	FGathererScopeLock ResultsScopeLock(&ResultsLock); // bIsIdle requires the lock
 	bOutGathererIsIdle = bIsIdle;
-	return !IsStopped && !IsSaveRequested() && (IsGatheringPaused || bIsIdle);
+	return !IsStopped && !IsCacheSaveReady() && (IsGatheringPaused || bIsIdle);
 }
 
 UE::AssetRegistry::Impl::EGatherStatus FAssetDataGatherer::TryTickOnBackgroundThread()
@@ -5198,25 +5198,30 @@ void FAssetDataGatherer::ConsumeCacheFiles(TArray<UE::AssetDataGather::Private::
 	AssetResults.Reserve(DiskCachedAssetDataMap.Num());
 }
 
-void FAssetDataGatherer::TryReserveSaveCache(bool& bOutShouldSave, TArray<TPair<FName,FDiskCachedAssetData*>>& AssetsToSave)
+bool FAssetDataGatherer::IsCacheSaveReady() const
 {
 	using namespace UE::AssetDataGather::Private;
 
-	bOutShouldSave = false;
 	double SaveAsyncCacheStartTime = 0.0;
-	if (!IsSaveRequested(&SaveAsyncCacheStartTime))
+	if (bIsSavingAsyncCache || !IsCacheSaveRequested(&SaveAsyncCacheStartTime))
 	{
-		return;
-	}
-	if (bIsSavingAsyncCache)
-	{
-		// we are already saving
-		return;
+		return false;
 	}
+
 	if (!IsStopped.load(std::memory_order_relaxed) && (FPlatformTime::Seconds() - SaveAsyncCacheStartTime) < FGatherSettings::MinSecondsToElapseBeforeCacheWrite)
 	{
-		// A delayed save request was made but we haven't waited long enough to start the save yet. 
+		// A delayed save request was made but we haven't waited long enough to start the save yet so we aren't ready.
 		// However if we are stopped, ignore the timer and handle the save request immediately
+		return false;
+	}
+	return true;
+}
+
+void FAssetDataGatherer::TryReserveSaveCache(bool& bOutShouldSave, TArray<TPair<FName,FDiskCachedAssetData*>>& AssetsToSave)
+{
+	bOutShouldSave = false;
+	if (!IsCacheSaveReady())
+	{
 		return;
 	}
 	CHECK_IS_LOCKED_CURRENT_THREAD(TickLock);
diff --git a/Engine/Source/Runtime/AssetRegistry/Private/AssetDataGatherer.h b/Engine/Source/Runtime/AssetRegistry/Private/AssetDataGatherer.h
--- a/Engine/Source/Runtime/AssetRegistry/Private/AssetDataGatherer.h
+++ b/Engine/Source/Runtime/AssetRegistry/Private/AssetDataGatherer.h
@@ -195,7 +195,12 @@ public:
 	* Sets OutSaveRequestTime to the time the save was requested (value may 
 	* be negative to indicate a save should occur immediately), zero otherwise.
 	*/
-	bool IsSaveRequested(double* OutSaveRequestTime = nullptr) const;
+	bool IsCacheSaveRequested(double* OutSaveRequestTime = nullptr) const;
+	/**
+	* Returns true when a save has been requested and is ready to be handled. Delayed save requests are ready once
+	* FGatherSettings::MinSecondsToElapseBeforeCacheWrite seconds have elapsed since the time of the save request. 
+	*/
+	bool IsCacheSaveReady() const;
 
 	// Controlling Async behavior
 
@@ -626,7 +631,7 @@ private:
 // Inline implementations
 ///////////////////////////////////////////////////////
 
-inline bool FAssetDataGatherer::IsSaveRequested(double* OutSaveRequestTime) const
+inline bool FAssetDataGatherer::IsCacheSaveRequested(double* OutSaveRequestTime) const
 {
 	double LocalSaveAsyncCacheStartTime = DelayedSaveAsyncCacheStartTime.load(std::memory_order_relaxed);
 	if (OutSaveRequestTime)
