AssetRegistry: Fix bug in compressed accelerator tables (UE_ASSETREGISTRY_INDIRECT_ASSETDATA_POINTERS) due to bug in TSetKeyFuncs used for them. GetTypeHash(Element) was possibly being called after MoveTemp(Element) was called, depending on compiler ordering of parameters. This caused elements to be added to the set under the wrong index, and either not be found later or cause an assertion failure.
#rnx
#Synced-CL 49917028

diff --git a/Engine/Source/Runtime/AssetRegistry/Private/SetKeyFuncs.h b/Engine/Source/Runtime/AssetRegistry/Private/SetKeyFuncs.h
--- a/Engine/Source/Runtime/AssetRegistry/Private/SetKeyFuncs.h
+++ b/Engine/Source/Runtime/AssetRegistry/Private/SetKeyFuncs.h
@@ -109,7 +109,7 @@ private:
 		return FIterationSentinel();
 	}
 
-	void AddByHashNoReallocate(uint32 TypeHash, ElementType Value, bool* bAlreadyExists);
+	void AddByHashNoReallocate(uint32 TypeHash, ElementType&& Value, bool* bAlreadyExists);
 	void Reallocate(uint32 NewHashSize);
 	uint32 GetTargetHashSize() const;
 	uint32 GetTargetHashSize(uint32 TargetNumValues) const;
@@ -437,7 +437,7 @@ inline void TSetKeyFuncs<ElementType, KeyFuncsType>::AddByHash(uint32 TypeHash,
 	}
 	check(Hash != nullptr && HashSize > 0);
 
-	AddByHashNoReallocate(TypeHash, Value, bAlreadyExists);
+	AddByHashNoReallocate(TypeHash, MoveTemp(Value), bAlreadyExists);
 
 	float LoadFactor = static_cast<float>(NumValues) / static_cast<float>(HashSize);
 	if (LoadFactor > MaxLoadFactorDuringAdd)
@@ -447,7 +447,7 @@ inline void TSetKeyFuncs<ElementType, KeyFuncsType>::AddByHash(uint32 TypeHash,
 }
 
 template <typename ElementType, typename KeyFuncsType>
-inline void TSetKeyFuncs<ElementType, KeyFuncsType>::AddByHashNoReallocate(uint32 TypeHash, ElementType Value, bool* bAlreadyExists)
+inline void TSetKeyFuncs<ElementType, KeyFuncsType>::AddByHashNoReallocate(uint32 TypeHash, ElementType&& Value, bool* bAlreadyExists)
 {
 	uint32 CollisionChainCount = 0;
 	uint32 Bucket = HashSpaceToBucketSpace(TypeHash);
@@ -619,7 +619,8 @@ inline void TSetKeyFuncs<ElementType, KeyFuncsType>::Reallocate(uint32 NewHashSi
 			ElementType& OldElement = OldHash[OldBucket];
 			if (!KeyFuncs.IsInvalid(OldElement))
 			{
-				AddByHashNoReallocate(KeyFuncs.GetTypeHash(OldElement), MoveTemp(OldElement), nullptr);
+				uint32 OldTypeHash = KeyFuncs.GetTypeHash(OldElement);
+				AddByHashNoReallocate(OldTypeHash, MoveTemp(OldElement), nullptr);
 			}
 			OldElement.~ElementType();
 		}
