[AssetRegistry]

- Avoid possible hang on unix based platforms as they attempt to save the cache during shutdown but consistently fail acquiring the system wide mutex
- Corrects the naming used for the system wide mutex as unix platforms internally rely on a valid file path. This will be corrected internally, but this name change resolves the current internal requirements

#rnx

#Synced-CL 50040740

diff --git a/Engine/Source/Runtime/AssetRegistry/Private/AssetDataGatherer.cpp b/Engine/Source/Runtime/AssetRegistry/Private/AssetDataGatherer.cpp
--- a/Engine/Source/Runtime/AssetRegistry/Private/AssetDataGatherer.cpp
+++ b/Engine/Source/Runtime/AssetRegistry/Private/AssetDataGatherer.cpp
@@ -347,7 +347,7 @@ FGatherSettings GGatherSettings;
 {
 	check(GGatherSettings.IsInitialized());
 	TStringBuilder<256> Builder;
-	Builder.Appendf(TEXT("AssetRegistryCacheLock-%s"), *GGatherSettings.GetAssetRegistryCacheRootFolder());
+	Builder.Append(*FPaths::Combine(*GGatherSettings.GetAssetRegistryCacheRootFolder(), TEXT("AssetRegistryCacheLock")));
 	return UE::FPlatformSystemWideMutex(Builder.ToString());
 }
 #endif
@@ -4211,7 +4211,7 @@ bool FAssetDataGatherer::TrySaveCacheFile(const TArray<TPair<FName,FDiskCachedAs
 	if (!CacheFilesLock.IsValid())
 	{
 		bSuccess = false;
-		if (bRequestSaveOnFailure)
+		if (bRequestSaveOnFailure && !IsStopped.load(std::memory_order_relaxed))
 		{
 			RequestAsyncCacheSave();
 		}
