AssetRegistry: Readability: Move Preloading code from AssetRegistry.cpp into a new cpp file ARLoader.cpp.
#rnx

#Synced-CL 50109798

diff --git a/Engine/Source/Runtime/AssetRegistry/Private/ARLoader.cpp b/Engine/Source/Runtime/AssetRegistry/Private/ARLoader.cpp
new file mode 100644
--- /dev/null
+++ b/Engine/Source/Runtime/AssetRegistry/Private/ARLoader.cpp
@@ -0,0 +1,380 @@
+// Copyright Epic Games, Inc. All Rights Reserved.
+
+#include "ARLoader.h"
+
+#include "AssetRegistry.h"
+#include "AssetRegistryImpl.h"
+#include "AssetRegistryPrivate.h"
+#include "Async/Async.h"
+#include "Async/TaskGraphInterfaces.h"
+#include "CoreGlobals.h"
+#include "Containers/Array.h"
+#include "Containers/UnrealString.h"
+#include "GenericPlatform/GenericPlatformFile.h"
+#include "HAL/Event.h"
+#include "HAL/FileManager.h"
+#include "HAL/LowLevelMemTracker.h"
+#include "HAL/PlatformMisc.h"
+#include "HAL/PlatformProcess.h"
+#include "HAL/PlatformProperties.h"
+#include "Misc/CommandLine.h"
+#include "Misc/CoreDelegates.h"
+#include "Misc/CoreMisc.h"
+#include "Misc/DelayedAutoRegister.h"
+#include "Misc/Parse.h"
+#include "Misc/Paths.h"
+#include "Misc/ScopeExit.h"
+#include "Misc/ScopeLock.h"
+#include "ProfilingDebugging/CpuProfilerTrace.h"
+#include "ProfilingDebugging/CsvProfiler.h"
+#include "Templates/UnrealTemplate.h"
+
+namespace UE::AssetRegistry::Premade
+{
+
+FPreloader GPreloader;
+
+bool IsEnabled()
+{
+	bool PlatformRequiresCookedData = FPlatformProperties::RequiresCookedData() && (IsRunningGame() || IsRunningDedicatedServer());
+
+#if WITH_EDITOR && !ASSETREGISTRY_FORCE_PREMADE_REGISTRY_IN_EDITOR
+	bool bUsePremadeInEditor = false;
+	if (FCommandLine::IsInitialized())
+	{
+		bUsePremadeInEditor = FParse::Param(FCommandLine::Get(), TEXT("EnablePremadeAssetRegistry"));
+	}
+#else
+	constexpr bool bUsePremadeInEditor = WITH_EDITOR;
+#endif
+
+	return PlatformRequiresCookedData || bUsePremadeInEditor;
+}
+
+bool CanLoadAsync()
+{
+	// TaskGraphSystemReady callback doesn't really mean it's running
+	return FPlatformProcess::SupportsMultithreading() && FTaskGraphInterface::IsRunning();
+}
+
+/** Returns the paths to possible Premade AssetRegistry files, ordered from highest priority to lowest. */
+static TArray<FString, TInlineAllocator<2>> GetPriorityPaths()
+{
+	TArray<FString, TInlineAllocator<2>> Paths;
+#if WITH_EDITOR
+	Paths.Add(FPaths::Combine(FPaths::ProjectDir(), TEXT("EditorClientAssetRegistry.bin")));
+#endif
+	Paths.Add(FPaths::Combine(FPaths::ProjectDir(), TEXT("AssetRegistry.bin")));
+	return Paths;
+}
+
+FPreloader::FPreloader()
+{
+	//In the editor premade Asset Registry can be enabled by a command line argument so we need to wait until the task graph is ready 
+	//before we rely on UE::AssetRegistry::Premade::IsEnabled() to return the correct result
+	bool PremadeCanBeEnabled =
+#if WITH_EDITOR
+		true;
+#else
+		UE::AssetRegistry::Premade::IsEnabled();
+#endif
+
+	if (PremadeCanBeEnabled)
+	{
+		// run DelayedInitialize when TaskGraph system is ready
+		OnTaskGraphReady.Emplace(STATS ? EDelayedRegisterRunPhase::StatSystemReady :
+			EDelayedRegisterRunPhase::TaskGraphSystemReady,
+			[this]()
+			{
+				if (UE::AssetRegistry::Premade::IsEnabled())
+				{
+					LoadState = EState::NotFound;
+					DelayedInitialize();
+				}
+			});
+	}
+}
+
+FPreloader::~FPreloader()
+{
+	// We are destructed after Main exits, which means that our AsyncThread was either never called
+	// or it was waited on to complete by TaskGraph. Therefore we do not need to handle waiting for it ourselves.
+	Shutdown(true /* bFromGlobalDestructor */);
+}
+
+bool FPreloader::Consume(FConsumeFunction&& ConsumeFunction)
+{
+	EConsumeResult Result = ConsumeInternal(MoveTemp(ConsumeFunction), FConsumeFunction());
+	check(Result != EConsumeResult::Deferred);
+	return Result == EConsumeResult::Succeeded;
+}
+
+FPreloader::EConsumeResult FPreloader::ConsumeOrDefer(FConsumeFunction&& ConsumeSynchronous,
+	FConsumeFunction&& ConsumeAsynchronous)
+{
+	return ConsumeInternal(MoveTemp(ConsumeSynchronous), MoveTemp(ConsumeAsynchronous));
+}
+
+bool FPreloader::TrySetPath()
+{
+	for (FString& LocalPath : GetPriorityPaths())
+	{
+		if (IFileManager::Get().FileExists(*LocalPath))
+		{
+			ARPath = MoveTemp(LocalPath);
+			return true;
+		}
+	}
+	return false;
+}
+
+bool FPreloader::TrySetPath(const IPakFile& Pak)
+{
+	for (FString& LocalPath : GetPriorityPaths())
+	{
+		if (Pak.PakContains(LocalPath))
+		{
+			ARPath = MoveTemp(LocalPath);
+			return true;
+		}
+	}
+	return false;
+}
+
+ELoadResult FPreloader::TryLoad()
+{
+	TRACE_CPUPROFILER_EVENT_SCOPE(FCookedAssetRegistryPreloader::TryLoad);
+	LLM_SCOPE(ELLMTag::AssetRegistry);
+	checkf(!ARPath.IsEmpty(), TEXT("TryLoad must not be called until after TrySetPath has succeeded."));
+
+	FAssetRegistryLoadOptions Options;
+	const int32 ThreadReduction = 2; // This thread + main thread already has work to do 
+	int32 MaxWorkers = CanLoadAsync() ? FPlatformMisc::NumberOfCoresIncludingHyperthreads() - ThreadReduction : 0;
+	Options.ParallelWorkers = FMath::Clamp(MaxWorkers, 0, 16);
+	bool bLoadSucceeded = FAssetRegistryState::LoadFromDisk(*ARPath, Options, Payload);
+	UE_CLOG(!bLoadSucceeded, LogAssetRegistry, Warning, TEXT("Premade AssetRegistry path %s existed but failed to load."), *ARPath);
+	UE_CLOG(bLoadSucceeded, LogAssetRegistry, Log, TEXT("Premade AssetRegistry loaded from '%s'"), *ARPath);
+	LoadResult = bLoadSucceeded ? ELoadResult::Succeeded : ELoadResult::FailedToLoad;
+	return LoadResult;
+}
+
+void FPreloader::DelayedInitialize()
+{
+	TRACE_CPUPROFILER_EVENT_SCOPE(FCookedAssetRegistryPreloader::DelayedInitialize);
+	// This function will run before any UObject (ie UAssetRegistryImpl) code can run, so we don't need to do any thread safety
+	// CanLoadAsync - we have to check this after the task graph is ready
+	if (!CanLoadAsync())
+	{
+		LoadState = EState::LoadSynchronous;
+		return;
+	}
+
+	// PreloadReady is in Triggered state until the Async thread is created. It is Reset in KickPreload.
+	PreloadReady = FPlatformProcess::GetSynchEventFromPool(true /* bIsManualReset */);
+	PreloadReady->Trigger();
+
+	if (TrySetPath())
+	{
+		KickPreload();
+	}
+	else
+	{
+		// set to NotFound, although PakMounted may set it to found later
+		LoadState = EState::NotFound;
+
+		// The PAK with the main registry isn't mounted yet
+		PakMountedDelegate = FCoreDelegates::GetOnPakFileMounted2().AddLambda([this](const IPakFile& Pak)
+			{
+				FScopeLock Lock(&StateLock);
+				if (LoadState == EState::NotFound && TrySetPath(Pak))
+				{
+					KickPreload();
+					// Remove the callback from OnPakFileMounted2 to avoid wasting time in all future PakFile mounts
+					// Do not access any of the lambda captures after the call to Remove, because deallocating the 
+					// DelegateHandle also deallocates our lambda captures
+					FDelegateHandle LocalPakMountedDelegate = PakMountedDelegate;
+					PakMountedDelegate.Reset();
+					FCoreDelegates::GetOnPakFileMounted2().Remove(LocalPakMountedDelegate);
+				}
+			});
+	}
+}
+
+void FPreloader::KickPreload()
+{
+	CSV_TRACE_REGION_BEGIN(TEXT("Preload"), TEXT("AssetRegistry"))
+	TRACE_CPUPROFILER_EVENT_SCOPE(FCookedAssetRegistryPreloader::KickPreload);
+	// Called from Within the Lock
+	check(LoadState == EState::NotFound && !ARPath.IsEmpty());
+	LoadState = EState::Loading;
+	PreloadReady->Reset();
+	Async(EAsyncExecution::TaskGraph, [this]() { TryLoadAsync(); });
+}
+
+void FPreloader::TryLoadAsync()
+{
+	// This function is active only after State has been set to Loading and PreloadReady has been Reset
+	// Until this function triggers PreloadReady, it has exclusive ownership of bLoadSucceeded and Payload
+	// Load outside the lock so that ConsumeOrDefer does not have to wait for the Load before it can defer and exit
+	ELoadResult LocalResult = TryLoad();
+	// Trigger outside the lock so that a locked Consume function that is waiting on PreloadReady can wait inside the lock.
+	PreloadReady->Trigger();
+
+	FConsumeFunction LocalConsumeCallback;
+	{
+		FScopeLock Lock(&StateLock);
+		// The consume function may have woken up after the trigger and already consumed and changed LoadState to Consumed
+		if (LoadState == EState::Loading)
+		{
+			LoadState = EState::Loaded;
+			if (ConsumeCallback)
+			{
+				LocalConsumeCallback = MoveTemp(ConsumeCallback);
+				ConsumeCallback.Reset();
+				LoadState = EState::Consumed;
+			}
+		}
+	}
+
+	if (LocalConsumeCallback)
+	{
+		// No further threads will read/write payload at this point (until destructor, which is called after all async threads are complete
+		// so we can use it outside the lock
+		LocalConsumeCallback(LocalResult, MoveTemp(Payload));
+		Shutdown();
+	}
+}
+
+FPreloader::EConsumeResult FPreloader::ConsumeInternal(FConsumeFunction&& ConsumeSynchronous, 
+	FConsumeFunction&& ConsumeAsynchronous)
+{
+	SCOPED_BOOT_TIMING("FCookedAssetRegistryPreloader::Consume");
+
+	FScopeLock Lock(&StateLock);
+	// Report failure if constructor decided not to preload or this has already been Consumed
+	if (LoadState == EState::WillNeverPreload || LoadState == EState::Consumed || ConsumeCallback)
+	{
+		Lock.Unlock(); // Unlock before calling external code in Consume callback
+		ELoadResult LocalResult = (LoadState == EState::Consumed || ConsumeCallback) ? ELoadResult::AlreadyConsumed : ELoadResult::Inactive;
+		ConsumeSynchronous(LocalResult, FAssetRegistryState());
+		return EConsumeResult::Failed;
+	}
+
+	if (LoadState == EState::LoadSynchronous)
+	{
+		ELoadResult LocalResult = TrySetPath() ? TryLoad() : ELoadResult::NotFound;
+		LoadState = EState::Consumed;
+		Lock.Unlock(); // Unlock before calling external code in Consume callback
+		ConsumeSynchronous(LocalResult, MoveTemp(Payload));
+		Shutdown(); // Shutdown can be called outside the lock since AsyncThread doesn't exist
+		return LocalResult == ELoadResult::Succeeded ? EConsumeResult::Succeeded : EConsumeResult::Failed;
+	}
+
+	// Cancel any further searching in Paks since we will no longer accept preloads starting after this point
+	FCoreDelegates::GetOnPakFileMounted2().Remove(PakMountedDelegate);
+	PakMountedDelegate.Reset();
+
+	if (ConsumeAsynchronous && LoadState == EState::Loading)
+	{
+		// The load might have completed and the TryAsyncLoad thread is waiting to enter the lock, but we will still defer since Consume won the race
+		ConsumeCallback = MoveTemp(ConsumeAsynchronous);
+		return EConsumeResult::Deferred;
+	}
+
+	{
+		SCOPED_BOOT_TIMING("BlockingConsume");
+		// If the load is in progress, wait for it to finish (which it does outside the lock)
+		PreloadReady->Wait();
+	}
+
+	// TryAsyncLoad might not yet have set state to Loaded
+	check(LoadState == EState::Loaded || LoadState == EState::Loading || LoadState == EState::NotFound);
+	ELoadResult LocalResult = LoadState == EState::NotFound ? ELoadResult::NotFound : LoadResult;
+	LoadState = EState::Consumed;
+
+	// No further async threads exist that will read/write payload at this point so we can use it outside the lock
+	Lock.Unlock(); // Unlock before calling external code in Consume callback
+	ConsumeSynchronous(LocalResult, MoveTemp(Payload));
+	Shutdown(); // Shutdown can be called outside the lock since we have set state to Consumed and the Async thread will notice and exit
+	return LocalResult == ELoadResult::Succeeded ? EConsumeResult::Succeeded : EConsumeResult::Failed;
+}
+
+void FPreloader::Shutdown(bool bFromGlobalDestructor)
+{
+	OnTaskGraphReady.Reset();
+	if (PreloadReady)
+	{
+		// If we are exiting the process early while PreloadReady is still allocated, the event
+		// system has already been torn down and there is nothing for us to free for PreloadReady.
+		if (!bFromGlobalDestructor)
+		{
+			FPlatformProcess::ReturnSynchEventToPool(PreloadReady);
+		}
+		PreloadReady = nullptr;
+	}
+	ARPath.Reset();
+	Payload.Reset();
+	CSV_TRACE_REGION_END(TEXT("Preload"), TEXT("AssetRegistry"))
+}
+
+FAsyncConsumer::~FAsyncConsumer()
+{
+	if (Consumed)
+	{
+		FPlatformProcess::ReturnSynchEventToPool(Consumed);
+		Consumed = nullptr;
+	}
+}
+
+void FAsyncConsumer::PrepareForConsume()
+{
+	// Called within the lock
+	check(!Consumed);
+	Consumed = FPlatformProcess::GetSynchEventFromPool(true /* bIsManualReset */);
+	++ReferenceCount;
+};
+
+void FAsyncConsumer::Wait(UE::AssetRegistry::FInterfaceWriteScopeLock& ScopeLock)
+{
+	// Called within the lock
+	if (ReferenceCount == 0)
+	{
+		return;
+	}
+	++ReferenceCount;
+
+	// Wait outside of the lock so that the AsyncThread can enter the lock to call Consume
+	{
+		ScopeLock.Lock.WriteUnlock();
+		ON_SCOPE_EXIT{ ScopeLock.Lock.WriteLock(); };
+		check(Consumed != nullptr);
+		Consumed->Wait();
+	}
+
+	--ReferenceCount;
+	if (ReferenceCount == 0)
+	{
+		// We're the last one to drop the refcount, so delete Consumed
+		check(Consumed != nullptr);
+		FPlatformProcess::ReturnSynchEventToPool(Consumed);
+		Consumed = nullptr;
+	}
+}
+
+void FAsyncConsumer::Consume(UAssetRegistryImpl& UARI, UE::AssetRegistry::Impl::FEventContext& EventContext, ELoadResult LoadResult, FAssetRegistryState&& ARState)
+{
+	// Called within the lock
+	UARI.GuardedData.LoadPremadeAssetRegistry(EventContext, LoadResult, MoveTemp(ARState));
+	check(ReferenceCount >= 1);
+	check(Consumed != nullptr);
+	Consumed->Trigger();
+	--ReferenceCount;
+	if (ReferenceCount == 0)
+	{
+		// We're the last one to drop the refcount, so delete Consumed
+		FPlatformProcess::ReturnSynchEventToPool(Consumed);
+		Consumed = nullptr;
+	}
+}
+
+} // namespace UE::AssetRegistry::Premade
diff --git a/Engine/Source/Runtime/AssetRegistry/Private/ARLoader.h b/Engine/Source/Runtime/AssetRegistry/Private/ARLoader.h
new file mode 100644
--- /dev/null
+++ b/Engine/Source/Runtime/AssetRegistry/Private/ARLoader.h
@@ -0,0 +1,158 @@
+// Copyright Epic Games, Inc. All Rights Reserved.
+
+#pragma once
+
+#include "ARTypes.h"
+#include "AssetRegistry/AssetRegistryState.h"
+#include "Containers/UnrealString.h"
+#include "Delegates/Delegate.h"
+#include "HAL/CriticalSection.h"
+#include "Misc/Optional.h"
+
+// Preloading a premade AssetRegistry is disabled in editor unless the target specifies it in UBT, because the editor
+// usually gathers to create an up to date AssetRegistry rather than using a previously cooked AssetRegistry.bin.
+#ifndef ASSETREGISTRY_FORCE_PREMADE_REGISTRY_IN_EDITOR
+#define ASSETREGISTRY_FORCE_PREMADE_REGISTRY_IN_EDITOR 0
+#endif
+
+class FEvent;
+class IPakFile;
+class UAssetRegistryImpl;
+namespace UE::AssetRegistry::Impl { struct FEventContext; }
+struct FDelayedAutoRegisterHelper;
+
+namespace UE::AssetRegistry::Premade
+{
+
+enum class ELoadResult : uint8
+{
+	Succeeded = 0,
+	NotFound = 1,
+	FailedToLoad = 2,
+	Inactive = 3,
+	AlreadyConsumed = 4,
+	UninitializedMemberLoadResult = 5,
+};
+
+/** Loads AssetRegistry.bin during engine startup, using an async preload task if available and sync otherwise. */
+class FPreloader
+{
+public:
+	enum class EConsumeResult
+	{
+		Succeeded,
+		Failed,
+		Deferred
+	};
+	using FConsumeFunction = TFunction<void(ELoadResult LoadResult, FAssetRegistryState&& ARState)>;
+
+	FPreloader();
+	~FPreloader();
+
+	/**
+	 * Block on any pending async load, load if synchronous, and call ConsumeFunction with the results before returning.
+	 * If Consume has been called previously, the current ConsumeFunction is ignored and this call returns false.
+	 *
+	 * @return Whether the load succeeded (this information is also passed to the ConsumeFunction).
+	 */
+	bool Consume(FConsumeFunction&& ConsumeFunction);
+
+	/**
+	 * If a load is pending, store ConsumeAsynchronous for later calling and return EConsumeResult::Deferred.
+	 * If load is complete, or failed, or needs to run synchronously, load if necessary and call ConsumeSynchronous with results before returning.
+	 * Note if this function returns EConsumeResult::Deferred, the ConsumeAsynchronous will be called from another thread,
+	 * possibly before this call returns.
+	 * If Consume has been called previously, this call is ignored and returns EConsumeResult::Failed.
+	 *
+	 * @return Whether the load succeeded (this information is also passed to the ConsumeFunction).
+	 */
+	EConsumeResult ConsumeOrDefer(FConsumeFunction&& ConsumeSynchronous, FConsumeFunction&& ConsumeAsynchronous);
+
+private:
+	enum class EState : uint8
+	{
+		WillNeverPreload,
+		LoadSynchronous,
+		NotFound,
+		Loading,
+		Loaded,
+		Consumed,
+	};
+
+
+	bool TrySetPath();
+	bool TrySetPath(const IPakFile& Pak);
+	ELoadResult TryLoad();
+	void DelayedInitialize();
+	void KickPreload();
+	void TryLoadAsync();
+	EConsumeResult ConsumeInternal(FConsumeFunction&& ConsumeSynchronous, FConsumeFunction&& ConsumeAsynchronous);
+	/** Called when the Preloader has no further work to do, to free resources early since destruction occurs at end of process. */
+	void Shutdown(bool bFromGlobalDestructor = false);
+
+	/** simple way to trigger a callback at a specific time that TaskGraph is usable. */
+	TOptional<FDelayedAutoRegisterHelper> OnTaskGraphReady;
+
+	/** Lock that guards members on this (see notes on each member). */
+	FCriticalSection StateLock;
+	/** Trigger for blocking Consume to wait upon TryLoadAsync. This Trigger is only allocated when in the states NotFound, Loaded, Loading. */
+	FEvent* PreloadReady = nullptr;
+
+	/** Path discovered for the AssetRegistry; Read/Write only within the Lock. */
+	FString ARPath;
+
+	/**
+	 * The ARState loaded from disk. Owned exclusively by either the first Consume or by TryAsyncLoad.
+	 * If LoadState is never set to Loading, this state is read/written only by the first thread to call Consume.
+	 * If LoadState is set to Loading (which happens before threading starts), the thread running TryAsyncLoad
+	 * owns this payload until it triggers PayloadReady, after which ownership returns to the first thread to call Consume.
+	 */
+	FAssetRegistryState Payload;
+
+	FDelegateHandle PakMountedDelegate;
+
+	/** Callback from ConsumeOrDefer that is set so TryLoadAsync can trigger the Consume when it completes.Read / Write only within the lock. */
+	FConsumeFunction ConsumeCallback;
+
+	/** State machine state. Read/Write only within the lock (or before threading starts). */
+	EState LoadState = EState::WillNeverPreload;
+
+	/** Result of TryLoad.Thread ownership rules are the same as the rules for Payload. */
+	ELoadResult LoadResult = ELoadResult::UninitializedMemberLoadResult;
+};
+
+/** Returns whether the given executable configuration supports AssetRegistry Preloading. Called before Main. */
+extern bool IsEnabled();
+extern bool CanLoadAsync();
+
+extern FPreloader GPreloader;
+
+/**
+* A struct to consume a Premade asset registry on an async thread. It supports a cheap Wait() call so that it can be Waited on
+* by frequent AssetRegistry interface calls for the rest of the process.
+*/
+struct FAsyncConsumer
+{
+	~FAsyncConsumer();
+	/** Sets the Consumer into need-to-wait mode; Wait will block until Consume is called. */
+	void PrepareForConsume();
+	/**
+	* Does not return until after the Premade AssetRegistryState has been added into the target AssetRegistry, or it has
+	* been decided not to be used. For performance reasons, this must be called within the WriteScopeLock, and the caller
+	* must handle the possibility that it leaves and reenters the lock.
+	*/
+	void Wait(FInterfaceWriteScopeLock& ScopeLock);
+	/** Callback from the async thread to consume the Premade ARState */
+	void Consume(UAssetRegistryImpl& UARI, UE::AssetRegistry::Impl::FEventContext& EventContext,
+		ELoadResult LoadResult, FAssetRegistryState&& ARState);
+
+private:
+	/**
+	* ReferenceCounter for the Consumed event. Also used to decide whether Waiting is necessary. Read/Write only inside the lock.
+	*/
+	int32 ReferenceCount = 0;
+	/** Event used to Wait for Consume. Allocated/deallocated within the lock. Can be waited on outside the lock if RefCount is held. */
+	FEvent* Consumed = nullptr;
+};
+
+} // namespace UE::AssetRegistry::Premade
diff --git a/Engine/Source/Runtime/AssetRegistry/Private/ARTypes.h b/Engine/Source/Runtime/AssetRegistry/Private/ARTypes.h
new file mode 100644
--- /dev/null
+++ b/Engine/Source/Runtime/AssetRegistry/Private/ARTypes.h
@@ -0,0 +1,177 @@
+// Copyright Epic Games, Inc. All Rights Reserved.
+
+#pragma once
+
+#include "AutoRTFM.h"
+#include "HAL/Platform.h"
+#include "Misc/Optional.h"
+#include "Misc/ScopeRWLock.h"
+#include "Misc/TransactionallySafeRWLock.h"
+
+#include <atomic>
+
+namespace UE::AssetRegistry::Private { class FRWLockWithPriority; }
+namespace UE::AssetRegistry::Private { using FInterfaceRWLock = FRWLockWithPriority; }
+namespace UE::AssetRegistry { template <typename TScopeLockType> class TRWScopeLockWithPriority; }
+namespace UE::AssetRegistry { using FInterfaceReadScopeLock = TRWScopeLockWithPriority<UE::TReadScopeLock<Private::FRWLockWithPriority>>; }
+namespace UE::AssetRegistry { using FInterfaceWriteScopeLock = TRWScopeLockWithPriority<UE::TWriteScopeLock<Private::FRWLockWithPriority>>; }
+namespace UE::AssetRegistry { using FInterfaceRWScopeLock = class FRWScopeLockWithPriority; }
+
+namespace UE::AssetRegistry::Private
+{
+
+/**
+ * Custom lock class that can prioritize waiters upon request. This is used to allow threads with critical requests
+ * to request that the gatherer thread pause its work and allow the higher pri threads to jump in. Note that for this
+ * to work correctly it must be used with the associated TPriorityScopeLock types
+ * (TScopeLockWithPriority and FRWScopeLockWithPriority)
+*/
+class FRWLockWithPriority
+{
+public:
+	FORCEINLINE void WriteLock()
+	{
+		Mutex.WriteLock();
+	}
+
+	FORCEINLINE void WriteUnlock()
+	{
+		Mutex.WriteUnlock();
+	}
+
+	FORCEINLINE void ReadLock()
+	{
+		Mutex.ReadLock();
+	}
+
+	FORCEINLINE void ReadUnlock()
+	{
+		Mutex.ReadUnlock();
+	}
+
+	bool HasWaiters() const
+	{
+		bool Result;
+		UE_AUTORTFM_OPEN
+		{
+			Result = HighPriorityWaitersCount.load(std::memory_order_relaxed) > 0;
+		};
+		return Result;
+	}
+
+private:
+	friend UE::AssetRegistry::FInterfaceReadScopeLock;
+	friend UE::AssetRegistry::FInterfaceWriteScopeLock;
+	friend UE::AssetRegistry::FInterfaceRWScopeLock;
+	FTransactionallySafeRWLock Mutex;
+	std::atomic<int32> HighPriorityWaitersCount { 0 };
+};
+
+enum ELockPriority : uint8
+{
+	PriorityLow,
+	PriorityHigh
+};
+
+} // namespace UE::AssetRegistry::Private
+
+namespace UE::AssetRegistry
+{
+
+/** Tracks high-priority waiters on a lock, in an thread- and transactionally-safe manner. */
+class FScopedHighPriorityWaitTracker
+{
+public:
+	FScopedHighPriorityWaitTracker(std::atomic<int32>* InCounter, Private::ELockPriority InPriority)
+		: Counter((InPriority == Private::PriorityHigh) ? InCounter : nullptr)
+	{
+		if (Counter)
+		{
+			UE_AUTORTFM_OPEN
+			{
+				Counter->fetch_add(1, std::memory_order_relaxed);
+			};
+			AutoRTFM::PushOnAbortHandler(this, [this]
+				{
+					Counter->fetch_sub(1, std::memory_order_relaxed);
+				});
+		}
+	}
+
+	~FScopedHighPriorityWaitTracker()
+	{
+		if (Counter)
+		{
+			UE_AUTORTFM_OPEN
+			{
+				Counter->fetch_sub(1, std::memory_order_relaxed);
+			};
+			AutoRTFM::PopOnAbortHandler(this);
+		}
+	}
+
+private:
+	std::atomic<int32>* Counter = nullptr;
+};
+
+/**
+ * Keeps an FRWLock read-locked while this scope lives.
+ * This is almost a clone of the existing FReadScopeLock and similar types
+ * however this adds an extra flag to help the background processing thread
+ * know when a higher priority thread would like to gain access to the protected
+ * data
+ */
+template <typename TScopeLockType> class TRWScopeLockWithPriority
+{
+public:
+	UE_NODISCARD_CTOR explicit TRWScopeLockWithPriority(Private::FRWLockWithPriority& InLock,
+		Private::ELockPriority InPriority = Private::PriorityHigh)
+		: Lock(InLock)
+		, Priority(InPriority)
+	{
+		FScopedHighPriorityWaitTracker Tracker(&Lock.HighPriorityWaitersCount, Priority);
+		GuardWrapper.Emplace(Lock);
+	}
+
+	TOptional<TScopeLockType> GuardWrapper;
+	UE::AssetRegistry::Private::FRWLockWithPriority& Lock;
+	UE::AssetRegistry::Private::ELockPriority Priority;
+};
+
+class FRWScopeLockWithPriority
+{
+public:
+	UE_NODISCARD_CTOR explicit FRWScopeLockWithPriority(Private::FRWLockWithPriority& InLockObject,
+		FRWScopeLockType InLockType, Private::ELockPriority InPriority = Private::PriorityHigh)
+		: Lock(InLockObject)
+		, Priority(InPriority)
+		, LockType(InLockType)
+	{
+		FScopedHighPriorityWaitTracker Tracker(&Lock.HighPriorityWaitersCount, Priority);
+		GuardWrapper.Emplace(Lock, LockType);
+	}
+
+	// NOTE: As the name suggests, this function should be used with caution. 
+	// It releases the read lock _before_ acquiring a new write lock. This is not an atomic operation and the caller should 
+	// not treat it as such. 
+	// E.g. Pointers read from protected data structures prior to this call may be invalid after the function is called. 
+	void ReleaseReadOnlyLockAndAcquireWriteLock_USE_WITH_CAUTION()
+	{
+		if (LockType == SLT_ReadOnly)
+		{
+			GuardWrapper.Reset();
+			{
+				FScopedHighPriorityWaitTracker Tracker(&Lock.HighPriorityWaitersCount, Priority);
+				GuardWrapper.Emplace(Lock, SLT_Write);
+			}
+			LockType = SLT_Write;
+		}
+	}
+
+	Private::FRWLockWithPriority& Lock;
+	TOptional<TRWScopeLock<Private::FRWLockWithPriority>> GuardWrapper;
+	Private::ELockPriority Priority;
+	FRWScopeLockType LockType;
+};
+
+} // namespace UE::AssetRegistry
\ No newline at end of file
diff --git a/Engine/Source/Runtime/AssetRegistry/Private/AssetRegistry.cpp b/Engine/Source/Runtime/AssetRegistry/Private/AssetRegistry.cpp
--- a/Engine/Source/Runtime/AssetRegistry/Private/AssetRegistry.cpp
+++ b/Engine/Source/Runtime/AssetRegistry/Private/AssetRegistry.cpp
@@ -329,434 +329,6 @@ FAssetRegistryInterface GAssetRegistryInterface;
 DEFINE_LOG_CATEGORY(LogAssetRegistry);
 
 
-namespace UE::AssetRegistry::Premade
-{
-
-/** Returns whether the given executable configuration supports AssetRegistry Preloading. Called before Main. */
-static bool IsEnabled()
-{
-	bool PlatformRequiresCookedData = FPlatformProperties::RequiresCookedData() && (IsRunningGame() || IsRunningDedicatedServer());
-
-#if WITH_EDITOR && !ASSETREGISTRY_FORCE_PREMADE_REGISTRY_IN_EDITOR
-	bool bUsePremadeInEditor = false;
-	if (FCommandLine::IsInitialized())
-	{
-		bUsePremadeInEditor = FParse::Param(FCommandLine::Get(), TEXT("EnablePremadeAssetRegistry"));
-	}
-#else
-	constexpr bool bUsePremadeInEditor = WITH_EDITOR;
-#endif
-
-	return PlatformRequiresCookedData || bUsePremadeInEditor;
-}
-
-static bool CanLoadAsync()
-{
-	// TaskGraphSystemReady callback doesn't really mean it's running
-	return FPlatformProcess::SupportsMultithreading() && FTaskGraphInterface::IsRunning();
-}
-
-/** Returns the paths to possible Premade AssetRegistry files, ordered from highest priority to lowest. */
-TArray<FString, TInlineAllocator<2>> GetPriorityPaths()
-{
-	TArray<FString, TInlineAllocator<2>> Paths;
-#if WITH_EDITOR
-	Paths.Add(FPaths::Combine(FPaths::ProjectDir(), TEXT("EditorClientAssetRegistry.bin")));
-#endif
-	Paths.Add(FPaths::Combine(FPaths::ProjectDir(), TEXT("AssetRegistry.bin")));
-	return Paths;
-}
-
-enum class ELoadResult : uint8
-{
-	Succeeded = 0,
-	NotFound = 1,
-	FailedToLoad = 2,
-	Inactive = 3,
-	AlreadyConsumed = 4,
-	UninitializedMemberLoadResult = 5,
-};
-
-// Loads cooked AssetRegistry.bin using an async preload task if available and sync otherwise
-class FPreloader
-{
-public:
-	enum class EConsumeResult
-	{
-		Succeeded,
-		Failed,
-		Deferred
-	};
-	using FConsumeFunction = TFunction<void(ELoadResult LoadResult, FAssetRegistryState&& ARState)>;
-
-	FPreloader()
-	{
-		//In the editor premade Asset Registry can be enabled by a command line argument so we need to wait until the task graph is ready 
-		//before we rely on UE::AssetRegistry::Premade::IsEnabled() to return the correct result
-		bool PremadeCanBeEnabled =
-#if WITH_EDITOR
-			true;
-#else
-			UE::AssetRegistry::Premade::IsEnabled();
-#endif
-
-		if (PremadeCanBeEnabled)
-		{
-			// run DelayedInitialize when TaskGraph system is ready
-			OnTaskGraphReady.Emplace(STATS ? EDelayedRegisterRunPhase::StatSystemReady :
-				EDelayedRegisterRunPhase::TaskGraphSystemReady,
-				[this]()
-				{
-					if (UE::AssetRegistry::Premade::IsEnabled())
-					{
-						LoadState = EState::NotFound;
-						DelayedInitialize();
-					}
-				});
-		}
-	}
-	~FPreloader()
-	{
-		// We are destructed after Main exits, which means that our AsyncThread was either never called
-		// or it was waited on to complete by TaskGraph. Therefore we do not need to handle waiting for it ourselves.
-		Shutdown(true /* bFromGlobalDestructor */);
-	}
-
-	/**
-	 * Block on any pending async load, load if synchronous, and call ConsumeFunction with the results before returning.
-	 * If Consume has been called previously, the current ConsumeFunction is ignored and this call returns false.
-	 * 
-	 * @return Whether the load succeeded (this information is also passed to the ConsumeFunction).
-	 */
-	bool Consume(FConsumeFunction&& ConsumeFunction)
-	{
-		EConsumeResult Result = ConsumeInternal(MoveTemp(ConsumeFunction), FConsumeFunction());
-		check(Result != EConsumeResult::Deferred);
-		return Result == EConsumeResult::Succeeded;
-	}
-
-	/**
-	 * If a load is pending, store ConsumeAsynchronous for later calling and return EConsumeResult::Deferred.
-	 * If load is complete, or failed, or needs to run synchronously, load if necessary and call ConsumeSynchronous with results before returning.
-	 * Note if this function returns EConsumeResult::Deferred, the ConsumeAsynchronous will be called from another thread,
-	 * possibly before this call returns.
-	 * If Consume has been called previously, this call is ignored and returns EConsumeResult::Failed.
-	 *
-	 * @return Whether the load succeeded (this information is also passed to the ConsumeFunction).
-	 */
-	EConsumeResult ConsumeOrDefer(FConsumeFunction&& ConsumeSynchronous, FConsumeFunction&& ConsumeAsynchronous)
-	{
-		return ConsumeInternal(MoveTemp(ConsumeSynchronous), MoveTemp(ConsumeAsynchronous));
-	}
-
-private:
-	enum class EState : uint8
-	{
-		WillNeverPreload,
-		LoadSynchronous,
-		NotFound,
-		Loading,
-		Loaded,
-		Consumed,
-	};
-
-
-	bool TrySetPath()
-	{
-		for (FString& LocalPath : GetPriorityPaths())
-		{
-			if (IFileManager::Get().FileExists(*LocalPath))
-			{
-				ARPath = MoveTemp(LocalPath);
-				return true;
-			}
-		}
-		return false;
-	}
-
-	bool TrySetPath(const IPakFile& Pak)
-	{
-		for (FString& LocalPath : GetPriorityPaths())
-		{
-			if (Pak.PakContains(LocalPath))
-			{
-				ARPath = MoveTemp(LocalPath);
-				return true;
-			}
-		}
-		return false;
-	}
-
-	ELoadResult TryLoad()
-	{
-		TRACE_CPUPROFILER_EVENT_SCOPE(FCookedAssetRegistryPreloader::TryLoad);
-		LLM_SCOPE(ELLMTag::AssetRegistry);
-		checkf(!ARPath.IsEmpty(), TEXT("TryLoad must not be called until after TrySetPath has succeeded."));
-
-		FAssetRegistryLoadOptions Options;
-		const int32 ThreadReduction = 2; // This thread + main thread already has work to do 
-		int32 MaxWorkers = CanLoadAsync() ? FPlatformMisc::NumberOfCoresIncludingHyperthreads() - ThreadReduction : 0;
-		Options.ParallelWorkers = FMath::Clamp(MaxWorkers, 0, 16);
-		bool bLoadSucceeded = FAssetRegistryState::LoadFromDisk(*ARPath, Options, Payload);
-		UE_CLOG(!bLoadSucceeded, LogAssetRegistry, Warning, TEXT("Premade AssetRegistry path %s existed but failed to load."), *ARPath);
-		UE_CLOG(bLoadSucceeded, LogAssetRegistry, Log, TEXT("Premade AssetRegistry loaded from '%s'"), *ARPath);
-		LoadResult = bLoadSucceeded ? ELoadResult::Succeeded : ELoadResult::FailedToLoad;
-		return LoadResult;
-	}
-
-	void DelayedInitialize()
-	{
-		TRACE_CPUPROFILER_EVENT_SCOPE(FCookedAssetRegistryPreloader::DelayedInitialize);
-		// This function will run before any UObject (ie UAssetRegistryImpl) code can run, so we don't need to do any thread safety
-		// CanLoadAsync - we have to check this after the task graph is ready
-		if (!CanLoadAsync())
-		{
-			LoadState = EState::LoadSynchronous;
-			return;
-		}
-
-		// PreloadReady is in Triggered state until the Async thread is created. It is Reset in KickPreload.
-		PreloadReady = FPlatformProcess::GetSynchEventFromPool(true /* bIsManualReset */);
-		PreloadReady->Trigger();
-
-		if (TrySetPath())
-		{
-			KickPreload();
-		}
-		else
-		{
-			// set to NotFound, although PakMounted may set it to found later
-			LoadState = EState::NotFound;
-
-			// The PAK with the main registry isn't mounted yet
-			PakMountedDelegate = FCoreDelegates::GetOnPakFileMounted2().AddLambda([this](const IPakFile& Pak)
-				{
-					FScopeLock Lock(&StateLock);
-					if (LoadState == EState::NotFound && TrySetPath(Pak))
-					{
-						KickPreload();
-						// Remove the callback from OnPakFileMounted2 to avoid wasting time in all future PakFile mounts
-						// Do not access any of the lambda captures after the call to Remove, because deallocating the 
-						// DelegateHandle also deallocates our lambda captures
-						FDelegateHandle LocalPakMountedDelegate = PakMountedDelegate;
-						PakMountedDelegate.Reset();
-						FCoreDelegates::GetOnPakFileMounted2().Remove(LocalPakMountedDelegate);
-					}
-				});
-		}
-	}
-
-	void KickPreload()
-	{
-		CSV_TRACE_REGION_BEGIN(TEXT("Preload"), TEXT("AssetRegistry"))
-		TRACE_CPUPROFILER_EVENT_SCOPE(FCookedAssetRegistryPreloader::KickPreload);
-		// Called from Within the Lock
-		check(LoadState == EState::NotFound && !ARPath.IsEmpty());
-		LoadState = EState::Loading;
-		PreloadReady->Reset();
-		Async(EAsyncExecution::TaskGraph, [this]() { TryLoadAsync(); });
-	}
-
-	void TryLoadAsync()
-	{
-		// This function is active only after State has been set to Loading and PreloadReady has been Reset
-		// Until this function triggers PreloadReady, it has exclusive ownership of bLoadSucceeded and Payload
-		// Load outside the lock so that ConsumeOrDefer does not have to wait for the Load before it can defer and exit
-		ELoadResult LocalResult = TryLoad();
-		// Trigger outside the lock so that a locked Consume function that is waiting on PreloadReady can wait inside the lock.
-		PreloadReady->Trigger();
-
-		FConsumeFunction LocalConsumeCallback;
-		{
-			FScopeLock Lock(&StateLock);
-			// The consume function may have woken up after the trigger and already consumed and changed LoadState to Consumed
-			if (LoadState == EState::Loading)
-			{
-				LoadState = EState::Loaded;
-				if (ConsumeCallback)
-				{
-					LocalConsumeCallback = MoveTemp(ConsumeCallback);
-					ConsumeCallback.Reset();
-					LoadState = EState::Consumed;
-				}
-			}
-		}
-
-		if (LocalConsumeCallback)
-		{
-			// No further threads will read/write payload at this point (until destructor, which is called after all async threads are complete
-			// so we can use it outside the lock
-			LocalConsumeCallback(LocalResult, MoveTemp(Payload));
-			Shutdown();
-		}
-	}
-
-	EConsumeResult ConsumeInternal(FConsumeFunction&& ConsumeSynchronous, FConsumeFunction&& ConsumeAsynchronous)
-	{
-		SCOPED_BOOT_TIMING("FCookedAssetRegistryPreloader::Consume");
-
-		FScopeLock Lock(&StateLock);
-		// Report failure if constructor decided not to preload or this has already been Consumed
-		if (LoadState == EState::WillNeverPreload || LoadState == EState::Consumed || ConsumeCallback)
-		{
-			Lock.Unlock(); // Unlock before calling external code in Consume callback
-			ELoadResult LocalResult = (LoadState == EState::Consumed || ConsumeCallback) ? ELoadResult::AlreadyConsumed : ELoadResult::Inactive;
-			ConsumeSynchronous(LocalResult, FAssetRegistryState());
-			return EConsumeResult::Failed;
-		}
-
-		if (LoadState == EState::LoadSynchronous)
-		{
-			ELoadResult LocalResult = TrySetPath() ? TryLoad() : ELoadResult::NotFound;
-			LoadState = EState::Consumed;
-			Lock.Unlock(); // Unlock before calling external code in Consume callback
-			ConsumeSynchronous(LocalResult, MoveTemp(Payload));
-			Shutdown(); // Shutdown can be called outside the lock since AsyncThread doesn't exist
-			return LocalResult == ELoadResult::Succeeded ? EConsumeResult::Succeeded : EConsumeResult::Failed;
-		}
-
-		// Cancel any further searching in Paks since we will no longer accept preloads starting after this point
-		FCoreDelegates::GetOnPakFileMounted2().Remove(PakMountedDelegate);
-		PakMountedDelegate.Reset();
-
-		if (ConsumeAsynchronous && LoadState == EState::Loading)
-		{
-			// The load might have completed and the TryAsyncLoad thread is waiting to enter the lock, but we will still defer since Consume won the race
-			ConsumeCallback = MoveTemp(ConsumeAsynchronous);
-			return EConsumeResult::Deferred;
-		}
-
-		{
-			SCOPED_BOOT_TIMING("BlockingConsume");
-			// If the load is in progress, wait for it to finish (which it does outside the lock)
-			PreloadReady->Wait();
-		}
-
-		// TryAsyncLoad might not yet have set state to Loaded
-		check(LoadState == EState::Loaded || LoadState == EState::Loading || LoadState == EState::NotFound);
-		ELoadResult LocalResult = LoadState == EState::NotFound ? ELoadResult::NotFound : LoadResult;
-		LoadState = EState::Consumed;
-
-		// No further async threads exist that will read/write payload at this point so we can use it outside the lock
-		Lock.Unlock(); // Unlock before calling external code in Consume callback
-		ConsumeSynchronous(LocalResult, MoveTemp(Payload));
-		Shutdown(); // Shutdown can be called outside the lock since we have set state to Consumed and the Async thread will notice and exit
-		return LocalResult == ELoadResult::Succeeded ? EConsumeResult::Succeeded : EConsumeResult::Failed;
-	}
-
-	/** Called when the Preloader has no further work to do, to free resources early since destruction occurs at end of process. */
-	void Shutdown(bool bFromGlobalDestructor = false)
-	{
-		OnTaskGraphReady.Reset();
-		if (PreloadReady)
-		{
-			// If we are exiting the process early while PreloadReady is still allocated, the event
-			// system has already been torn down and there is nothing for us to free for PreloadReady.
-			if (!bFromGlobalDestructor)
-			{
-				FPlatformProcess::ReturnSynchEventToPool(PreloadReady);
-			}
-			PreloadReady = nullptr;
-		}
-		ARPath.Reset();
-		Payload.Reset();
-		CSV_TRACE_REGION_END(TEXT("Preload"), TEXT("AssetRegistry"))
-	}
-
-	/** simple way to trigger a callback at a specific time that TaskGraph is usable. */
-	TOptional<FDelayedAutoRegisterHelper> OnTaskGraphReady;
-
-	/** Lock that guards members on this (see notes on each member). */
-	FCriticalSection StateLock;
-	/** Trigger for blocking Consume to wait upon TryLoadAsync. This Trigger is only allocated when in the states NotFound, Loaded, Loading. */
-	FEvent* PreloadReady = nullptr;
-
-	/** Path discovered for the AssetRegistry; Read/Write only within the Lock. */
-	FString ARPath;
-
-	/**
-	 * The ARState loaded from disk. Owned exclusively by either the first Consume or by TryAsyncLoad.
-	 * If LoadState is never set to Loading, this state is read/written only by the first thread to call Consume.
-	 * If LoadState is set to Loading (which happens before threading starts), the thread running TryAsyncLoad
-	 * owns this payload until it triggers PayloadReady, after which ownership returns to the first thread to call Consume.
-	 */
-	FAssetRegistryState Payload;
-
-	FDelegateHandle PakMountedDelegate;
-
-	/** Callback from ConsumeOrDefer that is set so TryLoadAsync can trigger the Consume when it completes.Read / Write only within the lock. */
-	FConsumeFunction ConsumeCallback;
-
-	/** State machine state. Read/Write only within the lock (or before threading starts). */
-	EState LoadState = EState::WillNeverPreload;
-
-	/** Result of TryLoad.Thread ownership rules are the same as the rules for Payload. */
-	ELoadResult LoadResult = ELoadResult::UninitializedMemberLoadResult;
-}
-GPreloader;
-
-FAsyncConsumer::~FAsyncConsumer()
-{
-	if (Consumed)
-	{
-		FPlatformProcess::ReturnSynchEventToPool(Consumed);
-		Consumed = nullptr;
-	}
-}
-
-void FAsyncConsumer::PrepareForConsume()
-{
-	// Called within the lock
-	check(!Consumed);
-	Consumed = FPlatformProcess::GetSynchEventFromPool(true /* bIsManualReset */);
-	++ReferenceCount;
-};
-
-void FAsyncConsumer::Wait(UE::AssetRegistry::FInterfaceWriteScopeLock& ScopeLock)
-{
-	// Called within the lock
-	if (ReferenceCount == 0)
-	{
-		return;
-	}
-	++ReferenceCount;
-
-	// Wait outside of the lock so that the AsyncThread can enter the lock to call Consume
-	{
-		ScopeLock.Lock.WriteUnlock();
-		ON_SCOPE_EXIT{ ScopeLock.Lock.WriteLock(); };
-		check(Consumed != nullptr);
-		Consumed->Wait();
-	}
-
-	--ReferenceCount;
-	if (ReferenceCount == 0)
-	{
-		// We're the last one to drop the refcount, so delete Consumed
-		check(Consumed != nullptr);
-		FPlatformProcess::ReturnSynchEventToPool(Consumed);
-		Consumed = nullptr;
-	}
-}
-
-void FAsyncConsumer::Consume(UAssetRegistryImpl& UARI, UE::AssetRegistry::Impl::FEventContext& EventContext, ELoadResult LoadResult, FAssetRegistryState&& ARState)
-{
-	// Called within the lock
-	UARI.GuardedData.LoadPremadeAssetRegistry(EventContext, LoadResult, MoveTemp(ARState));
-	check(ReferenceCount >= 1);
-	check(Consumed != nullptr);
-	Consumed->Trigger();
-	--ReferenceCount;
-	if (ReferenceCount == 0)
-	{
-		// We're the last one to drop the refcount, so delete Consumed
-		FPlatformProcess::ReturnSynchEventToPool(Consumed);
-		Consumed = nullptr;
-	}
-}
-
-}
-
 namespace UE::AssetRegistry
 {
 
diff --git a/Engine/Source/Runtime/AssetRegistry/Private/AssetRegistryImpl.h b/Engine/Source/Runtime/AssetRegistry/Private/AssetRegistryImpl.h
--- a/Engine/Source/Runtime/AssetRegistry/Private/AssetRegistryImpl.h
+++ b/Engine/Source/Runtime/AssetRegistry/Private/AssetRegistryImpl.h
@@ -2,6 +2,8 @@
 
 #pragma once
 
+#include "ARLoader.h"
+#include "ARTypes.h"
 #include "AssetDataGatherer.h"
 #include "AssetRegistry/AssetData.h"
 #include "AssetRegistry/AssetRegistryState.h"
@@ -14,8 +16,6 @@
 #include "HAL/CriticalSection.h"
 #include "Math/NumericLimits.h"
 #include "Misc/Optional.h"
-#include "Misc/TransactionallySafeRWLock.h"
-#include "Misc/ScopeRWLock.h"
 #include "ModuleDescriptor.h"
 #include "PackageDependencyData.h"
 #include "Templates/Function.h"
@@ -23,15 +23,6 @@
 
 #include <atomic>
 
-// Disable premade asset registry until the target specifies it in UBT
-#ifdef ASSETREGISTRY_ENABLE_PREMADE_REGISTRY_IN_EDITOR
-#error "Please use ASSETREGISTRY_FORCE_PREMADE_REGISTRY_IN_EDITOR instead of ASSETREGISTRY_ENABLE_PREMADE_REGISTRY_IN_EDITOR"
-#endif
-
-#ifndef ASSETREGISTRY_FORCE_PREMADE_REGISTRY_IN_EDITOR
-#define ASSETREGISTRY_FORCE_PREMADE_REGISTRY_IN_EDITOR 0
-#endif
-
 class ITargetPlatform;
 class FDependsNode;
 struct FARFilter;
@@ -39,6 +30,7 @@ class FAssetDataGatherer;
 struct FFileChangeData;
 class FPackageReader;
 class UAssetRegistryImpl;
+namespace UE::AssetRegistry::DependsNode { struct FGatherDependencyDataScratchpad; }
 namespace UE::AssetRegistry::Impl { struct FEventContext; }
 namespace UE::AssetRegistry::Premade { enum class ELoadResult : uint8; }
 namespace UE::AssetRegistry::Premade { struct FAsyncConsumer; }
@@ -46,13 +38,6 @@ namespace UE::AssetRegistry::Impl { struct FInitializeContext; }
 namespace UE::AssetRegistry::Impl { struct FScanPathContext; }
 namespace UE::AssetRegistry::Impl { struct FTickContext; }
 namespace UE::AssetRegistry::Impl { struct FClassInheritanceContext; }
-namespace UE::AssetRegistry::Private { class FRWLockWithPriority; }
-namespace UE::AssetRegistry::Private { using FInterfaceRWLock = FRWLockWithPriority; }
-namespace UE::AssetRegistry::DependsNode { struct FGatherDependencyDataScratchpad; }
-namespace UE::AssetRegistry { template <typename TScopeLockType> class TRWScopeLockWithPriority; }
-namespace UE::AssetRegistry { using FInterfaceReadScopeLock = TRWScopeLockWithPriority<UE::TReadScopeLock<Private::FRWLockWithPriority>>; }
-namespace UE::AssetRegistry { using FInterfaceWriteScopeLock = TRWScopeLockWithPriority<UE::TWriteScopeLock<Private::FRWLockWithPriority>>; }
-namespace UE::AssetRegistry { using FInterfaceRWScopeLock = class FRWScopeLockWithPriority; }
 
 #if WITH_EDITOR
 namespace UE::AssetDependencyGatherer::Private { class FRegisteredAssetDependencyGatherer; }
@@ -60,186 +45,6 @@ namespace UE::AssetDependencyGatherer::Private { class FRegisteredAssetDependenc
 
 namespace UE::AssetRegistry
 {
-namespace Premade
-{
-/**
-* A struct to consume a Premade asset registry on an async thread. It supports a cheap Wait() call so that it can be Waited on
-* by frequent AssetRegistry interface calls for the rest of the process.
-*/
-struct FAsyncConsumer
-{
-	~FAsyncConsumer();
-	/** Sets the Consumer into need-to-wait mode; Wait will block until Consume is called. */
-	void PrepareForConsume();
-	/**
-	* Does not return until after the Premade AssetRegistryState has been added into the target AssetRegistry, or it has
-	* been decided not to be used. For performance reasons, this must be called within the WriteScopeLock, and the caller
-	* must handle the possibility that it leaves and reenters the lock.
-	*/
-	void Wait(FInterfaceWriteScopeLock& ScopeLock);
-	/** Callback from the async thread to consume the Premade ARState */
-	void Consume(UAssetRegistryImpl& UARI, UE::AssetRegistry::Impl::FEventContext& EventContext, ELoadResult LoadResult, FAssetRegistryState&& ARState);
-
-private:
-	/**
-	* ReferenceCounter for the Consumed event. Also used to decide whether Waiting is necessary. Read/Write only inside the lock.
-	*/
-	int32 ReferenceCount = 0;
-	/** Event used to Wait for Consume. Allocated/deallocated within the lock. Can be waited on outside the lock if RefCount is held. */
-	FEvent* Consumed = nullptr;
-};
-} // namespace Premade within UE::AssetRegistry
-
-namespace Private
-{
-/** Custom lock class that can prioritize waiters upon request. This is used to allow threads with critical requests
-*  to request that the gatherer thread pause its work and allow the higher pri threads to jump in. Note that for this
-*  to work correctly it must be used with the associated TPriorityScopeLock types
-*  (TScopeLockWithPriority and FRWScopeLockWithPriority)
-*/
-class FRWLockWithPriority
-{
-public:
-	FORCEINLINE void WriteLock()
-	{
-		Mutex.WriteLock();
-	}
-
-	FORCEINLINE void WriteUnlock()
-	{
-		Mutex.WriteUnlock();
-	}
-
-	FORCEINLINE void ReadLock()
-	{
-		Mutex.ReadLock();
-	}
-
-	FORCEINLINE void ReadUnlock()
-	{
-		Mutex.ReadUnlock();
-	}
-
-	bool HasWaiters() const
-	{
-		bool Result;
-		UE_AUTORTFM_OPEN
-		{
-			Result = HighPriorityWaitersCount.load(std::memory_order_relaxed) > 0;
-		};
-		return Result;
-	}
-
-private:
-	friend FInterfaceReadScopeLock;
-	friend FInterfaceWriteScopeLock;
-	friend FInterfaceRWScopeLock;
-	FTransactionallySafeRWLock Mutex;
-	std::atomic<int32> HighPriorityWaitersCount;
-};
-
-enum ELockPriority : uint8
-{
-	PriorityLow,
-	PriorityHigh
-};
-} // namespace Private within namespace UE::AssetRegistry
-
-/**
-* Tracks high-priority waiters on a lock, in an thread- and transactionally-safe manner.
-*/
-class FScopedHighPriorityWaitTracker
-{
-public:
-	FScopedHighPriorityWaitTracker(std::atomic<int32>* InCounter, Private::ELockPriority InPriority)
-		: Counter((InPriority == Private::PriorityHigh) ? InCounter : nullptr)
-	{
-		if (Counter)
-		{
-			UE_AUTORTFM_OPEN
-			{
-				Counter->fetch_add(1, std::memory_order_relaxed);
-			};
-			AutoRTFM::PushOnAbortHandler(this, [this]
-				{
-					Counter->fetch_sub(1, std::memory_order_relaxed);
-				});
-		}
-	}
-
-	~FScopedHighPriorityWaitTracker()
-	{
-		if (Counter)
-		{
-			UE_AUTORTFM_OPEN
-			{
-				Counter->fetch_sub(1, std::memory_order_relaxed);
-			};
-			AutoRTFM::PopOnAbortHandler(this);
-		}
-	}
-
-private:
-	std::atomic<int32>* Counter = nullptr;
-};
-
-/** Keeps a FRWLock read-locked while this scope lives */
-/** This is almost a clone of the existing FReadScopeLock and similar types
-	however this adds an extra flag to help the background processing thread
-	know when a higher priority thread would like to gain access to the protected
-	data */
-template <typename TScopeLockType> class TRWScopeLockWithPriority
-{
-public:
-	UE_NODISCARD_CTOR explicit TRWScopeLockWithPriority(Private::FRWLockWithPriority& InLock,
-		Private::ELockPriority InPriority = Private::PriorityHigh)
-		: Lock(InLock)
-		, Priority(InPriority)
-	{
-		FScopedHighPriorityWaitTracker Tracker(&Lock.HighPriorityWaitersCount, Priority);
-		GuardWrapper.Emplace(Lock);
-	}
-
-	TOptional<TScopeLockType> GuardWrapper;
-	UE::AssetRegistry::Private::FRWLockWithPriority& Lock;
-	UE::AssetRegistry::Private::ELockPriority Priority;
-};
-
-class FRWScopeLockWithPriority
-{
-public:
-	UE_NODISCARD_CTOR explicit FRWScopeLockWithPriority(Private::FRWLockWithPriority& InLockObject,
-		FRWScopeLockType InLockType, Private::ELockPriority InPriority = Private::PriorityHigh)
-		: Lock(InLockObject)
-		, Priority(InPriority)
-		, LockType(InLockType)
-	{
-		FScopedHighPriorityWaitTracker Tracker(&Lock.HighPriorityWaitersCount, Priority);
-		GuardWrapper.Emplace(Lock, LockType);
-	}
-
-	// NOTE: As the name suggests, this function should be used with caution. 
-	// It releases the read lock _before_ acquiring a new write lock. This is not an atomic operation and the caller should 
-	// not treat it as such. 
-	// E.g. Pointers read from protected data structures prior to this call may be invalid after the function is called. 
-	void ReleaseReadOnlyLockAndAcquireWriteLock_USE_WITH_CAUTION()
-	{
-		if (LockType == SLT_ReadOnly)
-		{
-			GuardWrapper.Reset();
-			{
-				FScopedHighPriorityWaitTracker Tracker(&Lock.HighPriorityWaitersCount, Priority);
-				GuardWrapper.Emplace(Lock, SLT_Write);
-			}
-			LockType = SLT_Write;
-		}
-	}
-
-	Private::FRWLockWithPriority& Lock;
-	TOptional<TRWScopeLock<Private::FRWLockWithPriority>> GuardWrapper;
-	Private::ELockPriority Priority;
-	FRWScopeLockType LockType;
-};
 
 namespace Impl
 {
