[AssetRegistry]

Gracefully error loading when reading in a corrupt amount of FNames from a corrupt cache file on disk.

#rnx

#Synced-CL 50345503

diff --git a/Engine/Source/Runtime/AssetRegistry/Private/AssetRegistryArchive.cpp b/Engine/Source/Runtime/AssetRegistry/Private/AssetRegistryArchive.cpp
--- a/Engine/Source/Runtime/AssetRegistry/Private/AssetRegistryArchive.cpp
+++ b/Engine/Source/Runtime/AssetRegistry/Private/AssetRegistryArchive.cpp
@@ -319,7 +319,14 @@ FArchive& FAssetRegistryReader::operator<<(FName& Out)
 		*this << Number;
 	}
 
-	Out = Names[Index].ToName(Number);
+	if (Index >= (uint32) Names.Num())
+	{
+		SetError();
+	}
+	else
+	{
+		Out = Names[Index].ToName(Number);
+	}
 
 	return *this;
 }
