Map tag store data in Asset Registry cache, saving GBs of editor memory usage on large projects.

- Support for earlier version of AR dropped due to a tight coupling between the file format and in-memory layout. Same reason for only supporting the latest version of the format.

#Synced-CL 50453497

diff --git a/Engine/Source/Runtime/AssetRegistry/Private/AssetDataGatherer.cpp b/Engine/Source/Runtime/AssetRegistry/Private/AssetDataGatherer.cpp
--- a/Engine/Source/Runtime/AssetRegistry/Private/AssetDataGatherer.cpp
+++ b/Engine/Source/Runtime/AssetRegistry/Private/AssetDataGatherer.cpp
@@ -2,8 +2,9 @@
 
 #include "AssetDataGatherer.h"
 #include "AssetDataGathererPrivate.h"
-#include "AssetRegistryImpl.h"
 #include "AssetRegistry.h"
+#include "AssetRegistryImpl.h"
+#include "AssetRegistry/AssetDataTagMapSerializationDetails.h"
 #include "AssetRegistry/AssetRegistryTelemetry.h"
 
 #include "Algo/AnyOf.h"
@@ -13,7 +14,6 @@
 #include "AssetRegistry/AssetData.h"
 #include "AssetRegistryArchive.h"
 #include "AssetRegistryPrivate.h"
-#include "Async/MappedFileHandle.h"
 #include "Async/ParallelFor.h"
 #include "Containers/BinaryHeap.h"
 #include "HAL/LowLevelMemTracker.h"
@@ -46,6 +46,7 @@
 #include "Templates/UnrealTemplate.h"
 #include "Trace/Trace.h"
 #include "Trace/Trace.inl"
+#include "Serialization/MemoryHasher.h"
 
 namespace UE::AssetDataGather::Private
 {
@@ -62,6 +63,12 @@ static FAutoConsoleVariableRef CVarIgnoreEmptyDirectories(
 	bIgnoreEmptyDirectories,
 	TEXT("If true, completely empty leaf directories are ignored by the asset registry while scanning"));
 
+bool bMemoryMappedCacheFiles = true;
+static FAutoConsoleVariableRef CVarMemoryMappedCacheFiles(
+	TEXT("AssetRegistry.MemoryMappedCacheFiles"),
+	bMemoryMappedCacheFiles,
+	TEXT("If true, asset registry cache file will be memory mapped, reducing memory usage"));
+
 void LexFromString(EFeatureEnabledReadWrite& OutValue, FStringView Text)
 {
 	Text.TrimStartAndEndInline();
@@ -3345,395 +3352,6 @@ FAssetDataDiscovery::FDirectoryResult::FDirectoryResult(FStringView InDirAbsPath
 
 }
 
-// Reads an FMemoryView once
-class FMemoryViewReader
-{
-public:
-	FMemoryViewReader() = default;
-	FMemoryViewReader(FMemoryView Data) : Remaining(Data), TotalSize(Data.GetSize()) {}
-
-	uint64 GetRemainingSize() const { return Remaining.GetSize(); }
-	uint64 GetTotalSize() const { return TotalSize; }
-	uint64 Tell() const { return TotalSize - Remaining.GetSize(); }
-
-	FMemoryView Load(uint64 Size)
-	{
-		check(Size <= Remaining.GetSize());
-		FMemoryView Out(Remaining.GetData(), Size);
-		Remaining += Size;
-		return Out;
-	}
-
-	void Load(FMutableMemoryView Out)
-	{
-		FMemoryView In = Load(Out.GetSize());
-		if (In.GetSize())
-		{
-			FMemory::Memcpy(Out.GetData(), In.GetData(), In.GetSize());
-		}
-	}
-
-	template<typename T>
-	T Load()
-	{
-		static_assert(std::is_integral_v<T>, "Only integer loading supported");
-		static_assert(PLATFORM_LITTLE_ENDIAN, "Byte-swapping not implemented");
-		return FPlatformMemory::ReadUnaligned<T>(Load(sizeof(T)).GetData());
-	}
-
-	template<typename T>
-	TOptional<T> TryLoad()
-	{
-		return sizeof(T) <= Remaining.GetSize() ? Load<T>() : TOptional<T>();
-	}
-
-private:
-	FMemoryView Remaining;
-	uint64 TotalSize = 0;
-};
-
-
-// Enables both versioning and distinguishing out-of-sync reads from data corruption
-static constexpr uint32 BlockMagic = 0xb1a3;
-
-struct FBlockHeader
-{
-	uint32 Magic = 0;
-	uint32 Size = 0;
-	uint64 Checksum = 0;
-};
-
-TOptional<FBlockHeader> LoadBlockHeader(FMemoryView Data)
-{
-	check(Data.GetSize() == sizeof(FBlockHeader));
-
-	FMemoryViewReader Reader(Data);
-	FBlockHeader Header;
-	Header.Magic = Reader.Load<uint32>();
-	Header.Size = Reader.Load<uint32>();
-	Header.Checksum = Reader.Load<uint64>();
-
-	if (Header.Magic != BlockMagic)
-	{
-		UE_LOG(LogAssetRegistry, Warning, TEXT("Wrong block magic (0x%x)"), Header.Magic);
-		return {};
-	}
-
-	return Header;
-}
-
-static uint64 CalculateBlockChecksum(FMemoryView Data)
-{
-	return INTEL_ORDER64(FXxHash64::HashBuffer(Data).Hash);
-}
-
-class FChecksumArchiveBase : public FArchiveProxy
-{
-	static constexpr uint32 SaveBlockSize = 4 << 20;
-
-	struct FBlock
-	{
-		uint8* Begin = nullptr;
-		uint8* Cursor = nullptr;
-		uint8* End = nullptr;
-
-		explicit FBlock(uint32 Size)				{ Reset(Size); }
-		~FBlock()									{ FMemory::Free(Begin); }
-		uint64 GetCapacity() const					{ return End - Begin; }
-		uint64 GetUsedSize() const					{ return Cursor - Begin; }
-		uint64 GetRemainingSize() const				{ return End - Cursor; }
-		FMutableMemoryView GetUsed() const			{ return {Begin, GetUsedSize()}; }
-		FMutableMemoryView GetRemaining() const		{ return {Cursor, GetRemainingSize()}; };
-
-		void Reset(uint32 Size)
-		{
-			if (GetCapacity() < Size)
-			{
-				FMemory::Free(Begin);
-				Begin = (uint8*)FMemory::Malloc(Size);
-			}
-			
-			// All blocks have the same size except the last one, which may be smaller.
-			// It doesn't matter that we lose some capacity when loading the last block.
-			End = Begin + Size;
-			Cursor = Begin;
-		}
-
-		void Write(FMemoryView In)
-		{
-			check(GetRemainingSize() >= In.GetSize());
-			if (In.GetSize())
-			{
-				FMemory::Memcpy(Cursor, In.GetData(), In.GetSize());
-			}
-			Cursor += In.GetSize();	
-		}
-
-		void Read(FMutableMemoryView Out)
-		{
-			check(GetRemainingSize() >= Out.GetSize());
-			if (Out.GetSize())
-			{
-				FMemory::Memcpy(Out.GetData(), Cursor, Out.GetSize());
-			}
-			Cursor += Out.GetSize();	
-		}
-	};
-
-	FBlock Block;
-
-	void SaveBlock()
-	{
-		FBlockHeader Header;
-		Header.Magic = BlockMagic;
-		Header.Size = IntCastChecked<uint32>(Block.GetUsedSize());
-		Header.Checksum = CalculateBlockChecksum(Block.GetUsed());
-		InnerArchive << Header.Magic << Header.Size << Header.Checksum;
-
-		InnerArchive.Serialize(Block.Begin, Header.Size);
-
-		Block.Cursor = Block.Begin;
-	}
-
-	bool LoadBlock()
-	{
-		check(Block.GetRemainingSize() == 0);
-
-		uint8 HeaderData[sizeof(FBlockHeader)];
-		InnerArchive.Serialize(HeaderData, sizeof(HeaderData));
-		if (InnerArchive.IsError())
-		{
-			UE_LOG(LogAssetRegistry, Warning, TEXT("Couldn't read block header"));
-			return false;
-		}
-
-		if (TOptional<FBlockHeader> Header = LoadBlockHeader(MakeMemoryView(HeaderData)))
-		{
-			Block.Reset(Header->Size);
-
-			InnerArchive.Serialize(Block.Begin, Header->Size);
-
-			if (InnerArchive.IsError())
-			{
-				UE_LOG(LogAssetRegistry, Warning, TEXT("Couldn't read block data"));
-				return false;
-			}
-			else if (CalculateBlockChecksum(Block.GetRemaining()) != Header->Checksum)
-			{
-				UE_LOG(LogAssetRegistry, Warning, TEXT("Wrong block checksum"));
-				return false;
-			}
-
-			return true;
-		}
-
-		return false;
-	}
-
-public:
-	explicit FChecksumArchiveBase(FArchive& Inner)
-	: FArchiveProxy(Inner)
-	, Block(IsLoading() ? 0 : SaveBlockSize)
-	{}
-
-protected:
-	~FChecksumArchiveBase()
-	{
-		if (!IsLoading() && Block.GetUsedSize() > 0)
-		{
-			SaveBlock();
-		}
-	}
-	const FBlock& GetCurrentBlock() const
-	{
-		return Block;
-	}
-
-	void Save(FMemoryView Data)
-	{
-		for (uint64 Size = Block.GetRemainingSize(); Size < Data.GetSize(); Size = Block.GetRemainingSize())
-		{
-			Block.Write(Data.Left(Size));
-			Data += Size;
-			SaveBlock();
-		}
-
-		Block.Write(Data);
-	}
-
-	void Load(FMutableMemoryView Data)
-	{
-		if (IsError())
-		{
-			return;
-		}
-
-		for (uint64 Size = Block.GetRemainingSize(); Size < Data.GetSize(); Size = Block.GetRemainingSize())
-		{
-			Block.Read(Data.Left(Size));
-			Data += Size;
-
-			if (!LoadBlock())
-			{
-				UE_LOG(LogAssetRegistry, Warning, TEXT("Integrity check failed, '%s' cache will be discarded"), *InnerArchive.GetArchiveName());
-				SetError();
-				return;
-			}
-		}
-
-		Block.Read(Data);
-	}
-
-public:
-	// Use FArchive implementations that map back to Serialize() instead of FArchiveProxy overloads
-	// that forward to the inner archive and bypass integrity checking
-	virtual FArchive& operator<<(FText& Value) override					{ return FArchive::operator<<(Value); }
-	virtual void SerializeBits(void* Bits, int64 LengthBits) override	{ FArchive::SerializeBits(Bits, LengthBits); }
-	virtual void SerializeInt(uint32& Value, uint32 Max) override		{ FArchive::SerializeInt(Value, Max); }
-	virtual void SerializeIntPacked(uint32& Value) override				{ FArchive::SerializeIntPacked(Value); }
-
-private:
-	// Wrapping an inner FArchiveUObject is not supported. The inner archive should be low-level
-	// and an outer archive should have intercepted these calls.
-	virtual FArchive& operator<<(FName& Value) override					{ unimplemented(); return *this; }
-	virtual FArchive& operator<<(UObject*& Value) override				{ unimplemented(); return *this; }
-	virtual FArchive& operator<<(FObjectPtr& Value) override			{ unimplemented(); return *this; }
-	virtual FArchive& operator<<(FLazyObjectPtr& Value) override		{ unimplemented(); return *this; }
-	virtual FArchive& operator<<(FSoftObjectPath& Value) override		{ unimplemented(); return *this; }
-	virtual FArchive& operator<<(FSoftObjectPtr& Value) override		{ unimplemented(); return *this; }
-	virtual FArchive& operator<<(FWeakObjectPtr& Value) override		{ unimplemented(); return *this; }
-	virtual FArchive& operator<<(FField*& Value) override				{ unimplemented(); return *this; }
-
-	virtual void Seek(int64 InPos) override								{ unimplemented(); }
-};
-
-class FChecksumArchiveWriter : public FChecksumArchiveBase
-{
-public:
-	using FChecksumArchiveBase::FChecksumArchiveBase;
-	virtual void Serialize(void* V, int64 Len) override { Save(FMemoryView(V, Len)); }
-	virtual int64 Tell() override { return InnerArchive.Tell() + GetCurrentBlock().GetUsedSize(); }
-};
-
-class FChecksumArchiveReader : public FChecksumArchiveBase
-{
-public:
-	using FChecksumArchiveBase::FChecksumArchiveBase;
-	virtual void Serialize(void* V, int64 Len) override { Load(FMutableMemoryView(V, Len)); }
-	virtual int64 Tell() override { return InnerArchive.Tell() - GetCurrentBlock().GetRemainingSize(); }
-};
-
-// Memory-mapped equivalent of FChecksumArchiveReader
-class FChecksumViewReader : public FArchive
-{
-public:
-	explicit FChecksumViewReader(FMemoryViewReader&& Reader, FStringView InFileName)
-	: RemainingBlocks(MoveTemp(Reader))
-	, FileName(InFileName)
-	{
-		SetIsLoading(true);
-	}
-
-private:
-	FMemoryViewReader RemainingBlocks;
-	FMemoryViewReader CurrentBlock;
-	FString FileName;
-
-	virtual void Seek(int64) override { unimplemented(); }
-	virtual int64 Tell() override { return RemainingBlocks.Tell() - CurrentBlock.GetRemainingSize(); }
-	virtual int64 TotalSize() override { return RemainingBlocks.GetTotalSize(); }
-
-	virtual void Serialize(void* V, int64 Len) override
-	{
-		FMutableMemoryView Out(V, Len);
-
-		while (CurrentBlock.GetRemainingSize() < Out.GetSize())
-		{
-			if (IsError())
-			{
-				return;
-			}
-
-			FMutableMemoryView OutSlice(Out.GetData(), CurrentBlock.GetRemainingSize());
-			Out += OutSlice.GetSize();
-			CurrentBlock.Load(OutSlice);
-			check(CurrentBlock.GetRemainingSize() == 0);
-
-			TOptional<FMemoryView> NextBlock = LoadNextBlock(RemainingBlocks);
-			if (!NextBlock)
-			{
-				UE_LOG(LogAssetRegistry, Warning, TEXT("Integrity check failed, '%s' cache will be discarded"), *FileName);
-				SetError();
-				return;
-			}
-			CurrentBlock = *NextBlock;
-		}
-
-		CurrentBlock.Load(Out);
-	}
-
-	FORCENOINLINE static TOptional<FMemoryView> LoadNextBlock(FMemoryViewReader& In)
-	{
-		if (In.GetRemainingSize() < sizeof(FBlockHeader))
-		{
-			UE_LOG(LogAssetRegistry, Warning, TEXT("Couldn't read block header"));
-			return {};
-		}
-
-		if (TOptional<FBlockHeader> Header = LoadBlockHeader(In.Load(sizeof(FBlockHeader))))
-		{
-			if (Header->Size > In.GetRemainingSize())
-			{
-				UE_LOG(LogAssetRegistry, Warning, TEXT("Incomplete block"));
-				return {};
-			}
-
-			FMemoryView Block = In.Load(Header->Size);
-			if (CalculateBlockChecksum(Block) != Header->Checksum)
-			{
-				UE_LOG(LogAssetRegistry, Warning, TEXT("Wrong block checksum"));
-				return {};
-			}
-
-			return Block;
-		}
-
-		return {};
-	}
-};
-
-// Util that maps an entire file
-class FMemoryMappedFile
-{
-public:
-	explicit FMemoryMappedFile(const TCHAR* Path)
-	{
-		FOpenMappedResult Result = FPlatformFileManager::Get().GetPlatformFile().OpenMappedEx(Path);
-		if (!Result.HasError())
-		{
-			Handle = Result.StealValue();
-			Region.Reset(Handle->MapRegion());
-		}
-	}
-
-	void Preload(int64 Size = MAX_int64) const
-	{
-		if (Region)
-		{
-			Region->PreloadHint(0, Size);
-		}
-	}
-	
-	FMemoryView View() const
-	{
-		return Region ? FMemoryView(Region->GetMappedPtr(), Region->GetMappedSize()) : FMemoryView();
-	}
-
-private:
-	TUniquePtr<IMappedFileHandle> Handle;
-	TUniquePtr<IMappedFileRegion> Region;
-};
-
 #if UE_EDITOR
 /** A class to preload the AssetData cache and start the AssetDataDiscovery thread (which does its own cache loading)
  * used by the FAssetDataGatherer. Preloading allows us to start very early in editor startup, so that we 
@@ -5324,21 +4942,28 @@ int64 FAssetDataGatherer::SaveCacheFileInternal(const FString& CacheFilename, co
 	TUniquePtr<FArchive> FileAr(IFileManager::Get().CreateFileWriter(*TempFilename, 0));
 	if (FileAr)
 	{
-		using namespace UE::AssetDataGather::Private;
-		uint64 CurrentVersion = FGatherSettings::CurrentVersion;
-		*FileAr << CurrentVersion;
-
-#if ALLOW_NAME_BATCH_SAVING
 		{
-			// We might be able to reduce load time by using AssetRegistry::SerializationOptions
-			// to save certain common tags as FName.
-			UE::AssetDataGather::Private::FChecksumArchiveWriter ChecksummingWriter(*FileAr);
-			FAssetRegistryWriter Ar(FAssetRegistryWriterOptions(), ChecksummingWriter);
-			UE::AssetDataGather::Private::SerializeCacheSave(Ar, AssetsToSave);
+			UE::THashingArchiveProxy<FBlake3> FileArHash(*FileAr);
+
+			using namespace UE::AssetDataGather::Private;
+			uint64 CurrentVersion = FGatherSettings::CurrentVersion;
+			FileArHash << CurrentVersion;
+
+	#if ALLOW_NAME_BATCH_SAVING
+			{
+				// We might be able to reduce load time by using AssetRegistry::SerializationOptions
+				// to save certain common tags as FName.
+				FAssetRegistryWriter Ar(FAssetRegistryWriterOptions(), FileArHash);
+				UE::AssetDataGather::Private::SerializeCacheSave(Ar, AssetsToSave);
+			}
+	#else		
+			static_assert(false, "Cannot save asset registry cache in this configuration");
+	#endif
+			// write the integrity checksum
+			FBlake3Hash Hash = FileArHash.GetHash();
+			FileAr->Serialize(&Hash, sizeof(Hash));
 		}
-#else		
-		checkf(false, TEXT("Cannot save asset registry cache in this configuration"));
-#endif
+
 		int64 CacheSize = FileAr->TotalSize();
 		// Close file handle before moving temp file to target 
 		FileAr.Reset();
@@ -5440,7 +5065,13 @@ TArray<FCachePayload> LoadCacheFiles(TConstArrayView<FString> InCacheFilenames)
 	TArray<FCachePayload> Results;
 	Results.AddDefaulted(InCacheFilenames.Num());
 
-	auto DoLoad = [](FArchive& ChecksummingReader)
+	// disable memory mapping on the builder machines until the situation with deletion clarifies
+	if (GIsBuildMachine)
+	{
+		UE::AssetDataGather::Private::bMemoryMappedCacheFiles = false;
+	}
+
+	auto DoLoad = [](FArchive& FilemappedReader, TSharedPtr<UE::AssetRegistry::FMemoryMappedFile> BackingFile)
 		{
 			// We are already using AsyncTasks for our own ParallelFor, passing non-zero NumAsyncWorkers into
 			// FAssetRegistryReader will cause our current task to block waiting for another task to run
@@ -5453,159 +5084,119 @@ TArray<FCachePayload> LoadCacheFiles(TConstArrayView<FString> InCacheFilenames)
 			FAssetRegistryHeader Header;
 			Header.Version = FAssetRegistryVersion::LatestVersion;
 			Header.bFilterEditorOnlyData = false;
-			FAssetRegistryReader RegistryReader(ChecksummingReader, NumAsyncWorkers, Header);
+
+			FAssetRegistryReader RegistryReader(FilemappedReader, UE::AssetDataGather::Private::bMemoryMappedCacheFiles ? BackingFile : nullptr, NumAsyncWorkers, Header);
 			return RegistryReader.IsError() ? FCachePayload() : SerializeCacheLoad(RegistryReader);
 		};
 
 	if (FPlatformProperties::SupportsMemoryMappedFiles())
 	{
-		FCriticalSection ParallelForLock;
-		struct FSharedFileData
-		{
-			void ConditionalConstruct(const TCHAR* FilePath)
-			{
-				if (!File)
-				{
-					File.Emplace(FilePath);
-				}
-			}
-			void ConditionalQueueDestruction(TArray<FSharedFileData*>& Queue)
-			{
-				if (bLoadComplete && bAsyncCacheComplete)
-				{
-					Queue.Add(this);
-				}
-			}
-			TOptional<FMemoryMappedFile> File;
-			bool bLoadComplete = false;
-			bool bAsyncCacheComplete = false;
-		};
-
-		// Allocate multiple ParallelFor body invocations for each cache file.
-		// Each group of ParallelFor body invocations will work on the same cache file together in parallel,
-		// or in series if we run out of threads.
-		enum class ETaskType : uint8
-		{
-			// Load byts from the MemoryMapped file and parse its data into Results
-			DoLoad = 0,
-			// Call Preload on the MemoryMapped file to swap its bytes into memory.
-			// This starts at the same time or later than DoLoad, but it will outpace DoLoad
-			// because DoLoad is doing a lot of cpuwork, and so the AsyncCache thread will 
-			// be the one that incurs IO costs rather than the DoLoad thread.
-			AsyncCache,
-			Count,
-		};
-		constexpr int32 TasksPerResult = (int32)ETaskType::Count;
-		TArray<FSharedFileData> CacheFiles;
-		TArray<FSharedFileData*> DestructQueue;
-		CacheFiles.SetNum(InCacheFilenames.Num());
-		DestructQueue.Reserve(InCacheFilenames.Num());
-		const int32 NumTasks = InCacheFilenames.Num() * TasksPerResult;
-		int32 NextTaskIndex = 0;
-
-		// We want to restrict the number of threads to reduce maximum memory use due to having multiple intermediate cache files
-		// in parallel. But ParallelFor only provides an API for restricting MinBatchSize.
-		const int32 MaxNumThreads = FGatherSettings::GARGatherCacheParallelism * TasksPerResult;
+		const int32 NumTasks = InCacheFilenames.Num();
+		const int32 MaxNumThreads = FGatherSettings::GARGatherCacheParallelism;
 		const int32 MinBatchSize = FMath::DivideAndRoundUp(NumTasks, MaxNumThreads);
 
 		ParallelFor(TEXT("AssetDataGatherer::LoadCacheFiles"), NumTasks, MinBatchSize,
-			[&DoLoad, &ParallelForLock, &InCacheFilenames, &Results, &CacheFiles, &DestructQueue, &NextTaskIndex]
-			(int32 /* Unused ParallelFor Index */)
+			[&DoLoad, &InCacheFilenames, &Results](int32 Index)
 		{
-			FSharedFileData* FileData = nullptr;
-			const FString* CacheFilename = nullptr;
-			bool bShouldAsyncCache = false;
+			const FString& CacheFilename = InCacheFilenames[Index];
 
-			int32 ResultIndex = -1;
-			ETaskType TaskType = ETaskType::DoLoad;
-			FSharedFileData* DestructData = nullptr;
+			FCachePayload Payload;
+			bool bVerifiedOk = false;
 
-			do
+			// First, verify the file and refuse to read if hashes do not agree
+			// We do it before memory mapping the file to avoid bringing into the process memory areas of the file
+			// that we do not expect to touch, and would have to explicitly evict later.
+			TUniquePtr<FArchive> Ar(IFileManager::Get().CreateFileReader(*CacheFilename));
+			if (Ar)
 			{
-				// Pick up the next TaskIndex, but before doing that, destruct any previously completed CacheFiles,
-				// to reduce the maximum amount of memory we have allocated at once across the multiple threads.
-				DestructData = nullptr;
-				{
-					FScopeLock ParallelForScopeLock(&ParallelForLock);
+				// same size as in FMD5Hash::HashFileFromArchive
+				constexpr int32 kLocalScratchSize = 64 * 1024;
+				TArray<uint8> LocalScratch;
+				LocalScratch.SetNumUninitialized(kLocalScratchSize);
 
-					// Peek at what our task would be but don't claim it yet.
-					int32 TaskIndex = NextTaskIndex;
-					ResultIndex = TaskIndex / TasksPerResult;
-					TaskType = (ETaskType)(TaskIndex - ResultIndex * TasksPerResult);
-					check(0 <= (int32)TaskType && TaskType < ETaskType::Count);
+				// we're going to hash all the file except the last N bytes which hold the hash we need to compare
+				constexpr int32 kHashSize = sizeof(FBlake3Hash);
+				const int64 SizeToHash = Ar->TotalSize() - kHashSize;
 
-					if (!DestructQueue.IsEmpty()
-						// Don't do destruct work on the last AsyncCache thread, save it for after the parallelfor
-						&& (TaskType == ETaskType::DoLoad || ResultIndex + 1 < InCacheFilenames.Num()))
-					{
-						DestructData = DestructQueue.Pop(EAllowShrinking::No);
-						check(DestructData != nullptr);
-					}
-					else
+				if (SizeToHash > 0)
+				{
+					int64 Position = 0;
+					FBlake3Hash ComputedHash = [&]()
 					{
-						// Claim the task.
-						NextTaskIndex++;
+						THashingArchiveProxy<FBlake3> HashingAr(*Ar);
+						// Read in BufferSize chunks
+						while (Position < SizeToHash)
+						{
+							const auto ReadNum = FMath::Min(SizeToHash - Position, (int64)LocalScratch.Num());
+							HashingAr.Serialize(LocalScratch.GetData(), ReadNum);
 
-						CacheFilename = &InCacheFilenames[ResultIndex];
-						FileData = &CacheFiles[ResultIndex];
-						FileData->ConditionalConstruct(**CacheFilename);
-						bShouldAsyncCache = TaskType == ETaskType::AsyncCache && !FileData->bLoadComplete;
-					}
-				}
+							Position += ReadNum;
+						};
 
-				if (DestructData)
-				{
-					DestructData->File.Reset();
-				}
-			} while (DestructData != nullptr);
-			check(FileData != nullptr && CacheFilename != nullptr); // Required to inform static analysis.
+						return HashingAr.GetHash();
+					}();
 
-			if (TaskType == ETaskType::AsyncCache)
-			{
-				if (bShouldAsyncCache)
-				{
-					FileData->File->Preload();
+					int64 HashToRead = Ar->TotalSize() - Position;
+					if (HashToRead == sizeof(ComputedHash))
+					{
+						FBlake3Hash ReadHash;
+						Ar->Serialize(&ReadHash, sizeof(ReadHash));
+
+						bVerifiedOk = (ReadHash == ComputedHash);
+					}
 				}
-				FScopeLock ParallelForScopeLock(&ParallelForLock);
-				FileData->bAsyncCacheComplete = true;
-				FileData->ConditionalQueueDestruction(DestructQueue);
-				// Don't do any destruction work on our current loop: we are possibly the last iteration
-				// and should allow the ParallelFor to complete, and destruct in a non-blocking task afterwards.
 			}
-			else // ETaskType::DoLoad
+
+			// now read the file if it is valid
+			if (bVerifiedOk)
 			{
-				FCachePayload Payload;
-				FMemoryViewReader FileReader(FileData->File->View());
-				TOptional<uint64> Version = FileReader.TryLoad<uint64>();
+				TSharedPtr<UE::AssetRegistry::FMemoryMappedFile> File = MakeShared<UE::AssetRegistry::FMemoryMappedFile>(*CacheFilename);
+				Ar->Seek(0);
+
+				// we already know (from the verification) that Ar is non-empty
+				uint64 Version;
+				*Ar << Version;
+
 				if (Version == FGatherSettings::CurrentVersion)
 				{
-					FChecksumViewReader ChecksummingReader(MoveTemp(FileReader), *CacheFilename);
-					Payload = DoLoad(ChecksummingReader);
+					Payload = DoLoad(*Ar, File);
 					UE_CLOG(Payload.bSucceeded, LogAssetRegistry, Display,
-						TEXT("Asset registry cache read as %.1f MiB from %s."),
-						static_cast<float>(FileReader.GetTotalSize()) / 1024.f / 1024.f, **CacheFilename);
+						TEXT("Asset registry cache read as %.1f MiB from %s"),
+						static_cast<float>(Ar->TotalSize()) / 1024.f / 1024.f, *CacheFilename);
 					UE_CLOG(!Payload.bSucceeded, LogAssetRegistry, Warning,
-						TEXT("There was an error loading the asset registry cache using memory mapping from %s."),
-						**CacheFilename);
+						TEXT("There was an error loading the asset registry cache from %s."),
+						*CacheFilename);
 				}
-				Results[ResultIndex] = MoveTemp(Payload);
+			}
+			else if (Ar) // check if it was a past version
+			{
+				Ar->Seek(0);
+
+				if (Ar->TotalSize() >= sizeof(uint64))
+				{
+					uint64 VersionAndMagic;
+					*Ar << VersionAndMagic;
 
-				FScopeLock ParallelForScopeLock(&ParallelForLock);
-				FileData->bLoadComplete = true;
-				FileData->ConditionalQueueDestruction(DestructQueue);
-				// Don't do any destruction work on our current loop: we are possibly the last iteration
-				// and should allow the ParallelFor to complete, and destruct in a non-blocking task afterwards.
+					uint32 Magic = (uint32)(VersionAndMagic >> 32);
+					if (Magic == FGatherSettings::CacheSerializationMagic)
+					{
+						uint32 Version = (uint32)(VersionAndMagic & (uint32)(-1));
+						// all later versions should have a valid integrity hash, so if we're here they were genuinely corrupted and this is not a compat issue.
+						if (Version < FAssetRegistryVersion::MemoryMappedTagDataStore)
+						{
+							UE_LOG(LogAssetRegistry, Display, TEXT("Attempted to read AssetRegistry.bin file with an older format version. Backwards compatibility is not supported for AssetRegistry.bin files; they should be regenerated instead. Contact Epic if you need this functionality. Version == %d (LatestVersion is %d), Archive == %s"),
+								Version,
+								FAssetRegistryVersion::LatestVersion,
+								*CacheFilename
+							);
+						}
+					}
+				}
 			}
-		}, EParallelForFlags::BackgroundPriority);
-		for (FSharedFileData& FileData : CacheFiles)
-		{
-			check(FileData.bLoadComplete && FileData.bAsyncCacheComplete);
-		}
 
-		// Ignore the remaining elements of the DestructQueue, and instead pass the entire set of CacheFiles to another
-		// thread to finish the destruction of any that still need it.
-		DestructQueue.Empty();
-		UE::Tasks::Launch(UE_SOURCE_LOCATION, [KillAsync = MoveTemp(CacheFiles)]() {});
+			Results[Index] = MoveTemp(Payload);
+
+		}, EParallelForFlags::BackgroundPriority);
 	}
 	else
 	{
@@ -5621,8 +5212,7 @@ TArray<FCachePayload> LoadCacheFiles(TConstArrayView<FString> InCacheFilenames)
 				*FileAr << Version;
 				if (Version == FGatherSettings::CurrentVersion)
 				{
-					FChecksumArchiveReader ChecksummingReader(*FileAr);
-					Payload = DoLoad(ChecksummingReader);
+					Payload = DoLoad(*FileAr, nullptr);
 					UE_CLOG(Payload.bSucceeded, LogAssetRegistry, Display,
 						TEXT("Asset registry cache read as %.1f MiB from %s"),
 						static_cast<float>(FileAr->TotalSize()) / 1024.f / 1024.f, *CacheFilename);
diff --git a/Engine/Source/Runtime/AssetRegistry/Private/AssetRegistryArchive.cpp b/Engine/Source/Runtime/AssetRegistry/Private/AssetRegistryArchive.cpp
--- a/Engine/Source/Runtime/AssetRegistry/Private/AssetRegistryArchive.cpp
+++ b/Engine/Source/Runtime/AssetRegistry/Private/AssetRegistryArchive.cpp
@@ -5,6 +5,7 @@
 #include "Algo/Sort.h"
 #include "AssetRegistryPrivate.h"
 #include "AssetRegistry/AssetRegistryState.h"
+#include "Serialization/MemoryHasher.h"
 
 #ifndef ASSET_REGISTRY_LOAD_NAMES_IN_ORDER
 #define ASSET_REGISTRY_LOAD_NAMES_IN_ORDER 0
@@ -222,11 +223,17 @@ FAssetRegistryWriter::~FAssetRegistryWriter()
 {
 	// Save store data and collect FNames
 	int64 BodySize = MemWriter.TotalSize();
+	UE::AssetRegistry::AlignPosInArchive(MemWriter);	// make sure the data before store leaves it in the aligned position
+	int64 StoreStart = MemWriter.TotalSize();
 	SaveStore(Tags.Finalize(), *this);
 
 	// Save in load-friendly order - names, store then body / tag maps
 	SaveNameBatch(FlattenIndex(Names), TargetAr);
-	TargetAr.Serialize(MemWriter.GetData() + BodySize, MemWriter.TotalSize() - BodySize);
+
+	// make sure the Store is aligned
+	UE::AssetRegistry::AlignPosInArchive(TargetAr);
+	TargetAr.Serialize(MemWriter.GetData() + StoreStart, MemWriter.TotalSize() - StoreStart);
+	// Body does not need to be aligned because we copy all the data out of it rather than pointing to its memory-mapped position
 	TargetAr.Serialize(MemWriter.GetData(), BodySize);
 }
 
@@ -261,7 +268,7 @@ void FAssetRegistryWriter::SerializeTagsAndBundles(const FAssetData& Out)
 
 #endif
 
-FAssetRegistryReader::FAssetRegistryReader(FArchive& Inner, int32 NumWorkers, FAssetRegistryHeader Header)
+FAssetRegistryReader::FAssetRegistryReader(FArchive& Inner, TSharedPtr<UE::AssetRegistry::FMemoryMappedFile> InBackingFile, int32 NumWorkers, FAssetRegistryHeader Header)
 	: FArchiveProxy(Inner)
 {
 	check(IsLoading());
@@ -269,25 +276,38 @@ FAssetRegistryReader::FAssetRegistryReader(FArchive& Inner, int32 NumWorkers, FA
 	SetFilterEditorOnly(Header.bFilterEditorOnlyData);
 	FArchive::SetFilterEditorOnly(Header.bFilterEditorOnlyData); // Workaround for bug in FArchiveProxy
 
-	// Use parallel loading if requested and available; the parallel path is only for latest version
-	if (NumWorkers > 0 && Header.Version == FAssetRegistryVersion::LatestVersion)
+	// Always use one path to simplify maintenance. We don't support earlier versions anyway
+	if (Header.Version != FAssetRegistryVersion::LatestVersion)
 	{
-		ENameBatchLoadingFlags LoadingFlags = ENameBatchLoadingFlags::None;
+		SetError();
+		UE_LOG(LogAssetRegistry, Display, TEXT("Attempting to read AssetRegistry.bin file with older format version. Backwards compatibility is not supported for AssetRegistry.bin files; they should be regenerated instead. Contact Epic if you need this functionality. Version == %d (LatestVersion is %d), Archive == %s"),
+			Header.Version,
+			FAssetRegistryVersion::LatestVersion,
+			*Inner.GetArchiveName()
+			);		
+	}
+	else
+	{
+		if (NumWorkers > 0)
+		{
+			ENameBatchLoadingFlags LoadingFlags = ENameBatchLoadingFlags::None;
 #if ASSET_REGISTRY_LOAD_NAMES_IN_ORDER
-		LoadingFlags |= ENameBatchLoadingFlags::RespectOrder;
+			LoadingFlags |= ENameBatchLoadingFlags::RespectOrder;
 #endif
-		TFunction<TArray<FDisplayNameEntryId> ()> GetFutureNames = LoadNameBatchAsync(*this, NumWorkers, LoadingFlags);
+			TFunction<TArray<FDisplayNameEntryId> ()> GetFutureNames = LoadNameBatchAsync(*this, NumWorkers, LoadingFlags);
 
-		FixedTagPrivate::FAsyncStoreLoader StoreLoader;
-		Task = StoreLoader.ReadInitialDataAndKickLoad(*this, NumWorkers, Header.Version);
+			FixedTagPrivate::FAsyncStoreLoader StoreLoader;
+			StoreLoader.SetStoreBackingFile(InBackingFile);
+			Task = StoreLoader.ReadInitialDataAndKickLoad(*this, NumWorkers, Header.Version);
 		
-		Names = GetFutureNames();
-		Tags = StoreLoader.LoadFinalData(*this, Header.Version);
-	}
-	else
-	{
-		Names = LoadNameBatch(Inner, ENameBatchLoadingFlags::RespectOrder);
-		Tags = FixedTagPrivate::LoadStore(*this, Header.Version);
+			Names = GetFutureNames();
+			Tags = StoreLoader.LoadFinalData(*this, Header.Version);
+		}
+		else
+		{
+			Names = LoadNameBatch(Inner, ENameBatchLoadingFlags::RespectOrder);
+			Tags = FixedTagPrivate::LoadStore(*this, Header.Version, InBackingFile);
+		}
 	}
 }
 
@@ -422,7 +442,9 @@ bool FAssetRegistryTagSerializationTest::RunTest(const FString& Parameters)
 	{
 		FMemoryReader DataReader(Data);
 		FAssetRegistryHeader Header;
-		FAssetRegistryReader RegistryReader(DataReader, 0, Header);
+		FAssetRegistryReader RegistryReader(DataReader, nullptr, 0, Header);
+		RegistryReader.WaitForTasks();
+
 		for (FAssetDataTagMapSharedView& FixedMap : FixedMaps)
 		{
 			FixedMap = LoadTags(RegistryReader);
@@ -435,7 +457,8 @@ bool FAssetRegistryTagSerializationTest::RunTest(const FString& Parameters)
 	{
 		FMemoryReader DataReader(Data);
 		FAssetRegistryHeader Header;
-		FAssetRegistryReader RegistryReader(DataReader, 0, Header);
+		FAssetRegistryReader RegistryReader(DataReader, nullptr, 0, Header);
+		RegistryReader.WaitForTasks();
 
 		for (const FAssetDataTagMapSharedView& FixedMap1 : FixedMaps)
 		{
diff --git a/Engine/Source/Runtime/AssetRegistry/Private/AssetRegistryArchive.h b/Engine/Source/Runtime/AssetRegistry/Private/AssetRegistryArchive.h
--- a/Engine/Source/Runtime/AssetRegistry/Private/AssetRegistryArchive.h
+++ b/Engine/Source/Runtime/AssetRegistry/Private/AssetRegistryArchive.h
@@ -3,6 +3,7 @@
 #pragma once
 
 #include "AssetRegistry/AssetDataTagMapSerializationDetails.h"
+#include "AssetRegistry/AssetRegistrySerializationDetails.h"
 #include "Serialization/LargeMemoryWriter.h"
 #include "UObject/NameBatchSerialization.h"
 
@@ -22,8 +23,9 @@ struct FAssetRegistryHeader
 class FAssetRegistryReader : public FArchiveProxy
 {
 public:
+	/// @param InBackingFile shared pointer to a class that allows mapping the underlying file into memory
 	/// @param NumWorkers > 0 for parallel loading
-	FAssetRegistryReader(FArchive& Inner, int32 NumWorkers, FAssetRegistryHeader Header);
+	FAssetRegistryReader(FArchive& Inner, TSharedPtr<UE::AssetRegistry::FMemoryMappedFile> InBackingFile, int32 NumWorkers, FAssetRegistryHeader Header);
 	~FAssetRegistryReader();
 
 	virtual FArchive& operator<<(FName& Value) override;
diff --git a/Engine/Source/Runtime/AssetRegistry/Private/AssetRegistryState.cpp b/Engine/Source/Runtime/AssetRegistry/Private/AssetRegistryState.cpp
--- a/Engine/Source/Runtime/AssetRegistry/Private/AssetRegistryState.cpp
+++ b/Engine/Source/Runtime/AssetRegistry/Private/AssetRegistryState.cpp
@@ -2364,31 +2364,24 @@ bool FAssetRegistryState::Load(FArchive& OriginalAr, const FAssetRegistryLoadOpt
 	FSoftObjectPathSerializationScope SerializationScope(NAME_None, NAME_None,
 		ESoftObjectPathCollectType::NonPackage, ESoftObjectPathSerializeType::AlwaysSerialize);
 
-	if (Header.Version < FAssetRegistryVersion::RemovedMD5Hash)
+	if (Header.Version != FAssetRegistryVersion::LatestVersion)
 	{
-		// Cannot read states before this version
+		// Cannot read non-latest caches
 		return false;
 	}
-	else if (Header.Version < FAssetRegistryVersion::FixedTags)
+
+	FAssetRegistryReader Reader(OriginalAr, nullptr, Options.ParallelWorkers, Header);
+
+	if (Reader.IsError())
 	{
-		FNameTableArchiveReader NameTableReader(OriginalAr);
-		Load(NameTableReader, Header, Options);
+		return false;
 	}
-	else
-	{
-		FAssetRegistryReader Reader(OriginalAr, Options.ParallelWorkers, Header);
 
-		if (Reader.IsError())
-		{
-			return false;
-		}
+	// Load won't resolve asset registry tag values loaded in parallel
+	// and can run before WaitForTasks
+	Load(Reader, Header, Options);
 
-		// Load won't resolve asset registry tag values loaded in parallel
-		// and can run before WaitForTasks
-		Load(Reader, Header, Options);
-
-		Reader.WaitForTasks();
-	}
+	Reader.WaitForTasks();
 
 	return !OriginalAr.IsError();
 }
diff --git a/Engine/Source/Runtime/Core/Private/Windows/WindowsPlatformFile.cpp b/Engine/Source/Runtime/Core/Private/Windows/WindowsPlatformFile.cpp
--- a/Engine/Source/Runtime/Core/Private/Windows/WindowsPlatformFile.cpp
+++ b/Engine/Source/Runtime/Core/Private/Windows/WindowsPlatformFile.cpp
@@ -1591,14 +1591,18 @@ FOpenMappedResult FWindowsPlatformFile::OpenMappedEx(const TCHAR* Filename, EOpe
 {
 	int64 Size = FileSize(Filename);
 	uint32  Access = GENERIC_READ;
-	uint32  WinFlags = FILE_SHARE_READ | GetFileShareDeleteFlagIfEnabled();
+	uint32  WinFlags = FILE_SHARE_READ |
+		(EnumHasAnyFlags(OpenOptions, EOpenReadFlags::AllowDelete) ? FILE_SHARE_DELETE : 0) |
+		GetFileShareDeleteFlagIfEnabled();
 	uint32  Create = OPEN_EXISTING;
 	DWORD	Protect = PAGE_READONLY;
 
 	if (EnumHasAnyFlags(OpenOptions, EOpenReadFlags::AllowWrite))
 	{
 		Access = GENERIC_READ | GENERIC_WRITE;
-		WinFlags = FILE_SHARE_READ | FILE_SHARE_WRITE | GetFileShareDeleteFlagIfEnabled();
+		WinFlags = FILE_SHARE_READ | FILE_SHARE_WRITE |
+			(EnumHasAnyFlags(OpenOptions, EOpenReadFlags::AllowDelete) ? FILE_SHARE_DELETE : 0) |
+			GetFileShareDeleteFlagIfEnabled();
 		Create = CREATE_ALWAYS;
 		Protect = PAGE_READWRITE;
 		Size = MaximumSize;
diff --git a/Engine/Source/Runtime/CoreUObject/Private/AssetRegistry/AssetDataTagMap.cpp b/Engine/Source/Runtime/CoreUObject/Private/AssetRegistry/AssetDataTagMap.cpp
--- a/Engine/Source/Runtime/CoreUObject/Private/AssetRegistry/AssetDataTagMap.cpp
+++ b/Engine/Source/Runtime/CoreUObject/Private/AssetRegistry/AssetDataTagMap.cpp
@@ -4,7 +4,10 @@
 
 #include "Algo/Sort.h"
 #include "AssetRegistry/AssetDataTagMapSerializationDetails.h"
+#include "AssetRegistry/AssetRegistrySerializationDetails.h"
 #include "AutoRTFM.h"
+#include "HAL/PlatformFile.h"
+#include "HAL/PlatformFileManager.h"
 #include "Misc/PackageName.h"
 #include "Misc/Optional.h"
 #include "Misc/TransactionallySafeCriticalSection.h"
@@ -12,12 +15,45 @@
 #include "Serialization/ArrayWriter.h"
 #include "Serialization/MemoryReader.h"
 #include "Serialization/MemoryWriter.h"
+
 #include <limits>
 
 DEFINE_LOG_CATEGORY_STATIC(LogAssetDataTags, Log, All);
 
 //////////////////////////////////////////////////////////////////////////
 
+void UE::AssetRegistry::AlignPosInArchive(FArchive& Archive)
+{
+	// see if current offset is aligned
+	int64 Pos = Archive.Tell();
+	int64 AlignedPos = Align(Pos, AlignmentForMemoryMappedArchive);
+	if (AlignedPos != Pos)
+	{
+		int64 DataToSkip = AlignedPos - Pos;
+
+		if (Archive.IsSaving())
+		{
+			char Dunnage[AlignmentForMemoryMappedArchive] = { 0 };
+
+#if 0 // for easier debugging			
+			for (SIZE_T Idx = 0; Idx < sizeof(Dunnage); ++Idx)
+			{
+				Dunnage[Idx] = static_cast<char>(Idx);
+			}
+#endif
+			Archive.Serialize(Dunnage, DataToSkip);
+		}
+		else
+		{
+			Archive.Seek(AlignedPos);
+		}
+
+		checkf(Archive.Tell() == AlignedPos, TEXT("We were not able to align the position in the archive to the desired boundary."));
+	}
+}
+
+//////////////////////////////////////////////////////////////////////////
+
 template<class T>
 FAssetRegistryExportPath ParseExportPath(T ExportPath)
 {
@@ -999,14 +1035,36 @@ namespace FixedTagPrivate
 		};
 	}
 
+	void* FStore::AllocateMemoryForView(SIZE_T NumBytes, SIZE_T AlignmentRequired)
+	{
+		void* Memory = FMemory::Malloc(NumBytes, static_cast<uint32>(AlignmentRequired));
+		if (LIKELY(Memory))
+		{
+			FScopeLock Lock(&ViewPointersCS);
+			ViewPointers.Add(Memory);
+		}
+		return Memory;
+	}
+
+
 	FStore::~FStore()
 	{
 		GStores.Unregister(*this);
 
-		if (Data)
+		// not locking the pointers since no other thread should be using the store at this point
+		VisitViews(*this, [this](auto& View)
+			{
+				// If the view is not memory-mmapped, destroy its elements. Memory-mapped views have no-op destructors
+				if (ViewPointers.Contains(View.GetData()))
+				{
+					DestroyElements(View);
+				}
+			}
+		);
+
+		for (void* Ptr : ViewPointers)
 		{
-			VisitViews(*this, [] (auto View) { DestroyElements(View); });
-			FMemory::Free(Data);
+			FMemory::Free(Ptr);
 		}
 	}
 
@@ -1182,6 +1240,110 @@ namespace FixedTagPrivate
 		void LoadViewData(TArrayView<WIDECHAR> View) { BulkSerialize(View); }
 		static_assert(PLATFORM_LITTLE_ENDIAN, "Byte-swapping WIDECHAR on load needed to load on big-endian platforms");
 
+		template<class T>
+		inline void LoadViewFromMemoryWithoutAttemptingToMap(FStore& Store, TArrayView<T>& View)
+		{
+			uint64 Bytes = GetBytes(View);
+			if (Bytes)
+			{
+				SetUntypedDataPtr(View, Store.AllocateMemoryForView(Bytes, alignof(T)));
+				MemZeroViewDataIfNeeded(View);
+				LoadViewData(View);
+			}
+		}
+
+		// by default, assume the type cannot be memory mapped
+		template<class T>
+		void LoadViewFromMemoryMappedArchive(FStore& Store, TArrayView<T>& View)
+		{
+			LoadViewFromMemoryWithoutAttemptingToMap(Store, View);
+		}
+
+		void LoadViewFromMemoryMappedArchive(FStore& Store, TArrayView<uint32>& View)
+		{
+#if WITH_EDITOR
+			if (Store.BackingFile.IsValid())
+			{
+				if (int64 Size = GetBytes(View))
+				{
+					int64 Offset = Ar.Tell();
+					FMemoryView Part = Store.BackingFile->GetSubView(Offset, Size);
+					if (void* DataPtr = const_cast<void*>(Part.GetData()))
+					{
+						SetUntypedDataPtr(View, DataPtr);
+						Ar.Seek(Offset + Size);
+					}
+					else
+					{
+						// in an unlikely case we could not map, switch to allocation
+						LoadViewFromMemoryWithoutAttemptingToMap(Store, View);
+					}
+				}
+			}
+			else
+#endif
+			{
+				LoadViewFromMemoryWithoutAttemptingToMap(Store, View);
+			}
+		}
+
+		// ANSI and WIDE chars can be memory mapped
+		void LoadViewFromMemoryMappedArchive(FStore& Store, TArrayView<ANSICHAR>& View)
+		{
+#if WITH_EDITOR
+			if (Store.BackingFile.IsValid())
+			{
+				if (int64 Size = GetBytes(View))
+				{
+					int64 Offset = Ar.Tell();
+					FMemoryView Part = Store.BackingFile->GetSubView(Offset, Size);
+					if (void* DataPtr = const_cast<void*>(Part.GetData()))
+					{
+						SetUntypedDataPtr(View, DataPtr);
+						Ar.Seek(Offset + Size);
+					}
+					else
+					{
+						// in an unlikely case we could not map, switch to allocation
+						LoadViewFromMemoryWithoutAttemptingToMap(Store, View);
+					}
+				}
+			}
+			else
+#endif
+			{
+				LoadViewFromMemoryWithoutAttemptingToMap(Store, View);
+			}
+		}
+		void LoadViewFromMemoryMappedArchive(FStore& Store, TArrayView<WIDECHAR>& View)
+		{
+#if WITH_EDITOR
+			if (Store.BackingFile.IsValid())
+			{
+				if (int64 Size = GetBytes(View))
+				{
+					int64 Offset = Ar.Tell();
+					FMemoryView Part = Store.BackingFile->GetSubView(Offset, Size);
+					if (void* DataPtr = const_cast<void*>(Part.GetData()))
+					{
+						SetUntypedDataPtr(View, DataPtr);
+						Ar.Seek(Offset + Size);
+					}
+					else
+					{
+						// in an unlikely case we could not map, switch to allocation
+						LoadViewFromMemoryWithoutAttemptingToMap(Store, View);
+					}
+				}
+			}
+			else
+#endif
+			{
+				LoadViewFromMemoryWithoutAttemptingToMap(Store, View);
+			}
+		}
+
+
 		FArchive& Ar;
 		FString Scratch;
 
@@ -1200,15 +1362,19 @@ namespace FixedTagPrivate
 			FArrayWriter Data;
 			TSerializer<Version>(Data).SaveViewData(Texts);
 			SaveItem(Data.Num());
+			UE::AssetRegistry::AlignPosInArchive(Ar);
 			Ar.Serialize(Data.GetData(), Data.Num());
 		}
 
 		void Save(const FStoreData& Store)
 		{
+			// Store is expected to be saved at an aligned position
+			check(UE::AssetRegistry::IsPosInArchiveAligned(Ar));
 			SaveItem(BeginMagic);
 			VisitViews(Store, [&](auto& Array) { SaveItem(Array.Num()); });
 			SaveTextData(MakeArrayView(Store.Texts));
-			VisitViews<EOrder::SkipText>(Store, [&] (auto& Array) { SaveViewData(MakeArrayView(Array)); });
+			// while not all data will be mmaped, give all of it a possibility to be so by aligning everything
+			VisitViews<EOrder::SkipText>(Store, [&] (auto& Array) { UE::AssetRegistry::AlignPosInArchive(Ar); SaveViewData(MakeArrayView(Array)); });
 			SaveItem(EndMagic);
 		}
 
@@ -1224,6 +1390,8 @@ namespace FixedTagPrivate
 
 		TOptional<ELoadOrder> LoadHeader(FStore& Store)
 		{
+			// store is expected to be saved at an aligned position, so it can align its data
+			UE::AssetRegistry::AlignPosInArchive(Ar);
 			uint32 InitialMagic = LoadItem<uint32>();
 			TOptional<ELoadOrder> Order = GetLoadOrder(InitialMagic);
 
@@ -1237,29 +1405,6 @@ namespace FixedTagPrivate
 			// Load view sizes
 			VisitViews(Store, [&] (auto& View) { SetNum(View, LoadItem<int32>()); });
 
-			// Calculate total size and allocate data
-			uint64 Bytes = 0;
-			VisitViews(Store, [&] (auto& View) 
-				{
-					static_assert(alignof(typename std::decay_t<decltype(View)>::ElementType) <= MaxViewAlignment, "");
-					Bytes = Align(Bytes, View.GetTypeAlignment()) + GetBytes(View);
-				});
-
-			uint8* Ptr = reinterpret_cast<uint8*>(FMemory::Malloc(Bytes, MaxViewAlignment));
-			Store.Data = Ptr;
-
-			// Set view data pointers
-			VisitViews(Store, [&] (auto& View)
-				{
-					Ptr = Align(Ptr, View.GetTypeAlignment());
-					SetUntypedDataPtr(View, Ptr);
-					Ptr += GetBytes(View);
-				});
-			check(Ptr - Bytes == Store.Data);
-
-			// Zero out required data
-			VisitViews(Store, [&](auto& View) { MemZeroViewDataIfNeeded(View); });
-
 			return Order;
 		}
 		
@@ -1268,16 +1413,10 @@ namespace FixedTagPrivate
 			if (TOptional<ELoadOrder> Order = LoadHeader(/* Out */ Store))
 			{
 				// Load view data
-				if (Order.GetValue() == ELoadOrder::TextFirst)
-				{
-					uint32 TextDataBytes = LoadItem<uint32>();
+				checkf(Order.GetValue() == ELoadOrder::TextFirst, TEXT("Loading old versions of cache is no longer supported"));
+				uint32 TextDataBytes = LoadItem<uint32>();
 
-					VisitViews<EOrder::TextFirst>(Store, [&] (auto View) { LoadViewData(View); });
-				}
-				else
-				{
-					VisitViews(Store, [&] (auto View) { LoadViewData(View); });
-				}
+				VisitViews<EOrder::TextFirst>(Store, [&] (auto& View) { UE::AssetRegistry::AlignPosInArchive(Ar); LoadViewFromMemoryMappedArchive(Store, View); });
 
 				if (LoadItem<uint32>() != EndMagic)
 				{
@@ -1292,25 +1431,20 @@ namespace FixedTagPrivate
 			uint32 TextDataBytes = LoadItem<uint32>();
 			TArray<uint8> Out;
 			Out.SetNumUninitialized(TextDataBytes);
+			UE::AssetRegistry::AlignPosInArchive(Ar);
 			Ar.Serialize(Out.GetData(), Out.Num());
 			return Out;
 		}
 
 		void LoadTextData(FStore& Store)
 		{
-			LoadViewData(Store.Texts);
+			LoadViewFromMemoryMappedArchive(Store, Store.Texts);
 		}
 
 		void LoadFinalData(FStore& Store, ELoadOrder Order)
 		{
-			if (Order == ELoadOrder::TextFirst)
-			{
-				VisitViews<EOrder::SkipText>(Store, [&] (auto View) { LoadViewData(View); });
-			}
-			else
-			{
-				VisitViews(Store, [&] (auto View) { LoadViewData(View); });
-			}
+			checkf(Order == ELoadOrder::TextFirst, TEXT("Loading old versions of cache is no longer supported"));
+			VisitViews<EOrder::SkipText>(Store, [&] (auto& View) { UE::AssetRegistry::AlignPosInArchive(Ar); LoadViewFromMemoryMappedArchive(Store, View); });
 			
 			if (LoadItem<uint32>() != EndMagic)
 			{
@@ -1326,7 +1460,7 @@ namespace FixedTagPrivate
 		TSerializer<FAssetRegistryVersion::LatestVersion>(Ar).Save(Store);
 	}
 
-	TRefCountPtr<const FStore> LoadStore(FArchive& Ar, FAssetRegistryVersion::Type Version /*= FAssetRegistryVersion::LatestVersion*/)
+	TRefCountPtr<const FStore> LoadStore(FArchive& Ar, FAssetRegistryVersion::Type Version /*= FAssetRegistryVersion::LatestVersion*/, TSharedPtr<UE::AssetRegistry::FMemoryMappedFile> InBackingFile /*= nullptr*/)
 	{
 		if (Ar.IsError())
 		{
@@ -1334,26 +1468,12 @@ namespace FixedTagPrivate
 		}
 
 		FStore* Store = GStores.CreateAndRegister();
+#if WITH_EDITOR
+		Store->BackingFile = InBackingFile;
+#endif
 
-		// Versioning support inside of TSerializer was added with FAssetRegistryVersion::ClassPaths so everything before can just be
-		// serialized with FAssetRegistryVersion::ClassPaths - 1 = FAssetRegistryVersion::AddedChunkHashes.
-		// The main purpose of hardcoding version value and passing it as template argument is to eliminate
-		// branches in performance critical serialization code.
-		// Note that we branch in this function only because it's the only function used to load legacy 
-		// asset registries in DiffAssetRegistries commandlet hence it's editor only code too
-		if (Version < FAssetRegistryVersion::ClassPaths)
-		{
-			TSerializer<FAssetRegistryVersion::AddedChunkHashes>(Ar).Load(*Store);
-		}
-		else if (Version <= FAssetRegistryVersion::AssetPackageDataHasPackageLocation)
-		{
-			TSerializer<FAssetRegistryVersion::AssetPackageDataHasPackageLocation>(Ar).Load(*Store);
-		}
-		else // Add more branches should serialization format change
-		{
-			TSerializer<FAssetRegistryVersion::LatestVersion>(Ar).Load(*Store);
-		}
-		
+		// We only support loading the latest version. All the past caches will need to be re-gathered.
+		TSerializer<FAssetRegistryVersion::LatestVersion>(Ar).Load(*Store);
 		return TRefCountPtr<const FStore>(Store);
 	}
 
@@ -1367,20 +1487,19 @@ namespace FixedTagPrivate
 		check(HeaderVersion == FAssetRegistryVersion::LatestVersion);
 		Order = TSerializer<FAssetRegistryVersion::LatestVersion>(Ar).LoadHeader(*Store);
 
-		if (Order == ELoadOrder::TextFirst)
-		{
-			TArray<uint8> TextData = TSerializer<FAssetRegistryVersion::LatestVersion>(Ar).ReadTextData();
+		checkf(Order == ELoadOrder::TextFirst, TEXT("Reading old formats isn't supported."));
+
+		UE::AssetRegistry::AlignPosInArchive(Ar);
+		TArray<uint8> TextData = TSerializer<FAssetRegistryVersion::LatestVersion>(Ar).ReadTextData();
 
-			if (TextData.Num())
+		if (TextData.Num())
+		{
+			return Async(EAsyncExecution::TaskGraph, [Data = MoveTemp(TextData), OutStore = Store]()
 			{
-				return Async(EAsyncExecution::TaskGraph, [Data = MoveTemp(TextData), OutStore = Store]()
-				{
-					FMemoryReader Reader(Data);
-					TSerializer<FAssetRegistryVersion::LatestVersion>(Reader).LoadTextData(*OutStore);
-				});
-			}
+				FMemoryReader Reader(Data);
+				TSerializer<FAssetRegistryVersion::LatestVersion>(Reader).LoadTextData(*OutStore);
+			});
 		}
-
 		return TFuture<void>();
 	}
 
@@ -1734,6 +1853,41 @@ void FAssetDataTagMapSharedView::FMemoryCounter::ReportFixedStoreBreakdown() con
 	}
 }
 
+namespace UE::AssetRegistry
+{
+	FMemoryMappedFile::FMemoryMappedFile(const TCHAR* Path)
+	{
+		FOpenMappedResult Result = FPlatformFileManager::Get().GetPlatformFile().OpenMappedEx(Path, IPlatformFile::EOpenReadFlags::AllowDelete);
+		if (!Result.HasError())
+		{
+			Handle = Result.StealValue();
+		}
+	}
+
+	FMemoryView FMemoryMappedFile::GetSubView(uint64 Start, uint64 Length) const
+	{
+		FMappedSubRegion New;
+		New.Start = Start;
+		New.Length = Length;
+
+		if (!Handle)
+		{
+			return FMemoryView();
+		}
+
+		New.Region.Reset(Handle->MapRegion(Start, Length));
+		if (!New.Region)
+		{
+			return FMemoryView();
+		}
+
+		FMemoryView Result = FMemoryView(New.Region->GetMappedPtr(), New.Region->GetMappedSize());
+		Regions.Add(MoveTemp(New));
+		return Result;
+	}
+
+} // namespace UE::AssetRegistry
+
 ////////////////////////////////////////////////////////////////////////////
 
 #if WITH_DEV_AUTOMATION_TESTS 
diff --git a/Engine/Source/Runtime/CoreUObject/Public/AssetRegistry/AssetData.h b/Engine/Source/Runtime/CoreUObject/Public/AssetRegistry/AssetData.h
--- a/Engine/Source/Runtime/CoreUObject/Public/AssetRegistry/AssetData.h
+++ b/Engine/Source/Runtime/CoreUObject/Public/AssetRegistry/AssetData.h
@@ -105,6 +105,8 @@ struct FAssetRegistryVersion
 		PackageSavedHash,					// Replaced FAssetPackageData::PackageGuid with PackageSavedHash
 		ExternalActorToWorldIsEditorOnly,   // FPackageDependencyData::LoadDependenciesFromPackageHeader changed how it calculates PackageDependencies
 		ManageDependenciesCookRule,			// CookRule property added to EDependencyProperties for EDependencyCategory::Manage
+		TagSetMapSupports64BitStringNum,	// 64-bit arrays for strings
+		MemoryMappedTagDataStore,			// Tag data store is now memory mapped
 
 		// -----<new versions can be added above this line>-------------------------------------------------
 		VersionPlusOne,
diff --git a/Engine/Source/Runtime/CoreUObject/Public/AssetRegistry/AssetDataTagMapSerializationDetails.h b/Engine/Source/Runtime/CoreUObject/Public/AssetRegistry/AssetDataTagMapSerializationDetails.h
--- a/Engine/Source/Runtime/CoreUObject/Public/AssetRegistry/AssetDataTagMapSerializationDetails.h
+++ b/Engine/Source/Runtime/CoreUObject/Public/AssetRegistry/AssetDataTagMapSerializationDetails.h
@@ -4,6 +4,7 @@
 
 #include "AssetRegistry/AssetData.h"
 #include "AssetRegistry/AssetDataTagMap.h"
+#include "AssetRegistrySerializationDetails.h"
 #include "Async/Async.h"
 #include "AutoRTFM.h"
 #include "HAL/ThreadSafeCounter.h"
@@ -62,7 +63,16 @@ namespace FixedTagPrivate
 		TArrayView<FMarshalledText> Texts;
 
 		const uint32 Index;
-		void* Data = nullptr;
+
+#if WITH_EDITOR
+		/**
+		 * File where the data is kept, if any. If nullptr, the archive is not using memory-mapped files. Editor cache
+		 * files usually use memory-mapped files, but in-game cooked AssetRegistry.bin data do not.
+		 */
+		mutable TSharedPtr<UE::AssetRegistry::FMemoryMappedFile> BackingFile;
+#endif
+		/** Allocates a pointer from the heap if a view cannot use the memory-mapped data in place. */
+		void* AllocateMemoryForView(SIZE_T NumBytes, SIZE_T AlignmentRequired);
 
 		void AddRef() const
 		{
@@ -93,6 +103,11 @@ namespace FixedTagPrivate
 		~FStore();
 
 		mutable FThreadSafeCounter RefCount;
+
+		/** Guards access to the View pointers. No need for an R/O as we're only accessing the set when we're about to modify it. */
+		FCriticalSection ViewPointersCS;
+		/** Pointers to individual views that aren't mmapped but allocated from the heap */
+		TSet<void*> ViewPointers;
 	};
 
 	struct FOptions
@@ -255,7 +270,7 @@ namespace FixedTagPrivate
 	enum class ELoadOrder { Member, TextFirst };
 
 	COREUOBJECT_API void SaveStore(const FStoreData& Store, FArchive& Ar);
-	COREUOBJECT_API TRefCountPtr<const FStore> LoadStore(FArchive& Ar, FAssetRegistryVersion::Type Version = FAssetRegistryVersion::LatestVersion);
+	COREUOBJECT_API TRefCountPtr<const FStore> LoadStore(FArchive& Ar, FAssetRegistryVersion::Type Version = FAssetRegistryVersion::LatestVersion, TSharedPtr<UE::AssetRegistry::FMemoryMappedFile> InBackingFile = nullptr);
 
 	/// Loads tag store with async creation of expensive tag values
 	///
@@ -280,6 +295,19 @@ namespace FixedTagPrivate
 		/// @return indexed store, usable for FPartialMapHandle::MakeFullHandle()
 		COREUOBJECT_API TRefCountPtr<const FStore> LoadFinalData(FArchive& Ar, FAssetRegistryVersion::Type HeaderVersion);
 
+		/** Passes backing file for memory mapping to the store. */
+		inline void SetStoreBackingFile(TSharedPtr<UE::AssetRegistry::FMemoryMappedFile> InBackingFile)
+		{
+#if WITH_EDITOR
+			Store->BackingFile = MoveTemp(InBackingFile);
+#else
+			if (InBackingFile)
+			{
+				UE_LOG(LogAssetData, Error, TEXT("Reading asset registry from a memory mapped file is not supported in a non-editor binary."));
+			}
+#endif
+		}
+
 	private:
 		TRefCountPtr<FStore> Store;
 		TOptional<ELoadOrder> Order;
diff --git a/Engine/Source/Runtime/CoreUObject/Public/AssetRegistry/AssetRegistrySerializationDetails.h b/Engine/Source/Runtime/CoreUObject/Public/AssetRegistry/AssetRegistrySerializationDetails.h
new file mode 100644
--- /dev/null
+++ b/Engine/Source/Runtime/CoreUObject/Public/AssetRegistry/AssetRegistrySerializationDetails.h
@@ -0,0 +1,49 @@
+// Copyright Epic Games, Inc. All Rights Reserved.
+
+#pragma once
+
+#include "Serialization/ArchiveProxy.h"
+#include "Async/MappedFileHandle.h"
+#include "Memory/MemoryFwd.h"
+#include "Templates/SharedPointer.h"
+#include "Serialization/MemoryHasher.h"
+
+namespace UE::AssetRegistry
+{
+	/** Memory maps the bytes of a file on disk into memory and provides MemoryViews to those bytes until *this is destructed. */
+	class FMemoryMappedFile
+	{
+	public:
+		FMemoryMappedFile() = default;
+		COREUOBJECT_API explicit FMemoryMappedFile(const TCHAR* Path);
+
+		/** Provides a new View, allocating it if it doesn't exist. */
+		COREUOBJECT_API FMemoryView GetSubView(uint64 Start, uint64 Length) const;
+
+	private:
+
+		/** Handle for the mapped file */
+		TUniquePtr<IMappedFileHandle> Handle;
+
+		struct FMappedSubRegion
+		{
+			uint64 Start;
+			uint64 Length;
+
+			TUniquePtr<IMappedFileRegion> Region;
+		};
+
+		mutable TArray<FMappedSubRegion>  Regions;
+	};
+
+	inline constexpr int64 AlignmentForMemoryMappedArchive = 16;
+
+	/** When saving, saves extra memory to make sure the following data is aligned. When loading, seeks over that many bytes (does not load them). */
+	void COREUOBJECT_API AlignPosInArchive(FArchive& Archive);
+
+	/** Reports whether the archive's current offset is aligned, relative to offset 0. */
+	inline bool IsPosInArchiveAligned(FArchive& Archive)
+	{
+		return IsAligned(Archive.Tell(), AlignmentForMemoryMappedArchive);
+	}
+}
