Map tag store data in Asset Registry cache, saving GBs of editor memory usage on large projects.

- Added support for older versions of the AR cache.

#was-backout 50479750

#Synced-CL 50510093

diff --git a/Engine/Source/Runtime/AssetRegistry/Private/AssetDataGatherer.cpp b/Engine/Source/Runtime/AssetRegistry/Private/AssetDataGatherer.cpp
--- a/Engine/Source/Runtime/AssetRegistry/Private/AssetDataGatherer.cpp
+++ b/Engine/Source/Runtime/AssetRegistry/Private/AssetDataGatherer.cpp
@@ -2,8 +2,9 @@
 
 #include "AssetDataGatherer.h"
 #include "AssetDataGathererPrivate.h"
-#include "AssetRegistryImpl.h"
 #include "AssetRegistry.h"
+#include "AssetRegistryImpl.h"
+#include "AssetRegistry/AssetDataTagMapSerializationDetails.h"
 #include "AssetRegistry/AssetRegistryTelemetry.h"
 
 #include "Algo/AnyOf.h"
@@ -13,7 +14,6 @@
 #include "AssetRegistry/AssetData.h"
 #include "AssetRegistryArchive.h"
 #include "AssetRegistryPrivate.h"
-#include "Async/MappedFileHandle.h"
 #include "Async/ParallelFor.h"
 #include "Containers/BinaryHeap.h"
 #include "HAL/LowLevelMemTracker.h"
@@ -46,6 +46,7 @@
 #include "Templates/UnrealTemplate.h"
 #include "Trace/Trace.h"
 #include "Trace/Trace.inl"
+#include "Serialization/MemoryHasher.h"
 
 namespace UE::AssetDataGather::Private
 {
@@ -62,6 +63,12 @@ static FAutoConsoleVariableRef CVarIgnoreEmptyDirectories(
 	bIgnoreEmptyDirectories,
 	TEXT("If true, completely empty leaf directories are ignored by the asset registry while scanning"));
 
+bool bMemoryMappedCacheFiles = true;
+static FAutoConsoleVariableRef CVarMemoryMappedCacheFiles(
+	TEXT("AssetRegistry.MemoryMappedCacheFiles"),
+	bMemoryMappedCacheFiles,
+	TEXT("If true, asset registry cache file will be memory mapped, reducing memory usage"));
+
 void LexFromString(EFeatureEnabledReadWrite& OutValue, FStringView Text)
 {
 	Text.TrimStartAndEndInline();
@@ -3345,6 +3352,10 @@ FAssetDataDiscovery::FDirectoryResult::FDirectoryResult(FStringView InDirAbsPath
 
 }
 
+// this namespace contains deprecated code
+namespace AssetRegistryBackwardCompat
+{
+
 // Reads an FMemoryView once
 class FMemoryViewReader
 {
@@ -3607,14 +3618,6 @@ private:
 	virtual void Seek(int64 InPos) override								{ unimplemented(); }
 };
 
-class FChecksumArchiveWriter : public FChecksumArchiveBase
-{
-public:
-	using FChecksumArchiveBase::FChecksumArchiveBase;
-	virtual void Serialize(void* V, int64 Len) override { Save(FMemoryView(V, Len)); }
-	virtual int64 Tell() override { return InnerArchive.Tell() + GetCurrentBlock().GetUsedSize(); }
-};
-
 class FChecksumArchiveReader : public FChecksumArchiveBase
 {
 public:
@@ -3702,37 +3705,7 @@ private:
 	}
 };
 
-// Util that maps an entire file
-class FMemoryMappedFile
-{
-public:
-	explicit FMemoryMappedFile(const TCHAR* Path)
-	{
-		FOpenMappedResult Result = FPlatformFileManager::Get().GetPlatformFile().OpenMappedEx(Path);
-		if (!Result.HasError())
-		{
-			Handle = Result.StealValue();
-			Region.Reset(Handle->MapRegion());
-		}
-	}
-
-	void Preload(int64 Size = MAX_int64) const
-	{
-		if (Region)
-		{
-			Region->PreloadHint(0, Size);
-		}
-	}
-	
-	FMemoryView View() const
-	{
-		return Region ? FMemoryView(Region->GetMappedPtr(), Region->GetMappedSize()) : FMemoryView();
-	}
-
-private:
-	TUniquePtr<IMappedFileHandle> Handle;
-	TUniquePtr<IMappedFileRegion> Region;
-};
+}
 
 #if UE_EDITOR
 /** A class to preload the AssetData cache and start the AssetDataDiscovery thread (which does its own cache loading)
@@ -5324,21 +5297,28 @@ int64 FAssetDataGatherer::SaveCacheFileInternal(const FString& CacheFilename, co
 	TUniquePtr<FArchive> FileAr(IFileManager::Get().CreateFileWriter(*TempFilename, 0));
 	if (FileAr)
 	{
-		using namespace UE::AssetDataGather::Private;
-		uint64 CurrentVersion = FGatherSettings::CurrentVersion;
-		*FileAr << CurrentVersion;
-
-#if ALLOW_NAME_BATCH_SAVING
 		{
-			// We might be able to reduce load time by using AssetRegistry::SerializationOptions
-			// to save certain common tags as FName.
-			UE::AssetDataGather::Private::FChecksumArchiveWriter ChecksummingWriter(*FileAr);
-			FAssetRegistryWriter Ar(FAssetRegistryWriterOptions(), ChecksummingWriter);
-			UE::AssetDataGather::Private::SerializeCacheSave(Ar, AssetsToSave);
+			UE::THashingArchiveProxy<FBlake3> FileArHash(*FileAr);
+
+			using namespace UE::AssetDataGather::Private;
+			uint64 CurrentVersion = FGatherSettings::CurrentVersion;
+			FileArHash << CurrentVersion;
+
+	#if ALLOW_NAME_BATCH_SAVING
+			{
+				// We might be able to reduce load time by using AssetRegistry::SerializationOptions
+				// to save certain common tags as FName.
+				FAssetRegistryWriter Ar(FAssetRegistryWriterOptions(), FileArHash);
+				UE::AssetDataGather::Private::SerializeCacheSave(Ar, AssetsToSave);
+			}
+	#else		
+			static_assert(false, "Cannot save asset registry cache in this configuration");
+	#endif
+			// write the integrity checksum
+			FBlake3Hash Hash = FileArHash.GetHash();
+			FileAr->Serialize(&Hash, sizeof(Hash));
 		}
-#else		
-		checkf(false, TEXT("Cannot save asset registry cache in this configuration"));
-#endif
+
 		int64 CacheSize = FileAr->TotalSize();
 		// Close file handle before moving temp file to target 
 		FileAr.Reset();
@@ -5440,7 +5420,13 @@ TArray<FCachePayload> LoadCacheFiles(TConstArrayView<FString> InCacheFilenames)
 	TArray<FCachePayload> Results;
 	Results.AddDefaulted(InCacheFilenames.Num());
 
-	auto DoLoad = [](FArchive& ChecksummingReader)
+	// disable memory mapping on the builder machines until the situation with deletion clarifies
+	if (GIsBuildMachine)
+	{
+		UE::AssetDataGather::Private::bMemoryMappedCacheFiles = false;
+	}
+
+	auto DoLoad = [](FArchive& FilemappedReader, TSharedPtr<UE::AssetRegistry::FMemoryMappedFile> BackingFile, FAssetRegistryVersion::Type VersionToRead = FAssetRegistryVersion::LatestVersion)
 		{
 			// We are already using AsyncTasks for our own ParallelFor, passing non-zero NumAsyncWorkers into
 			// FAssetRegistryReader will cause our current task to block waiting for another task to run
@@ -5448,181 +5434,152 @@ TArray<FCachePayload> LoadCacheFiles(TConstArrayView<FString> InCacheFilenames)
 			// if we don't have a large number of threads, so disable the parallelism.
 			constexpr int32 NumAsyncWorkers = 0;
 
-			// The discovery cache is always serialized with a fixed format.
-			// We discard it before this point if it's not the latest version, and it always includes editor-only data.
+			// The discovery cache always includes editor-only data.
 			FAssetRegistryHeader Header;
-			Header.Version = FAssetRegistryVersion::LatestVersion;
+			Header.Version = VersionToRead;
 			Header.bFilterEditorOnlyData = false;
-			FAssetRegistryReader RegistryReader(ChecksummingReader, NumAsyncWorkers, Header);
+
+			FAssetRegistryReader RegistryReader(FilemappedReader, UE::AssetDataGather::Private::bMemoryMappedCacheFiles ? BackingFile : nullptr, NumAsyncWorkers, Header);
 			return RegistryReader.IsError() ? FCachePayload() : SerializeCacheLoad(RegistryReader);
 		};
 
-	if (FPlatformProperties::SupportsMemoryMappedFiles())
-	{
-		FCriticalSection ParallelForLock;
-		struct FSharedFileData
+	auto VerifyIntegrity = [](FArchive* Ar)
 		{
-			void ConditionalConstruct(const TCHAR* FilePath)
-			{
-				if (!File)
-				{
-					File.Emplace(FilePath);
-				}
-			}
-			void ConditionalQueueDestruction(TArray<FSharedFileData*>& Queue)
+			// same size as in FMD5Hash::HashFileFromArchive
+			constexpr int32 kLocalScratchSize = 64 * 1024;
+			TArray<uint8> LocalScratch;
+			LocalScratch.SetNumUninitialized(kLocalScratchSize);
+
+			// we're going to hash all the file except the last N bytes which hold the hash we need to compare
+			constexpr int32 kHashSize = sizeof(FBlake3Hash);
+			const int64 SizeToHash = Ar->TotalSize() - kHashSize;
+
+			if (SizeToHash > 0)
 			{
-				if (bLoadComplete && bAsyncCacheComplete)
+				int64 Position = 0;
+				FBlake3Hash ComputedHash = [&]()
+					{
+						THashingArchiveProxy<FBlake3> HashingAr(*Ar);
+						// Read in BufferSize chunks
+						while (Position < SizeToHash)
+						{
+							const auto ReadNum = FMath::Min(SizeToHash - Position, (int64)LocalScratch.Num());
+							HashingAr.Serialize(LocalScratch.GetData(), ReadNum);
+
+							Position += ReadNum;
+						};
+
+						return HashingAr.GetHash();
+					}();
+
+				int64 HashToRead = Ar->TotalSize() - Position;
+				if (HashToRead == sizeof(ComputedHash))
 				{
-					Queue.Add(this);
+					FBlake3Hash ReadHash;
+					Ar->Serialize(&ReadHash, sizeof(ReadHash));
+					Ar->Seek(0);	// return to the beginning so if verification succeeds, we can proceed to read the file as usual
+					return (ReadHash == ComputedHash);
 				}
 			}
-			TOptional<FMemoryMappedFile> File;
-			bool bLoadComplete = false;
-			bool bAsyncCacheComplete = false;
-		};
 
-		// Allocate multiple ParallelFor body invocations for each cache file.
-		// Each group of ParallelFor body invocations will work on the same cache file together in parallel,
-		// or in series if we run out of threads.
-		enum class ETaskType : uint8
-		{
-			// Load byts from the MemoryMapped file and parse its data into Results
-			DoLoad = 0,
-			// Call Preload on the MemoryMapped file to swap its bytes into memory.
-			// This starts at the same time or later than DoLoad, but it will outpace DoLoad
-			// because DoLoad is doing a lot of cpuwork, and so the AsyncCache thread will 
-			// be the one that incurs IO costs rather than the DoLoad thread.
-			AsyncCache,
-			Count,
+			return false;
 		};
-		constexpr int32 TasksPerResult = (int32)ETaskType::Count;
-		TArray<FSharedFileData> CacheFiles;
-		TArray<FSharedFileData*> DestructQueue;
-		CacheFiles.SetNum(InCacheFilenames.Num());
-		DestructQueue.Reserve(InCacheFilenames.Num());
-		const int32 NumTasks = InCacheFilenames.Num() * TasksPerResult;
-		int32 NextTaskIndex = 0;
-
-		// We want to restrict the number of threads to reduce maximum memory use due to having multiple intermediate cache files
-		// in parallel. But ParallelFor only provides an API for restricting MinBatchSize.
-		const int32 MaxNumThreads = FGatherSettings::GARGatherCacheParallelism * TasksPerResult;
+
+	if (FPlatformProperties::SupportsMemoryMappedFiles())
+	{
+		const int32 NumTasks = InCacheFilenames.Num();
+		const int32 MaxNumThreads = FGatherSettings::GARGatherCacheParallelism;
 		const int32 MinBatchSize = FMath::DivideAndRoundUp(NumTasks, MaxNumThreads);
 
 		ParallelFor(TEXT("AssetDataGatherer::LoadCacheFiles"), NumTasks, MinBatchSize,
-			[&DoLoad, &ParallelForLock, &InCacheFilenames, &Results, &CacheFiles, &DestructQueue, &NextTaskIndex]
-			(int32 /* Unused ParallelFor Index */)
+			[&DoLoad, &InCacheFilenames, &Results, &VerifyIntegrity](int32 Index)
 		{
-			FSharedFileData* FileData = nullptr;
-			const FString* CacheFilename = nullptr;
-			bool bShouldAsyncCache = false;
+			const FString& CacheFilename = InCacheFilenames[Index];
 
-			int32 ResultIndex = -1;
-			ETaskType TaskType = ETaskType::DoLoad;
-			FSharedFileData* DestructData = nullptr;
+			FCachePayload Payload;
+			bool bVerifiedOk = false;
 
-			do
+			// First, verify the file and refuse to read if hashes do not agree
+			// We do it before memory mapping the file to avoid bringing into the process memory areas of the file
+			// that we do not expect to touch, and would have to explicitly evict later.
+			TUniquePtr<FArchive> Ar(IFileManager::Get().CreateFileReader(*CacheFilename));
+			if (Ar)
 			{
-				// Pick up the next TaskIndex, but before doing that, destruct any previously completed CacheFiles,
-				// to reduce the maximum amount of memory we have allocated at once across the multiple threads.
-				DestructData = nullptr;
-				{
-					FScopeLock ParallelForScopeLock(&ParallelForLock);
-
-					// Peek at what our task would be but don't claim it yet.
-					int32 TaskIndex = NextTaskIndex;
-					ResultIndex = TaskIndex / TasksPerResult;
-					TaskType = (ETaskType)(TaskIndex - ResultIndex * TasksPerResult);
-					check(0 <= (int32)TaskType && TaskType < ETaskType::Count);
-
-					if (!DestructQueue.IsEmpty()
-						// Don't do destruct work on the last AsyncCache thread, save it for after the parallelfor
-						&& (TaskType == ETaskType::DoLoad || ResultIndex + 1 < InCacheFilenames.Num()))
-					{
-						DestructData = DestructQueue.Pop(EAllowShrinking::No);
-						check(DestructData != nullptr);
-					}
-					else
-					{
-						// Claim the task.
-						NextTaskIndex++;
+				bVerifiedOk = VerifyIntegrity(Ar.Get());
+			}
 
-						CacheFilename = &InCacheFilenames[ResultIndex];
-						FileData = &CacheFiles[ResultIndex];
-						FileData->ConditionalConstruct(**CacheFilename);
-						bShouldAsyncCache = TaskType == ETaskType::AsyncCache && !FileData->bLoadComplete;
-					}
-				}
+			// now read the file if it is valid
+			if (bVerifiedOk)
+			{
+				TSharedPtr<UE::AssetRegistry::FMemoryMappedFile> File = MakeShared<UE::AssetRegistry::FMemoryMappedFile>(*CacheFilename);
 
-				if (DestructData)
-				{
-					DestructData->File.Reset();
-				}
-			} while (DestructData != nullptr);
-			check(FileData != nullptr && CacheFilename != nullptr); // Required to inform static analysis.
+				// we already know (from the verification) that Ar is non-empty
+				uint64 Version;
+				*Ar << Version;
 
-			if (TaskType == ETaskType::AsyncCache)
-			{
-				if (bShouldAsyncCache)
-				{
-					FileData->File->Preload();
-				}
-				FScopeLock ParallelForScopeLock(&ParallelForLock);
-				FileData->bAsyncCacheComplete = true;
-				FileData->ConditionalQueueDestruction(DestructQueue);
-				// Don't do any destruction work on our current loop: we are possibly the last iteration
-				// and should allow the ParallelFor to complete, and destruct in a non-blocking task afterwards.
-			}
-			else // ETaskType::DoLoad
-			{
-				FCachePayload Payload;
-				FMemoryViewReader FileReader(FileData->File->View());
-				TOptional<uint64> Version = FileReader.TryLoad<uint64>();
 				if (Version == FGatherSettings::CurrentVersion)
 				{
-					FChecksumViewReader ChecksummingReader(MoveTemp(FileReader), *CacheFilename);
-					Payload = DoLoad(ChecksummingReader);
+					Payload = DoLoad(*Ar, File);
 					UE_CLOG(Payload.bSucceeded, LogAssetRegistry, Display,
-						TEXT("Asset registry cache read as %.1f MiB from %s."),
-						static_cast<float>(FileReader.GetTotalSize()) / 1024.f / 1024.f, **CacheFilename);
+						TEXT("Asset registry cache read as %.1f MiB from %s"),
+						static_cast<float>(Ar->TotalSize()) / 1024.f / 1024.f, *CacheFilename);
 					UE_CLOG(!Payload.bSucceeded, LogAssetRegistry, Warning,
-						TEXT("There was an error loading the asset registry cache using memory mapping from %s."),
-						**CacheFilename);
+						TEXT("There was an error loading the asset registry cache from %s."),
+						*CacheFilename);
 				}
-				Results[ResultIndex] = MoveTemp(Payload);
+			}
+			else // check if it was a past version
+			{
+				Ar.Reset();
 
-				FScopeLock ParallelForScopeLock(&ParallelForLock);
-				FileData->bLoadComplete = true;
-				FileData->ConditionalQueueDestruction(DestructQueue);
-				// Don't do any destruction work on our current loop: we are possibly the last iteration
-				// and should allow the ParallelFor to complete, and destruct in a non-blocking task afterwards.
+				TSharedPtr<UE::AssetRegistry::FMemoryMappedFile> File = MakeShared<UE::AssetRegistry::FMemoryMappedFile>(*CacheFilename);
+				if (File)
+				{
+					AssetRegistryBackwardCompat::FMemoryViewReader FileReader(File->GetWholeFileView());
+					TOptional<uint64> VersionOpt = FileReader.TryLoad<uint64>();
+					// all later versions should have a valid integrity hash, so if we're here with a later version archive, it is genuinely corrupted and this is not a compat issue.
+					if (VersionOpt.IsSet() && static_cast<FAssetRegistryVersion::Type>(VersionOpt.GetValue()) < FAssetRegistryVersion::MemoryMappedTagDataStore)
+					{
+						FAssetRegistryVersion::Type Version = static_cast<FAssetRegistryVersion::Type>(VersionOpt.GetValue());
+						AssetRegistryBackwardCompat::FChecksumViewReader ChecksummingReader(MoveTemp(FileReader), *CacheFilename);
+						Payload = DoLoad(ChecksummingReader, nullptr, Version);	// cannot pass the mmaped file, as we cannot map the data in old format
+						UE_CLOG(Payload.bSucceeded, LogAssetRegistry, Display,
+							TEXT("Asset registry older version (%d) cache read as %.1f MiB from %s."),
+							Version,
+							static_cast<float>(FileReader.GetTotalSize()) / 1024.f / 1024.f, *CacheFilename);
+						UE_CLOG(!Payload.bSucceeded, LogAssetRegistry, Warning,
+							TEXT("There was an error loading the asset registry cache in old format (version %d) using memory mapping from %s."),
+							Version,
+							*CacheFilename);
+					}
+					else
+					{
+						UE_LOG(LogAssetRegistry, Warning, TEXT("Corrupted AssetRegistry cache file %s"),
+							*CacheFilename
+						);
+					}
+				}
 			}
-		}, EParallelForFlags::BackgroundPriority);
-		for (FSharedFileData& FileData : CacheFiles)
-		{
-			check(FileData.bLoadComplete && FileData.bAsyncCacheComplete);
-		}
+			Results[Index] = MoveTemp(Payload);
 
-		// Ignore the remaining elements of the DestructQueue, and instead pass the entire set of CacheFiles to another
-		// thread to finish the destruction of any that still need it.
-		DestructQueue.Empty();
-		UE::Tasks::Launch(UE_SOURCE_LOCATION, [KillAsync = MoveTemp(CacheFiles)]() {});
+		}, EParallelForFlags::BackgroundPriority);
 	}
 	else
 	{
-		ParallelFor(InCacheFilenames.Num(), [&DoLoad, &InCacheFilenames, &Results](int32 Index)
+		ParallelFor(InCacheFilenames.Num(), [&DoLoad, &InCacheFilenames, &Results, &VerifyIntegrity](int32 Index)
 		{
 			const FString& CacheFilename = InCacheFilenames[Index];
 
 			FCachePayload Payload;
 			TUniquePtr<FArchive> FileAr(IFileManager::Get().CreateFileReader(*CacheFilename, FILEREAD_Silent));
-			if (FileAr && !FileAr->IsError() && FileAr->TotalSize() > sizeof(uint64))
+			// old archive versions (predating the hash) may fail verification, but this path only supports current version anyway
+			if (FileAr && !FileAr->IsError() && FileAr->TotalSize() > sizeof(uint64) && VerifyIntegrity(FileAr.Get()))
 			{
 				uint64 Version = 0;
 				*FileAr << Version;
 				if (Version == FGatherSettings::CurrentVersion)
 				{
-					FChecksumArchiveReader ChecksummingReader(*FileAr);
-					Payload = DoLoad(ChecksummingReader);
+					Payload = DoLoad(*FileAr, nullptr);
 					UE_CLOG(Payload.bSucceeded, LogAssetRegistry, Display,
 						TEXT("Asset registry cache read as %.1f MiB from %s"),
 						static_cast<float>(FileAr->TotalSize()) / 1024.f / 1024.f, *CacheFilename);
diff --git a/Engine/Source/Runtime/AssetRegistry/Private/AssetRegistryArchive.cpp b/Engine/Source/Runtime/AssetRegistry/Private/AssetRegistryArchive.cpp
--- a/Engine/Source/Runtime/AssetRegistry/Private/AssetRegistryArchive.cpp
+++ b/Engine/Source/Runtime/AssetRegistry/Private/AssetRegistryArchive.cpp
@@ -5,6 +5,7 @@
 #include "Algo/Sort.h"
 #include "AssetRegistryPrivate.h"
 #include "AssetRegistry/AssetRegistryState.h"
+#include "Serialization/MemoryHasher.h"
 
 #ifndef ASSET_REGISTRY_LOAD_NAMES_IN_ORDER
 #define ASSET_REGISTRY_LOAD_NAMES_IN_ORDER 0
@@ -222,11 +223,17 @@ FAssetRegistryWriter::~FAssetRegistryWriter()
 {
 	// Save store data and collect FNames
 	int64 BodySize = MemWriter.TotalSize();
+	UE::AssetRegistry::AlignPosInArchive(MemWriter);	// make sure the data before store leaves it in the aligned position
+	int64 StoreStart = MemWriter.TotalSize();
 	SaveStore(Tags.Finalize(), *this);
 
 	// Save in load-friendly order - names, store then body / tag maps
 	SaveNameBatch(FlattenIndex(Names), TargetAr);
-	TargetAr.Serialize(MemWriter.GetData() + BodySize, MemWriter.TotalSize() - BodySize);
+
+	// make sure the Store is aligned
+	UE::AssetRegistry::AlignPosInArchive(TargetAr);
+	TargetAr.Serialize(MemWriter.GetData() + StoreStart, MemWriter.TotalSize() - StoreStart);
+	// Body does not need to be aligned because we copy all the data out of it rather than pointing to its memory-mapped position
 	TargetAr.Serialize(MemWriter.GetData(), BodySize);
 }
 
@@ -261,7 +268,7 @@ void FAssetRegistryWriter::SerializeTagsAndBundles(const FAssetData& Out)
 
 #endif
 
-FAssetRegistryReader::FAssetRegistryReader(FArchive& Inner, int32 NumWorkers, FAssetRegistryHeader Header)
+FAssetRegistryReader::FAssetRegistryReader(FArchive& Inner, TSharedPtr<UE::AssetRegistry::FMemoryMappedFile> InBackingFile, int32 NumWorkers, FAssetRegistryHeader Header)
 	: FArchiveProxy(Inner)
 {
 	check(IsLoading());
@@ -269,25 +276,34 @@ FAssetRegistryReader::FAssetRegistryReader(FArchive& Inner, int32 NumWorkers, FA
 	SetFilterEditorOnly(Header.bFilterEditorOnlyData);
 	FArchive::SetFilterEditorOnly(Header.bFilterEditorOnlyData); // Workaround for bug in FArchiveProxy
 
-	// Use parallel loading if requested and available; the parallel path is only for latest version
-	if (NumWorkers > 0 && Header.Version == FAssetRegistryVersion::LatestVersion)
+	// Load older versions for compatibility
+	if (UNLIKELY(Header.Version < FAssetRegistryVersion::MemoryMappedTagDataStore))
+	{
+		Names = LoadNameBatch(Inner, ENameBatchLoadingFlags::RespectOrder);
+		Tags = FixedTagPrivate::LoadStore(*this, Header.Version, nullptr);
+	}
+	else
 	{
-		ENameBatchLoadingFlags LoadingFlags = ENameBatchLoadingFlags::None;
+		if (NumWorkers > 0)
+		{
+			ENameBatchLoadingFlags LoadingFlags = ENameBatchLoadingFlags::None;
 #if ASSET_REGISTRY_LOAD_NAMES_IN_ORDER
-		LoadingFlags |= ENameBatchLoadingFlags::RespectOrder;
+			LoadingFlags |= ENameBatchLoadingFlags::RespectOrder;
 #endif
-		TFunction<TArray<FDisplayNameEntryId> ()> GetFutureNames = LoadNameBatchAsync(*this, NumWorkers, LoadingFlags);
+			TFunction<TArray<FDisplayNameEntryId> ()> GetFutureNames = LoadNameBatchAsync(*this, NumWorkers, LoadingFlags);
 
-		FixedTagPrivate::FAsyncStoreLoader StoreLoader;
-		Task = StoreLoader.ReadInitialDataAndKickLoad(*this, NumWorkers, Header.Version);
+			FixedTagPrivate::FAsyncStoreLoader StoreLoader;
+			StoreLoader.SetStoreBackingFile(InBackingFile);
+			Task = StoreLoader.ReadInitialDataAndKickLoad(*this, NumWorkers, Header.Version);
 		
-		Names = GetFutureNames();
-		Tags = StoreLoader.LoadFinalData(*this, Header.Version);
-	}
-	else
-	{
-		Names = LoadNameBatch(Inner, ENameBatchLoadingFlags::RespectOrder);
-		Tags = FixedTagPrivate::LoadStore(*this, Header.Version);
+			Names = GetFutureNames();
+			Tags = StoreLoader.LoadFinalData(*this, Header.Version);
+		}
+		else
+		{
+			Names = LoadNameBatch(Inner, ENameBatchLoadingFlags::RespectOrder);
+			Tags = FixedTagPrivate::LoadStore(*this, Header.Version, InBackingFile);
+		}
 	}
 }
 
@@ -422,7 +438,9 @@ bool FAssetRegistryTagSerializationTest::RunTest(const FString& Parameters)
 	{
 		FMemoryReader DataReader(Data);
 		FAssetRegistryHeader Header;
-		FAssetRegistryReader RegistryReader(DataReader, 0, Header);
+		FAssetRegistryReader RegistryReader(DataReader, nullptr, 0, Header);
+		RegistryReader.WaitForTasks();
+
 		for (FAssetDataTagMapSharedView& FixedMap : FixedMaps)
 		{
 			FixedMap = LoadTags(RegistryReader);
@@ -435,7 +453,8 @@ bool FAssetRegistryTagSerializationTest::RunTest(const FString& Parameters)
 	{
 		FMemoryReader DataReader(Data);
 		FAssetRegistryHeader Header;
-		FAssetRegistryReader RegistryReader(DataReader, 0, Header);
+		FAssetRegistryReader RegistryReader(DataReader, nullptr, 0, Header);
+		RegistryReader.WaitForTasks();
 
 		for (const FAssetDataTagMapSharedView& FixedMap1 : FixedMaps)
 		{
diff --git a/Engine/Source/Runtime/AssetRegistry/Private/AssetRegistryArchive.h b/Engine/Source/Runtime/AssetRegistry/Private/AssetRegistryArchive.h
--- a/Engine/Source/Runtime/AssetRegistry/Private/AssetRegistryArchive.h
+++ b/Engine/Source/Runtime/AssetRegistry/Private/AssetRegistryArchive.h
@@ -3,6 +3,7 @@
 #pragma once
 
 #include "AssetRegistry/AssetDataTagMapSerializationDetails.h"
+#include "AssetRegistry/AssetRegistrySerializationDetails.h"
 #include "Serialization/LargeMemoryWriter.h"
 #include "UObject/NameBatchSerialization.h"
 
@@ -22,8 +23,9 @@ struct FAssetRegistryHeader
 class FAssetRegistryReader : public FArchiveProxy
 {
 public:
+	/// @param InBackingFile shared pointer to a class that allows mapping the underlying file into memory
 	/// @param NumWorkers > 0 for parallel loading
-	FAssetRegistryReader(FArchive& Inner, int32 NumWorkers, FAssetRegistryHeader Header);
+	FAssetRegistryReader(FArchive& Inner, TSharedPtr<UE::AssetRegistry::FMemoryMappedFile> InBackingFile, int32 NumWorkers, FAssetRegistryHeader Header);
 	~FAssetRegistryReader();
 
 	virtual FArchive& operator<<(FName& Value) override;
diff --git a/Engine/Source/Runtime/AssetRegistry/Private/AssetRegistryState.cpp b/Engine/Source/Runtime/AssetRegistry/Private/AssetRegistryState.cpp
--- a/Engine/Source/Runtime/AssetRegistry/Private/AssetRegistryState.cpp
+++ b/Engine/Source/Runtime/AssetRegistry/Private/AssetRegistryState.cpp
@@ -2376,7 +2376,7 @@ bool FAssetRegistryState::Load(FArchive& OriginalAr, const FAssetRegistryLoadOpt
 	}
 	else
 	{
-		FAssetRegistryReader Reader(OriginalAr, Options.ParallelWorkers, Header);
+		FAssetRegistryReader Reader(OriginalAr, nullptr, Options.ParallelWorkers, Header);
 
 		if (Reader.IsError())
 		{
diff --git a/Engine/Source/Runtime/Core/Private/Windows/WindowsPlatformFile.cpp b/Engine/Source/Runtime/Core/Private/Windows/WindowsPlatformFile.cpp
--- a/Engine/Source/Runtime/Core/Private/Windows/WindowsPlatformFile.cpp
+++ b/Engine/Source/Runtime/Core/Private/Windows/WindowsPlatformFile.cpp
@@ -1591,14 +1591,18 @@ FOpenMappedResult FWindowsPlatformFile::OpenMappedEx(const TCHAR* Filename, EOpe
 {
 	int64 Size = FileSize(Filename);
 	uint32  Access = GENERIC_READ;
-	uint32  WinFlags = FILE_SHARE_READ | GetFileShareDeleteFlagIfEnabled();
+	uint32  WinFlags = FILE_SHARE_READ |
+		(EnumHasAnyFlags(OpenOptions, EOpenReadFlags::AllowDelete) ? FILE_SHARE_DELETE : 0) |
+		GetFileShareDeleteFlagIfEnabled();
 	uint32  Create = OPEN_EXISTING;
 	DWORD	Protect = PAGE_READONLY;
 
 	if (EnumHasAnyFlags(OpenOptions, EOpenReadFlags::AllowWrite))
 	{
 		Access = GENERIC_READ | GENERIC_WRITE;
-		WinFlags = FILE_SHARE_READ | FILE_SHARE_WRITE | GetFileShareDeleteFlagIfEnabled();
+		WinFlags = FILE_SHARE_READ | FILE_SHARE_WRITE |
+			(EnumHasAnyFlags(OpenOptions, EOpenReadFlags::AllowDelete) ? FILE_SHARE_DELETE : 0) |
+			GetFileShareDeleteFlagIfEnabled();
 		Create = CREATE_ALWAYS;
 		Protect = PAGE_READWRITE;
 		Size = MaximumSize;
diff --git a/Engine/Source/Runtime/CoreUObject/Private/AssetRegistry/AssetDataTagMap.cpp b/Engine/Source/Runtime/CoreUObject/Private/AssetRegistry/AssetDataTagMap.cpp
--- a/Engine/Source/Runtime/CoreUObject/Private/AssetRegistry/AssetDataTagMap.cpp
+++ b/Engine/Source/Runtime/CoreUObject/Private/AssetRegistry/AssetDataTagMap.cpp
@@ -4,7 +4,10 @@
 
 #include "Algo/Sort.h"
 #include "AssetRegistry/AssetDataTagMapSerializationDetails.h"
+#include "AssetRegistry/AssetRegistrySerializationDetails.h"
 #include "AutoRTFM.h"
+#include "HAL/PlatformFile.h"
+#include "HAL/PlatformFileManager.h"
 #include "Misc/PackageName.h"
 #include "Misc/Optional.h"
 #include "Misc/TransactionallySafeCriticalSection.h"
@@ -12,12 +15,45 @@
 #include "Serialization/ArrayWriter.h"
 #include "Serialization/MemoryReader.h"
 #include "Serialization/MemoryWriter.h"
+
 #include <limits>
 
 DEFINE_LOG_CATEGORY_STATIC(LogAssetDataTags, Log, All);
 
 //////////////////////////////////////////////////////////////////////////
 
+void UE::AssetRegistry::AlignPosInArchive(FArchive& Archive)
+{
+	// see if current offset is aligned
+	int64 Pos = Archive.Tell();
+	int64 AlignedPos = Align(Pos, AlignmentForMemoryMappedArchive);
+	if (AlignedPos != Pos)
+	{
+		int64 DataToSkip = AlignedPos - Pos;
+
+		if (Archive.IsSaving())
+		{
+			char Dunnage[AlignmentForMemoryMappedArchive] = { 0 };
+
+#if 0 // for easier debugging
+			for (SIZE_T Idx = 0; Idx < sizeof(Dunnage); ++Idx)
+			{
+				Dunnage[Idx] = static_cast<char>(Idx);
+			}
+#endif
+			Archive.Serialize(Dunnage, DataToSkip);
+		}
+		else
+		{
+			Archive.Seek(AlignedPos);
+		}
+
+		checkf(Archive.Tell() == AlignedPos, TEXT("We were not able to align the position in the archive to the desired boundary."));
+	}
+}
+
+//////////////////////////////////////////////////////////////////////////
+
 template<class T>
 FAssetRegistryExportPath ParseExportPath(T ExportPath)
 {
@@ -999,14 +1035,36 @@ namespace FixedTagPrivate
 		};
 	}
 
+	void* FStore::AllocateMemoryForView(SIZE_T NumBytes, SIZE_T AlignmentRequired)
+	{
+		void* Memory = FMemory::Malloc(NumBytes, static_cast<uint32>(AlignmentRequired));
+		if (LIKELY(Memory))
+		{
+			FScopeLock Lock(&ViewPointersCS);
+			ViewPointers.Add(Memory);
+		}
+		return Memory;
+	}
+
+
 	FStore::~FStore()
 	{
 		GStores.Unregister(*this);
 
-		if (Data)
+		// not locking the pointers since no other thread should be using the store at this point
+		VisitViews(*this, [this](auto& View)
+			{
+				// If the view is not memory-mmapped, destroy its elements. Memory-mapped views have no-op destructors
+				if (ViewPointers.Contains(View.GetData()))
+				{
+					DestroyElements(View);
+				}
+			}
+		);
+
+		for (void* Ptr : ViewPointers)
 		{
-			VisitViews(*this, [] (auto View) { DestroyElements(View); });
-			FMemory::Free(Data);
+			FMemory::Free(Ptr);
 		}
 	}
 
@@ -1182,6 +1240,110 @@ namespace FixedTagPrivate
 		void LoadViewData(TArrayView<WIDECHAR> View) { BulkSerialize(View); }
 		static_assert(PLATFORM_LITTLE_ENDIAN, "Byte-swapping WIDECHAR on load needed to load on big-endian platforms");
 
+		template<class T>
+		inline void LoadViewFromMemoryWithoutAttemptingToMap(FStore& Store, TArrayView<T>& View)
+		{
+			uint64 Bytes = GetBytes(View);
+			if (Bytes)
+			{
+				SetUntypedDataPtr(View, Store.AllocateMemoryForView(Bytes, alignof(T)));
+				MemZeroViewDataIfNeeded(View);
+				LoadViewData(View);
+			}
+		}
+
+		// by default, assume the type cannot be memory mapped
+		template<class T>
+		void LoadViewFromMemoryMappedArchive(FStore& Store, TArrayView<T>& View)
+		{
+			LoadViewFromMemoryWithoutAttemptingToMap(Store, View);
+		}
+
+		void LoadViewFromMemoryMappedArchive(FStore& Store, TArrayView<uint32>& View)
+		{
+#if WITH_EDITOR
+			if (Store.BackingFile.IsValid())
+			{
+				if (int64 Size = GetBytes(View))
+				{
+					int64 Offset = Ar.Tell();
+					FMemoryView Part = Store.BackingFile->GetView(Offset, Size);
+					if (void* DataPtr = const_cast<void*>(Part.GetData()))
+					{
+						SetUntypedDataPtr(View, DataPtr);
+						Ar.Seek(Offset + Size);
+					}
+					else
+					{
+						// in an unlikely case we could not map, switch to allocation
+						LoadViewFromMemoryWithoutAttemptingToMap(Store, View);
+					}
+				}
+			}
+			else
+#endif
+			{
+				LoadViewFromMemoryWithoutAttemptingToMap(Store, View);
+			}
+		}
+
+		// ANSI and WIDE chars can be memory mapped
+		void LoadViewFromMemoryMappedArchive(FStore& Store, TArrayView<ANSICHAR>& View)
+		{
+#if WITH_EDITOR
+			if (Store.BackingFile.IsValid())
+			{
+				if (int64 Size = GetBytes(View))
+				{
+					int64 Offset = Ar.Tell();
+					FMemoryView Part = Store.BackingFile->GetView(Offset, Size);
+					if (void* DataPtr = const_cast<void*>(Part.GetData()))
+					{
+						SetUntypedDataPtr(View, DataPtr);
+						Ar.Seek(Offset + Size);
+					}
+					else
+					{
+						// in an unlikely case we could not map, switch to allocation
+						LoadViewFromMemoryWithoutAttemptingToMap(Store, View);
+					}
+				}
+			}
+			else
+#endif
+			{
+				LoadViewFromMemoryWithoutAttemptingToMap(Store, View);
+			}
+		}
+		void LoadViewFromMemoryMappedArchive(FStore& Store, TArrayView<WIDECHAR>& View)
+		{
+#if WITH_EDITOR
+			if (Store.BackingFile.IsValid())
+			{
+				if (int64 Size = GetBytes(View))
+				{
+					int64 Offset = Ar.Tell();
+					FMemoryView Part = Store.BackingFile->GetView(Offset, Size);
+					if (void* DataPtr = const_cast<void*>(Part.GetData()))
+					{
+						SetUntypedDataPtr(View, DataPtr);
+						Ar.Seek(Offset + Size);
+					}
+					else
+					{
+						// in an unlikely case we could not map, switch to allocation
+						LoadViewFromMemoryWithoutAttemptingToMap(Store, View);
+					}
+				}
+			}
+			else
+#endif
+			{
+				LoadViewFromMemoryWithoutAttemptingToMap(Store, View);
+			}
+		}
+
+
 		FArchive& Ar;
 		FString Scratch;
 
@@ -1200,15 +1362,19 @@ namespace FixedTagPrivate
 			FArrayWriter Data;
 			TSerializer<Version>(Data).SaveViewData(Texts);
 			SaveItem(Data.Num());
+			AlignPosInArchive();
 			Ar.Serialize(Data.GetData(), Data.Num());
 		}
 
 		void Save(const FStoreData& Store)
 		{
+			// Store is expected to be saved at an aligned position
+			check(UE::AssetRegistry::IsPosInArchiveAligned(Ar));
 			SaveItem(BeginMagic);
 			VisitViews(Store, [&](auto& Array) { SaveItem(Array.Num()); });
 			SaveTextData(MakeArrayView(Store.Texts));
-			VisitViews<EOrder::SkipText>(Store, [&] (auto& Array) { SaveViewData(MakeArrayView(Array)); });
+			// while not all data will be mmaped, give all of it a possibility to be so by aligning everything
+			VisitViews<EOrder::SkipText>(Store, [&] (auto& Array) { AlignPosInArchive(); SaveViewData(MakeArrayView(Array)); });
 			SaveItem(EndMagic);
 		}
 
@@ -1222,8 +1388,18 @@ namespace FixedTagPrivate
 			}
 		}
 
+		inline void AlignPosInArchive()
+		{
+			if constexpr (Version >= FAssetRegistryVersion::MemoryMappedTagDataStore)
+			{
+				UE::AssetRegistry::AlignPosInArchive(Ar);
+			}
+		}
+
 		TOptional<ELoadOrder> LoadHeader(FStore& Store)
 		{
+			// store is expected to be saved at an aligned position, so it can align its data
+			AlignPosInArchive();
 			uint32 InitialMagic = LoadItem<uint32>();
 			TOptional<ELoadOrder> Order = GetLoadOrder(InitialMagic);
 
@@ -1237,29 +1413,6 @@ namespace FixedTagPrivate
 			// Load view sizes
 			VisitViews(Store, [&] (auto& View) { SetNum(View, LoadItem<int32>()); });
 
-			// Calculate total size and allocate data
-			uint64 Bytes = 0;
-			VisitViews(Store, [&] (auto& View) 
-				{
-					static_assert(alignof(typename std::decay_t<decltype(View)>::ElementType) <= MaxViewAlignment, "");
-					Bytes = Align(Bytes, View.GetTypeAlignment()) + GetBytes(View);
-				});
-
-			uint8* Ptr = reinterpret_cast<uint8*>(FMemory::Malloc(Bytes, MaxViewAlignment));
-			Store.Data = Ptr;
-
-			// Set view data pointers
-			VisitViews(Store, [&] (auto& View)
-				{
-					Ptr = Align(Ptr, View.GetTypeAlignment());
-					SetUntypedDataPtr(View, Ptr);
-					Ptr += GetBytes(View);
-				});
-			check(Ptr - Bytes == Store.Data);
-
-			// Zero out required data
-			VisitViews(Store, [&](auto& View) { MemZeroViewDataIfNeeded(View); });
-
 			return Order;
 		}
 		
@@ -1272,7 +1425,7 @@ namespace FixedTagPrivate
 				{
 					uint32 TextDataBytes = LoadItem<uint32>();
 
-					VisitViews<EOrder::TextFirst>(Store, [&] (auto View) { LoadViewData(View); });
+					VisitViews<EOrder::TextFirst>(Store, [&](auto& View) { AlignPosInArchive(); LoadViewFromMemoryMappedArchive(Store, View); });
 				}
 				else
 				{
@@ -1292,25 +1445,20 @@ namespace FixedTagPrivate
 			uint32 TextDataBytes = LoadItem<uint32>();
 			TArray<uint8> Out;
 			Out.SetNumUninitialized(TextDataBytes);
+			AlignPosInArchive();
 			Ar.Serialize(Out.GetData(), Out.Num());
 			return Out;
 		}
 
 		void LoadTextData(FStore& Store)
 		{
-			LoadViewData(Store.Texts);
+			LoadViewFromMemoryMappedArchive(Store, Store.Texts);
 		}
 
 		void LoadFinalData(FStore& Store, ELoadOrder Order)
 		{
-			if (Order == ELoadOrder::TextFirst)
-			{
-				VisitViews<EOrder::SkipText>(Store, [&] (auto View) { LoadViewData(View); });
-			}
-			else
-			{
-				VisitViews(Store, [&] (auto View) { LoadViewData(View); });
-			}
+			checkf(Order == ELoadOrder::TextFirst, TEXT("Loading old versions of cache is no longer supported"));
+			VisitViews<EOrder::SkipText>(Store, [&] (auto& View) { AlignPosInArchive(); LoadViewFromMemoryMappedArchive(Store, View); });
 			
 			if (LoadItem<uint32>() != EndMagic)
 			{
@@ -1326,7 +1474,7 @@ namespace FixedTagPrivate
 		TSerializer<FAssetRegistryVersion::LatestVersion>(Ar).Save(Store);
 	}
 
-	TRefCountPtr<const FStore> LoadStore(FArchive& Ar, FAssetRegistryVersion::Type Version /*= FAssetRegistryVersion::LatestVersion*/)
+	TRefCountPtr<const FStore> LoadStore(FArchive& Ar, FAssetRegistryVersion::Type Version /*= FAssetRegistryVersion::LatestVersion*/, TSharedPtr<UE::AssetRegistry::FMemoryMappedFile> InBackingFile /*= nullptr*/)
 	{
 		if (Ar.IsError())
 		{
@@ -1334,6 +1482,9 @@ namespace FixedTagPrivate
 		}
 
 		FStore* Store = GStores.CreateAndRegister();
+#if WITH_EDITOR
+		Store->BackingFile = InBackingFile;
+#endif
 
 		// Versioning support inside of TSerializer was added with FAssetRegistryVersion::ClassPaths so everything before can just be
 		// serialized with FAssetRegistryVersion::ClassPaths - 1 = FAssetRegistryVersion::AddedChunkHashes.
@@ -1349,11 +1500,15 @@ namespace FixedTagPrivate
 		{
 			TSerializer<FAssetRegistryVersion::AssetPackageDataHasPackageLocation>(Ar).Load(*Store);
 		}
+		else if (Version <= FAssetRegistryVersion::TagSetMapSupports64BitStringNum)
+		{
+			TSerializer<FAssetRegistryVersion::TagSetMapSupports64BitStringNum>(Ar).Load(*Store);
+		}
 		else // Add more branches should serialization format change
 		{
 			TSerializer<FAssetRegistryVersion::LatestVersion>(Ar).Load(*Store);
 		}
-		
+
 		return TRefCountPtr<const FStore>(Store);
 	}
 
@@ -1367,20 +1522,19 @@ namespace FixedTagPrivate
 		check(HeaderVersion == FAssetRegistryVersion::LatestVersion);
 		Order = TSerializer<FAssetRegistryVersion::LatestVersion>(Ar).LoadHeader(*Store);
 
-		if (Order == ELoadOrder::TextFirst)
-		{
-			TArray<uint8> TextData = TSerializer<FAssetRegistryVersion::LatestVersion>(Ar).ReadTextData();
+		checkf(Order == ELoadOrder::TextFirst, TEXT("Reading old formats isn't supported."));
+
+		UE::AssetRegistry::AlignPosInArchive(Ar);
+		TArray<uint8> TextData = TSerializer<FAssetRegistryVersion::LatestVersion>(Ar).ReadTextData();
 
-			if (TextData.Num())
+		if (TextData.Num())
+		{
+			return Async(EAsyncExecution::TaskGraph, [Data = MoveTemp(TextData), OutStore = Store]()
 			{
-				return Async(EAsyncExecution::TaskGraph, [Data = MoveTemp(TextData), OutStore = Store]()
-				{
-					FMemoryReader Reader(Data);
-					TSerializer<FAssetRegistryVersion::LatestVersion>(Reader).LoadTextData(*OutStore);
-				});
-			}
+				FMemoryReader Reader(Data);
+				TSerializer<FAssetRegistryVersion::LatestVersion>(Reader).LoadTextData(*OutStore);
+			});
 		}
-
 		return TFuture<void>();
 	}
 
@@ -1734,6 +1888,41 @@ void FAssetDataTagMapSharedView::FMemoryCounter::ReportFixedStoreBreakdown() con
 	}
 }
 
+namespace UE::AssetRegistry
+{
+	FMemoryMappedFile::FMemoryMappedFile(const TCHAR* Path)
+	{
+		FOpenMappedResult Result = FPlatformFileManager::Get().GetPlatformFile().OpenMappedEx(Path, IPlatformFile::EOpenReadFlags::AllowDelete);
+		if (!Result.HasError())
+		{
+			Handle = Result.StealValue();
+		}
+	}
+
+	FMemoryView FMemoryMappedFile::GetView(uint64 Start, uint64 Length) const
+	{
+		FMappedSubRegion New;
+		New.Start = Start;
+		New.Length = Length;
+
+		if (!Handle)
+		{
+			return FMemoryView();
+		}
+
+		New.Region.Reset(Handle->MapRegion(Start, Length));
+		if (!New.Region)
+		{
+			return FMemoryView();
+		}
+
+		FMemoryView Result = FMemoryView(New.Region->GetMappedPtr(), New.Region->GetMappedSize());
+		Regions.Add(MoveTemp(New));
+		return Result;
+	}
+
+} // namespace UE::AssetRegistry
+
 ////////////////////////////////////////////////////////////////////////////
 
 #if WITH_DEV_AUTOMATION_TESTS 
diff --git a/Engine/Source/Runtime/CoreUObject/Public/AssetRegistry/AssetData.h b/Engine/Source/Runtime/CoreUObject/Public/AssetRegistry/AssetData.h
--- a/Engine/Source/Runtime/CoreUObject/Public/AssetRegistry/AssetData.h
+++ b/Engine/Source/Runtime/CoreUObject/Public/AssetRegistry/AssetData.h
@@ -105,6 +105,8 @@ struct FAssetRegistryVersion
 		PackageSavedHash,					// Replaced FAssetPackageData::PackageGuid with PackageSavedHash
 		ExternalActorToWorldIsEditorOnly,   // FPackageDependencyData::LoadDependenciesFromPackageHeader changed how it calculates PackageDependencies
 		ManageDependenciesCookRule,			// CookRule property added to EDependencyProperties for EDependencyCategory::Manage
+		TagSetMapSupports64BitStringNum,	// 64-bit arrays for strings
+		MemoryMappedTagDataStore,			// Tag data store is now memory mapped
 
 		// -----<new versions can be added above this line>-------------------------------------------------
 		VersionPlusOne,
diff --git a/Engine/Source/Runtime/CoreUObject/Public/AssetRegistry/AssetDataTagMapSerializationDetails.h b/Engine/Source/Runtime/CoreUObject/Public/AssetRegistry/AssetDataTagMapSerializationDetails.h
--- a/Engine/Source/Runtime/CoreUObject/Public/AssetRegistry/AssetDataTagMapSerializationDetails.h
+++ b/Engine/Source/Runtime/CoreUObject/Public/AssetRegistry/AssetDataTagMapSerializationDetails.h
@@ -4,6 +4,7 @@
 
 #include "AssetRegistry/AssetData.h"
 #include "AssetRegistry/AssetDataTagMap.h"
+#include "AssetRegistrySerializationDetails.h"
 #include "Async/Async.h"
 #include "AutoRTFM.h"
 #include "HAL/ThreadSafeCounter.h"
@@ -62,7 +63,16 @@ namespace FixedTagPrivate
 		TArrayView<FMarshalledText> Texts;
 
 		const uint32 Index;
-		void* Data = nullptr;
+
+#if WITH_EDITOR
+		/**
+		 * File where the data is kept, if any. If nullptr, the archive is not using memory-mapped files. Editor cache
+		 * files usually use memory-mapped files, but in-game cooked AssetRegistry.bin data do not.
+		 */
+		mutable TSharedPtr<UE::AssetRegistry::FMemoryMappedFile> BackingFile;
+#endif
+		/** Allocates a pointer from the heap if a view cannot use the memory-mapped data in place. */
+		void* AllocateMemoryForView(SIZE_T NumBytes, SIZE_T AlignmentRequired);
 
 		void AddRef() const
 		{
@@ -93,6 +103,11 @@ namespace FixedTagPrivate
 		~FStore();
 
 		mutable FThreadSafeCounter RefCount;
+
+		/** Guards access to the View pointers. No need for an R/O as we're only accessing the set when we're about to modify it. */
+		FCriticalSection ViewPointersCS;
+		/** Pointers to individual views that aren't mmapped but allocated from the heap */
+		TSet<void*> ViewPointers;
 	};
 
 	struct FOptions
@@ -255,7 +270,7 @@ namespace FixedTagPrivate
 	enum class ELoadOrder { Member, TextFirst };
 
 	COREUOBJECT_API void SaveStore(const FStoreData& Store, FArchive& Ar);
-	COREUOBJECT_API TRefCountPtr<const FStore> LoadStore(FArchive& Ar, FAssetRegistryVersion::Type Version = FAssetRegistryVersion::LatestVersion);
+	COREUOBJECT_API TRefCountPtr<const FStore> LoadStore(FArchive& Ar, FAssetRegistryVersion::Type Version = FAssetRegistryVersion::LatestVersion, TSharedPtr<UE::AssetRegistry::FMemoryMappedFile> InBackingFile = nullptr);
 
 	/// Loads tag store with async creation of expensive tag values
 	///
@@ -280,6 +295,19 @@ namespace FixedTagPrivate
 		/// @return indexed store, usable for FPartialMapHandle::MakeFullHandle()
 		COREUOBJECT_API TRefCountPtr<const FStore> LoadFinalData(FArchive& Ar, FAssetRegistryVersion::Type HeaderVersion);
 
+		/** Passes backing file for memory mapping to the store. */
+		inline void SetStoreBackingFile(TSharedPtr<UE::AssetRegistry::FMemoryMappedFile> InBackingFile)
+		{
+#if WITH_EDITOR
+			Store->BackingFile = MoveTemp(InBackingFile);
+#else
+			if (InBackingFile)
+			{
+				UE_LOG(LogAssetData, Error, TEXT("Reading asset registry from a memory mapped file is not supported in a non-editor binary."));
+			}
+#endif
+		}
+
 	private:
 		TRefCountPtr<FStore> Store;
 		TOptional<ELoadOrder> Order;
diff --git a/Engine/Source/Runtime/CoreUObject/Public/AssetRegistry/AssetRegistrySerializationDetails.h b/Engine/Source/Runtime/CoreUObject/Public/AssetRegistry/AssetRegistrySerializationDetails.h
new file mode 100644
--- /dev/null
+++ b/Engine/Source/Runtime/CoreUObject/Public/AssetRegistry/AssetRegistrySerializationDetails.h
@@ -0,0 +1,54 @@
+// Copyright Epic Games, Inc. All Rights Reserved.
+
+#pragma once
+
+#include "Serialization/ArchiveProxy.h"
+#include "Async/MappedFileHandle.h"
+#include "Memory/MemoryFwd.h"
+#include "Templates/SharedPointer.h"
+#include "Serialization/MemoryHasher.h"
+
+namespace UE::AssetRegistry
+{
+	/** Memory maps the bytes of a file on disk into memory and provides MemoryViews to those bytes until *this is destructed. */
+	class FMemoryMappedFile
+	{
+	public:
+		FMemoryMappedFile() = default;
+		COREUOBJECT_API explicit FMemoryMappedFile(const TCHAR* Path);
+
+		inline FMemoryView GetWholeFileView() const
+		{
+			return GetView(0, MAX_int64);
+		}
+
+		/** Provides a new View, allocating it if it doesn't exist. */
+		COREUOBJECT_API FMemoryView GetView(uint64 Start, uint64 Length) const;
+
+	private:
+
+		/** Handle for the mapped file */
+		TUniquePtr<IMappedFileHandle> Handle;
+
+		struct FMappedSubRegion
+		{
+			uint64 Start;
+			uint64 Length;
+
+			TUniquePtr<IMappedFileRegion> Region;
+		};
+
+		mutable TArray<FMappedSubRegion>  Regions;
+	};
+
+	inline constexpr int64 AlignmentForMemoryMappedArchive = 16;
+
+	/** When saving, saves extra memory to make sure the following data is aligned. When loading, seeks over that many bytes (does not load them). */
+	void COREUOBJECT_API AlignPosInArchive(FArchive& Archive);
+
+	/** Reports whether the archive's current offset is aligned, relative to offset 0. */
+	inline bool IsPosInArchiveAligned(FArchive& Archive)
+	{
+		return IsAligned(Archive.Tell(), AlignmentForMemoryMappedArchive);
+	}
+}
