AssetRegistry: Fix Dump to text to correctly print PackageSizes > 2GB.
#rnx

#Synced-CL 50612662

diff --git a/Engine/Source/Runtime/AssetRegistry/Private/AssetRegistryState.cpp b/Engine/Source/Runtime/AssetRegistry/Private/AssetRegistryState.cpp
--- a/Engine/Source/Runtime/AssetRegistry/Private/AssetRegistryState.cpp
+++ b/Engine/Source/Runtime/AssetRegistry/Private/AssetRegistryState.cpp
@@ -4810,7 +4810,7 @@ void FAssetRegistryState::Dump(const TArray<FString>& Arguments, TArray<FString>
 #else
 			PageBuffer << FIoHash();
 #endif
-			PageBuffer.Appendf(TEXT(" : %d bytes"), PackageData->DiskSize);
+			PageBuffer.Appendf(TEXT(" : %" INT64_FMT " bytes"), PackageData->DiskSize);
 			AddLine();
 		}
 
