AssetRegistry: Change FAssetRegistryState.Load to support two new options required by ActiveMountLoader: allocating assets individually rather than in a single block (so they can be unloaded individually), and culling the allocation of AssetData, PackageData, or DependsNodes for assets filtered out by a caller's filter function. Add FAssetRegistryState.Consume to move (rather than the existing copy) assets into the global AssetRegistry.
#rnx

#Synced-CL 50811273

diff --git a/Engine/Source/Runtime/AssetRegistry/Private/AssetRegistryState.cpp b/Engine/Source/Runtime/AssetRegistry/Private/AssetRegistryState.cpp
--- a/Engine/Source/Runtime/AssetRegistry/Private/AssetRegistryState.cpp
+++ b/Engine/Source/Runtime/AssetRegistry/Private/AssetRegistryState.cpp
@@ -143,7 +143,7 @@ void FAssetRegistryState::Reset()
 	}
 	else
 	{
-		// Delete all depends nodes in the cache
+		// Delete all PackageDatas in the cache
 		for (TMap<FName, FAssetPackageData*>::TConstIterator PackageDataIt(CachedPackageData); PackageDataIt; ++PackageDataIt)
 		{
 			if (PackageDataIt.Value())
@@ -516,7 +516,7 @@ void FAssetRegistryState::InitializeFromExisting(const FAssetDataMap& AssetDataM
 	FixedTagPrivate::FStoreBuilder Builder(FAssetRegistryWriterOptions(Options).Tags);
 
 	// We accumulate Assets to add to *this in a queue inside the ParallelFor. Construction of each asset is expensive
-	// and we want to work on it on multple threads. But we can't add them to *this simultaneously because we can't append
+	// and we want to work on it on multiple threads. But we can't add them to *this simultaneously because we can't append
 	// to our CachedAsset and query accelerator maps in parallel. So we need to add them all in a batch after the
 	// ParallelFor is over. But waiting until after the ParallelFor is over causes cache misses on the ids, so instead
 	// of doing them all in a single batch, we add them in smaller batches as we go. Each worker of the parallelfor
@@ -626,7 +626,7 @@ void FAssetRegistryState::InitializeFromExisting(const FAssetDataMap& AssetDataM
 		};
 
 	ParallelFor(AssetDataMapMaxIndex,
-		[&AssetDataMap, &ConditionalAddDuplicatedAssetToQueue](int32 Index)
+		[&ConditionalAddDuplicatedAssetToQueue](int32 Index)
 		{
 			ConditionalAddDuplicatedAssetToQueue(Index);
 		});
@@ -736,6 +736,209 @@ void FAssetRegistryState::InitializeFromExisting(const FAssetDataMap& AssetDataM
 	Shrink();
 }
 
+void FAssetRegistryState::ConsumeExisting(FAssetRegistryState&& Consume,
+	const FAssetRegistrySerializationOptions& Options,
+	FAssetRegistryAppendResult* OutAppendResult)
+{
+	// Move over AssetDatas
+	TArray<FAssetData*> AssetsToDeallocate;
+	for (FAssetData* AssetData : Consume.CachedAssets)
+	{
+		if (CachedAssets.Find(FCachedAssetKey(AssetData)))
+		{
+			AssetsToDeallocate.Add(AssetData);
+			continue;
+		}
+
+		AddAssetData(AssetData);
+		if (OutAppendResult)
+		{
+			OutAppendResult->AddedAssets.Add(AssetData);
+		}
+	}
+
+	// Move over PackageDatas
+	TArray<FAssetPackageData*> PackageDatasToDeallocate;
+	for (const TPair<FName, FAssetPackageData*>& Pair : Consume.CachedPackageData)
+	{
+		if (!Pair.Value)
+		{
+			continue;
+		}
+		FAssetPackageData*& ExistingData = CachedPackageData.FindOrAdd(Pair.Key);
+		if (ExistingData)
+		{
+			PackageDatasToDeallocate.Add(Pair.Value);
+			continue;
+		}
+		ExistingData = Pair.Value;
+		NumPackageData++;
+	}
+
+	// Move over DependsNodes, more complicated because of the links.
+	// First, move over DependsNodes, and create a map from OldNode to NewNode, either OldNode to itself because keep
+	// it or OldNode to the existing node that takes priority over it.
+	TArray<FDependsNode*> NodesToDeallocate;
+	TMap<FDependsNode*, FDependsNode*> OldNodeToNewNode;
+	OldNodeToNewNode.Reserve(Consume.CachedDependsNodes.Num());
+	for (const TPair<FAssetIdentifier, FDependsNode*>& Pair : Consume.CachedDependsNodes)
+	{
+		FDependsNode* ConsumedNode = Pair.Value;
+		FDependsNode*& ExistingNode = CachedDependsNodes.FindOrAdd(Pair.Key);
+		if (ExistingNode == nullptr)
+		{
+			ExistingNode = ConsumedNode;
+			NumDependsNodes++;
+		}
+		else
+		{
+			NodesToDeallocate.Add(ConsumedNode);
+		}
+		OldNodeToNewNode.Add(ConsumedNode, ExistingNode);
+	}
+
+	// Second, remap:
+	// Add mirrored two-way links for any links between an old node we are keeping and an old node we are not
+	// keeping: copy over the link from the old node we are not keeping into the existing node that takes
+	// priority over it.
+	for (const TPair<FDependsNode*, FDependsNode*>& Pair : OldNodeToNewNode)
+	{
+		FDependsNode* OldSourceNode = Pair.Key;
+		FDependsNode* NewSourceNode = Pair.Value;
+		// We are keeping a node only if it maps to itself in FDependsNode*
+		if (NewSourceNode == OldSourceNode)
+		{
+			// For any dependencies or references from the kept old node to a replaced old node, copy the replaced old node's
+			// mirrored reference or dependency onto the new node that replaces it.
+			OldSourceNode->IterateOverDependencies([OldSourceNode, &OldNodeToNewNode]
+			(FDependsNode* OldTargetNode,
+				UE::AssetRegistry::EDependencyCategory Category,
+				UE::AssetRegistry::EDependencyProperty Properties,
+				bool bDuplicate)
+				{
+					if (!bDuplicate)
+					{
+						FDependsNode** NewTargetNode = OldNodeToNewNode.Find(OldTargetNode);
+						if (NewTargetNode && *NewTargetNode != OldTargetNode)
+						{
+							(*NewTargetNode)->AddReferencer(OldSourceNode);
+						}
+					}
+				});
+			OldSourceNode->IterateOverReferencers([OldSourceNode, &OldNodeToNewNode]
+			(FDependsNode* OldReferencer)
+				{
+					FDependsNode** NewReferencerPtr = OldNodeToNewNode.Find(OldReferencer);
+					FDependsNode* NewReferencer = NewReferencerPtr ? *NewReferencerPtr : nullptr;
+					if (NewReferencer && NewReferencer != OldReferencer)
+					{
+						// We have a remapped referencer and we need to add a dependency from the remapped version of that referencer
+						// to OldSourceNode. We have to look up what kind of dependency to add from that old referencer.
+						OldReferencer->IterateOverDependencies([OldSourceNode, NewReferencer](
+							FDependsNode* OldTargetOfOldRefencer,
+							UE::AssetRegistry::EDependencyCategory Category,
+							UE::AssetRegistry::EDependencyProperty Properties,
+							bool bDuplicate)
+							{
+								if (OldTargetOfOldRefencer == OldSourceNode)
+								{
+									NewReferencer->AddDependency(OldSourceNode, Category, Properties);
+								}
+							});
+					}
+				});
+
+			// Now that we've read what we need to from the old targets of the old links on this kept old node, remap all those
+			// links to old targets to instead point to the new remapped targets.
+			OldSourceNode->RemapLinks(OldNodeToNewNode);
+
+			// ConsumeExisting, like InitializeFromExisting, sets consumed nodes to sorted. We also set AllowShrinking to true.
+			// We use this contract because Consume and InitializeFromExisting are almost exclusively used when not gathering,
+			// so we should optimize for read and for memory rather than write speed.
+			OldSourceNode->SetAllowShrinking(true);
+			OldSourceNode->SetIsDependencyListSorted(UE::AssetRegistry::EDependencyCategory::All, true);
+			OldSourceNode->SetIsReferencersSorted(true);
+		}
+	}
+
+	// Steal all of the preallocated buffers; we don't currently track whether we are not using them.
+	for (FAssetData* Buffer : Consume.PreallocatedAssetDataBuffers)
+	{
+		this->PreallocatedAssetDataBuffers.Add(Buffer);
+	}
+	for (FDependsNode* Buffer : Consume.PreallocatedDependsNodeDataBuffers)
+	{
+		this->PreallocatedDependsNodeDataBuffers.Add(Buffer);
+	}
+	for (FAssetPackageData* Buffer : Consume.PreallocatedPackageDataBuffers)
+	{
+		this->PreallocatedPackageDataBuffers.Add(Buffer);
+	}
+
+	// Reset the Consume state, but only delete its data that we didn't steal.
+	if (Consume.PreallocatedAssetDataBuffers.Num())
+	{
+		Consume.PreallocatedAssetDataBuffers.Empty();
+	}
+	else
+	{
+		for (FAssetData* AssetData : AssetsToDeallocate)
+		{
+			delete AssetData;
+		}
+	}
+	Consume.NumAssets = 0;
+
+	if (Consume.PreallocatedDependsNodeDataBuffers.Num())
+	{
+		Consume.PreallocatedDependsNodeDataBuffers.Empty();
+		// Clear memory on the OldNodes, since they are staying around unused.
+		for (FDependsNode* DependsNode : NodesToDeallocate)
+		{
+			DependsNode->Empty();
+		}
+	}
+	else
+	{
+		for (FDependsNode* DependsNode : NodesToDeallocate)
+		{
+			delete DependsNode;
+		}
+	}
+	Consume.NumDependsNodes = 0;
+
+	if (Consume.PreallocatedPackageDataBuffers.Num())
+	{
+		Consume.PreallocatedPackageDataBuffers.Empty();
+	}
+	else
+	{
+		for (FAssetPackageData* AssetPackageData : PackageDatasToDeallocate)
+		{
+			delete AssetPackageData;
+		}
+	}
+	Consume.NumPackageData = 0;
+
+	// Clear the containers on the consumed state; we can't call Reset since it also deletes the elements.
+	Consume.CachedAssetsByPackageName.Empty();
+	Consume.CachedAssetsByPath.Empty();
+	Consume.CachedAssetsByClass.Empty();
+#if UE_ASSETREGISTRY_CACHEDASSETSBYTAG
+	Consume.CachedAssetsByTag.Empty();
+#else
+	Consume.CachedClassesByTag.Empty();
+#endif
+	Consume.CachedDependsNodes.Empty();
+	Consume.CachedPackageData.Empty();
+#if UE_ASSETREGISTRY_INDIRECT_ASSETDATA_POINTERS
+	Consume.IndirectAssetDataArrays.Empty();
+#endif
+	Consume.CachedAssets.Empty();
+
+	this->Shrink();
+}
+
 void FAssetRegistryState::PruneAssetData(const TSet<FName>& RequiredPackages, const TSet<FName>& RemovePackages,
 	const FAssetRegistrySerializationOptions& Options)
 {
@@ -2438,28 +2641,50 @@ void FAssetRegistryState::Load(Archive&& Ar, const FAssetRegistryHeader& Header,
 	Ar << LocalNumAssets;
 
 	// allocate one single block for all asset data structs (to reduce tens of thousands of heap allocations)
-	TArrayView<FAssetData> PreallocatedAssetDataBuffer(new FAssetData[LocalNumAssets], LocalNumAssets);
-	PreallocatedAssetDataBuffers.Add(PreallocatedAssetDataBuffer.GetData());
-
-	// Optimizing serialization of latest asset data format by moving version checking out of SerializeForCache
-	// function and falling back to versioned serialization should we attempt to load an older version of AR
-	// (usually commandlets)
-	if (Version == FAssetRegistryVersion::LatestVersion)
+	TArray<FAssetData*> AllocatedAssets;
+	if (!Options.bMountPointsRequireSeparateAllocationBlocks)
 	{
-		for (FAssetData& NewAssetData : PreallocatedAssetDataBuffer)
+		AllocatedAssets.SetNumUninitialized(LocalNumAssets);
+		FAssetData* PreallocatedAssetDataBuffer = new FAssetData[LocalNumAssets];
+		PreallocatedAssetDataBuffers.Add(PreallocatedAssetDataBuffer);
+		for (int32 Index = 0; Index < LocalNumAssets; ++Index)
 		{
-			NewAssetData.SerializeForCache(Ar);
+			AllocatedAssets[Index] = &PreallocatedAssetDataBuffer[Index];
 		}
 	}
 	else
 	{
-		for (FAssetData& NewAssetData : PreallocatedAssetDataBuffer)
+		AllocatedAssets.SetNumZeroed(LocalNumAssets);
+	}
+
+	for (FAssetData*& NewAssetData : AllocatedAssets)
+	{
+		if (!NewAssetData)
 		{
-			NewAssetData.SerializeForCacheOldVersion(Ar, Version);
+			NewAssetData = new FAssetData();
+		}
+		// Optimizing serialization of latest asset data format by moving version checking out of SerializeForCache
+		// function and falling back to versioned serialization should we attempt to load an older version of AR
+		// (usually commandlets)
+		if (Version == FAssetRegistryVersion::LatestVersion)
+		{
+			NewAssetData->SerializeForCache(Ar);
+		}
+		else
+		{
+			NewAssetData->SerializeForCacheOldVersion(Ar, Version);
 		}
-	}
 
-	SetAssetDatas(PreallocatedAssetDataBuffer, Options);
+		if (Options.ShouldKeepPackage && !Options.ShouldKeepPackage(NewAssetData->PackageName))
+		{
+			if (!PreallocatedAssetDataBuffers.Num())
+			{
+				delete NewAssetData;
+			}
+			NewAssetData = nullptr;
+		}
+	}
+	SetAssetDatas(AllocatedAssets, Options);
 
 	LoadDependencies(Ar, Options, Version);
 
@@ -2469,28 +2694,42 @@ void FAssetRegistryState::Load(Archive&& Ar, const FAssetRegistryHeader& Header,
 	if (LocalNumPackageData > 0)
 	{
 		FAssetPackageData SerializedElement;
-		TArrayView<FAssetPackageData> PreallocatedPackageDataBuffer;
+		FAssetPackageData* PreallocatedPackageDataBuffer = nullptr;
 		if (Options.bLoadPackageData)
 		{
-			PreallocatedPackageDataBuffer = TArrayView<FAssetPackageData>(
-				new FAssetPackageData[LocalNumPackageData], LocalNumPackageData);
-			PreallocatedPackageDataBuffers.Add(PreallocatedPackageDataBuffer.GetData());
 			CachedPackageData.Reserve(LocalNumPackageData);
+			if (!Options.bMountPointsRequireSeparateAllocationBlocks)
+			{
+				PreallocatedPackageDataBuffer = new FAssetPackageData[LocalNumPackageData];
+				PreallocatedPackageDataBuffers.Add(PreallocatedPackageDataBuffer);
+			}
 		}
+
 		for (int32 PackageDataIndex = 0; PackageDataIndex < LocalNumPackageData; PackageDataIndex++)
 		{
 			FName PackageName;
 			Ar << PackageName;
-			FAssetPackageData* NewPackageData;
-			if (Options.bLoadPackageData)
+			FAssetPackageData* NewPackageData = nullptr;
+			if (!Options.bLoadPackageData ||
+				(Options.ShouldKeepPackage && !Options.ShouldKeepPackage(PackageName)))
 			{
-				NewPackageData = &PreallocatedPackageDataBuffer[PackageDataIndex];
-				CachedPackageData.Add(PackageName, NewPackageData);
+				NewPackageData = &SerializedElement;
 			}
 			else
 			{
-				NewPackageData = &SerializedElement;
+				if (!Options.bMountPointsRequireSeparateAllocationBlocks)
+				{
+					NewPackageData = &PreallocatedPackageDataBuffer[PackageDataIndex];
+				}
+				else
+				{
+					NewPackageData = new FAssetPackageData();
+				}
+				CachedPackageData.Add(PackageName, NewPackageData);
+				++NumPackageData;
 			}
+			check(NewPackageData);
+
 			if (Version >= FAssetRegistryVersion::LatestVersion)
 			{
 				NewPackageData->SerializeForCache(Ar);
@@ -2508,6 +2747,10 @@ void FAssetRegistryState::LoadDependencies(FArchive& Ar, const FAssetRegistryLoa
 {
 	if (Version < FAssetRegistryVersion::AddedDependencyFlags)
 	{
+		// ShouldKeepPackage is not supported in this legacy loading path
+		check(!Options.ShouldKeepPackage);
+		// bMountPointsRequireSeparateAllocationBlocks is not supported in this legacy loading path
+		check(!Options.bMountPointsRequireSeparateAllocationBlocks);
 		LegacyLoadDependencies_BeforeFlags(Ar, Options.bLoadDependencies, Version);
 	}
 	else
@@ -2521,11 +2764,11 @@ void FAssetRegistryState::LoadDependencies(FArchive& Ar, const FAssetRegistryLoa
 		{
 			if (Version == FAssetRegistryVersion::LatestVersion)
 			{
-				LoadDependenciesImpl<true>(Ar, Version);
+				LoadDependenciesImpl<true>(Ar, Options, Version);
 			}
 			else
 			{
-				LoadDependenciesImpl<false>(Ar, Version);
+				LoadDependenciesImpl<false>(Ar, Options, Version);
 			}
 		}
 
@@ -2540,7 +2783,8 @@ void FAssetRegistryState::LoadDependencies(FArchive& Ar, const FAssetRegistryLoa
 }
 
 template <bool bLatestVersion>
-void FAssetRegistryState::LoadDependenciesImpl(FArchive& Ar, FAssetRegistryVersion::Type Version)
+void FAssetRegistryState::LoadDependenciesImpl(FArchive& Ar, const FAssetRegistryLoadOptions& Options,
+	FAssetRegistryVersion::Type Version)
 {
 	int32 LocalNumDependsNodes = 0;
 	Ar << LocalNumDependsNodes;
@@ -2550,26 +2794,84 @@ void FAssetRegistryState::LoadDependenciesImpl(FArchive& Ar, FAssetRegistryVersi
 		return;
 	}
 
-	FDependsNode* PreallocatedDependsNodeDataBuffer = new FDependsNode[LocalNumDependsNodes];
-	PreallocatedDependsNodeDataBuffers.Add(PreallocatedDependsNodeDataBuffer);
+	TArray<FDependsNode*> AllocatedNodes;
+	AllocatedNodes.SetNumUninitialized(LocalNumDependsNodes);
+	if (!Options.bMountPointsRequireSeparateAllocationBlocks)
+	{
+		FDependsNode* PreallocatedDependsNodeDataBuffer = new FDependsNode[LocalNumDependsNodes];
+		PreallocatedDependsNodeDataBuffers.Add(PreallocatedDependsNodeDataBuffer);
+		for (int32 Index = 0; Index < LocalNumDependsNodes; ++Index)
+		{
+			AllocatedNodes[Index] = &PreallocatedDependsNodeDataBuffer[Index];
+		}
+	}
+	else
+	{
+		for (FDependsNode*& DependsNode : AllocatedNodes)
+		{
+			DependsNode = new FDependsNode();
+		}
+	}
 	CachedDependsNodes.Reserve(LocalNumDependsNodes);
-	
+
 	TUniqueFunction<FDependsNode*(int32)> GetNodeFromSerializeIndex =
-		[&PreallocatedDependsNodeDataBuffer, LocalNumDependsNodes](int32 Index) -> FDependsNode*
+		[&AllocatedNodes, LocalNumDependsNodes](int32 Index) -> FDependsNode*
 	{
 		if (Index < 0 || LocalNumDependsNodes <= Index)
 		{
 			return nullptr;
 		}
-		return &PreallocatedDependsNodeDataBuffer[Index];
+		return AllocatedNodes[Index];
 	};
 
+	TSet<FDependsNode*> RemoveNodes;
+
 	FDependsNode::FLoadScratch Scratch;
 	for (int32 DependsNodeIndex = 0; DependsNodeIndex < LocalNumDependsNodes; DependsNodeIndex++)
 	{
-		FDependsNode* DependsNode = &PreallocatedDependsNodeDataBuffer[DependsNodeIndex];
+		FDependsNode* DependsNode = AllocatedNodes[DependsNodeIndex];
 		DependsNode->SerializeLoad<bLatestVersion>(Ar, GetNodeFromSerializeIndex, Scratch, Version);
-		CachedDependsNodes.Add(DependsNode->GetIdentifier(), DependsNode);
+		if (Options.ShouldKeepPackage && !Options.ShouldKeepPackage(DependsNode->GetPackageName()))
+		{
+			RemoveNodes.Add(DependsNode);
+		}
+		else
+		{
+			CachedDependsNodes.Add(DependsNode->GetIdentifier(), DependsNode);
+			++NumDependsNodes;
+		}
+	}
+
+	if (!RemoveNodes.IsEmpty())
+	{
+		for (FDependsNode* RemoveNode : RemoveNodes)
+		{
+			RemoveNode->IterateOverDependencies([&RemoveNodes, RemoveNode](FDependsNode* Dependency,
+				UE::AssetRegistry::EDependencyCategory Category,
+				UE::AssetRegistry::EDependencyProperty Properties,
+				bool bDuplicate)
+				{
+					if (!bDuplicate && !RemoveNodes.Contains(Dependency))
+					{
+						Dependency->RemoveReferencer(RemoveNode);
+					}
+				});
+			RemoveNode->IterateOverReferencers([&RemoveNodes, RemoveNode](FDependsNode* Referencer)
+				{
+					if (!RemoveNodes.Contains(Referencer))
+					{
+						Referencer->RemoveDependency(RemoveNode);
+					}
+				});
+			RemoveNode->Empty();
+		}
+		if (Options.bMountPointsRequireSeparateAllocationBlocks)
+		{
+			for (FDependsNode* RemoveNode : RemoveNodes)
+			{
+				delete RemoveNode;
+			}
+		}
 	}
 }
 
@@ -2861,21 +3163,30 @@ void ShrinkMultimap(TMap<KeyType, TArray<ValueType>>& Map)
 	}
 };
 
-void FAssetRegistryState::SetAssetDatas(TArrayView<FAssetData> AssetDatas, const FAssetRegistryLoadOptions& Options)
+void FAssetRegistryState::SetAssetDatas(TConstArrayView<FAssetData*> AssetDatas, const FAssetRegistryLoadOptions& Options)
 {
 	using namespace UE::AssetRegistry::Private;
 
 	UE_CLOG(NumAssets != 0, LogAssetRegistry, Fatal,
 		TEXT("Can only load into empty asset registry states. Load into temporary and append using InitializeFromExisting() instead."));
 
-	NumAssets = AssetDatas.Num();
+	for (FAssetData* AssetData : AssetDatas)
+	{
+		if (AssetData)
+		{
+			++NumAssets;
+		}
+	}
 	
 	auto SetObjectPathCache = [this, &AssetDatas]()
 	{
 		CachedAssets.Empty(NumAssets);
-		for (FAssetData& AssetData: AssetDatas)
+		for (FAssetData* AssetData : AssetDatas)
 		{
-			CachedAssets.Add(&AssetData);
+			if (AssetData)
+			{
+				CachedAssets.Add(AssetData);
+			}
 		}
 		ensure(NumAssets == CachedAssets.Num());
 	};
@@ -2886,9 +3197,12 @@ void FAssetRegistryState::SetAssetDatas(TArrayView<FAssetData> AssetDatas, const
 	{
 		CachedAssetsByPackageName.Empty(AssetDatas.Num());
 #if !UE_ASSETREGISTRY_INDIRECT_ASSETDATA_POINTERS
-		for (FAssetData& AssetData : AssetDatas)
+		for (FAssetData* AssetData : AssetDatas)
 		{
-			CachedAssetsByPackageName.FindOrAdd(AssetData.PackageName).Add(&AssetData);
+			if (AssetData)
+			{
+				CachedAssetsByPackageName.FindOrAdd(AssetData->PackageName).Add(AssetData);
+			}
 		}
 #else
 		CachedAssets.Enumerate([this](FAssetData& AssetData, FAssetDataPtrIndex AssetIndex)
@@ -2903,9 +3217,12 @@ void FAssetRegistryState::SetAssetDatas(TArrayView<FAssetData> AssetDatas, const
 	{
 		CachedAssetsByPath.Empty();
 #if !UE_ASSETREGISTRY_INDIRECT_ASSETDATA_POINTERS
-		for (FAssetData& AssetData : AssetDatas)
+		for (FAssetData* AssetData : AssetDatas)
 		{
-			CachedAssetsByPath.FindOrAdd(AssetData.PackagePath).Add(&AssetData);
+			if (AssetData)
+			{
+				CachedAssetsByPath.FindOrAdd(AssetData->PackagePath).Add(AssetData);
+			}
 		}
 #else
 		CachedAssets.Enumerate([this](FAssetData& AssetData, FAssetDataPtrIndex AssetIndex)
@@ -2921,9 +3238,12 @@ void FAssetRegistryState::SetAssetDatas(TArrayView<FAssetData> AssetDatas, const
 	{
 		CachedAssetsByClass.Empty();
 #if !UE_ASSETREGISTRY_INDIRECT_ASSETDATA_POINTERS
-		for (FAssetData& AssetData : AssetDatas)
+		for (FAssetData* AssetData : AssetDatas)
 		{
-			CachedAssetsByClass.FindOrAdd(AssetData.AssetClassPath).Add(&AssetData);
+			if (AssetData)
+			{
+				CachedAssetsByClass.FindOrAdd(AssetData->AssetClassPath).Add(AssetData);
+			}
 		}
 #else
 		CachedAssets.Enumerate([this](FAssetData& AssetData, FAssetDataPtrIndex AssetIndex)
@@ -2937,11 +3257,14 @@ void FAssetRegistryState::SetAssetDatas(TArrayView<FAssetData> AssetDatas, const
 #if UE_ASSETREGISTRY_CACHEDASSETSBYTAG
 		CachedAssetsByTag.Empty();
 #if !UE_ASSETREGISTRY_INDIRECT_ASSETDATA_POINTERS
-		for (FAssetData& AssetData : AssetDatas)
+		for (FAssetData* AssetData : AssetDatas)
 		{
-			for (const TPair<FName, FAssetTagValueRef>& Pair : AssetData.TagsAndValues)
+			if (AssetData)
 			{
-				CachedAssetsByTag.FindOrAdd(Pair.Key).Add(&AssetData);
+				for (const TPair<FName, FAssetTagValueRef>& Pair : AssetData->TagsAndValues)
+				{
+					CachedAssetsByTag.FindOrAdd(Pair.Key).Add(AssetData);
+				}
 			}
 		}
 #else
@@ -2961,11 +3284,14 @@ void FAssetRegistryState::SetAssetDatas(TArrayView<FAssetData> AssetDatas, const
 		}
 #else
 		CachedClassesByTag.Empty();
-		for (FAssetData& AssetData : AssetDatas)
+		for (FAssetData* AssetData : AssetDatas)
 		{
-			for (const TPair<FName, FAssetTagValueRef>& Pair : AssetData.TagsAndValues)
+			if (AssetData)
 			{
-				CachedClassesByTag.FindOrAdd(Pair.Key).Add(AssetData.AssetClassPath);
+				for (const TPair<FName, FAssetTagValueRef>& Pair : AssetData->TagsAndValues)
+				{
+					CachedClassesByTag.FindOrAdd(Pair.Key).Add(AssetData->AssetClassPath);
+				}
 			}
 		}
 
diff --git a/Engine/Source/Runtime/AssetRegistry/Private/DependsNode.cpp b/Engine/Source/Runtime/AssetRegistry/Private/DependsNode.cpp
--- a/Engine/Source/Runtime/AssetRegistry/Private/DependsNode.cpp
+++ b/Engine/Source/Runtime/AssetRegistry/Private/DependsNode.cpp
@@ -6,6 +6,9 @@
 #include "AssetRegistryPrivate.h"
 #include "Misc/StringBuilder.h"
 
+template <uint32 FlagWidth>
+void SortDependencyList(TArray<FDependsNode*>& Dependencies, TBitArray<>* Flags, EAllowShrinking ShrinkPolicy);
+
 void FDependsNode::PrintNode() const
 {
 	UE_LOG(LogAssetRegistry, Log, TEXT("*** Printing DependsNode: %s ***"), *GetIdentifier().ToString());
@@ -702,6 +705,37 @@ void FDependsNode::RemoveLinks(const TUniqueFunction<bool(const FDependsNode*)>&
 	::RemoveAll<0>(ShouldRemove, Referencers, nullptr, ReferencersIsSorted, ShrinkPolicy);
 }
 
+template <uint32 FlagWidth>
+void RemapLinks(const TMap<FDependsNode*, FDependsNode*>& OldNodeToNewNode,
+	TArray<FDependsNode*>& Dependencies,
+	TBitArray<>* FlagBits,
+	bool IsSorted,
+	EAllowShrinking ShrinkPolicy)
+{
+	for (FDependsNode*& OldNode : Dependencies)
+	{
+		FDependsNode*const* NewNodePtr = OldNodeToNewNode.Find(OldNode);
+		OldNode = NewNodePtr ? *NewNodePtr : nullptr;
+	}
+	::RemoveAll<FlagWidth>([](const FDependsNode* Node)
+		{
+			return Node == nullptr;
+		},
+		Dependencies, FlagBits, false /* IsSorted */, ShrinkPolicy);
+	if (IsSorted)
+	{
+		::SortDependencyList<FlagWidth>(Dependencies, FlagBits, ShrinkPolicy);
+	}
+}
+
+void FDependsNode::RemapLinks(const TMap<FDependsNode*, FDependsNode*>& OldNodeToNewNode)
+{
+	const EAllowShrinking ShrinkPolicy = bAllowShrinking ? EAllowShrinking::Yes : EAllowShrinking::No;
+	::RemapLinks<PackageFlagWidth>(OldNodeToNewNode, PackageDependencies, &PackageFlags, PackageIsSorted, ShrinkPolicy);
+	::RemapLinks<SearchableNameFlagWidth>(OldNodeToNewNode, NameDependencies, nullptr, SearchableNameIsSorted, ShrinkPolicy);
+	::RemapLinks<ManageFlagWidth>(OldNodeToNewNode, ManageDependencies, &ManageFlags, ManageIsSorted, ShrinkPolicy);
+	::RemapLinks<0>(OldNodeToNewNode, Referencers, nullptr, ReferencersIsSorted, ShrinkPolicy);
+}
 
 bool FDependsNode::ContainsDependency(const FDependsNode* InDependency,
 	UE::AssetRegistry::EDependencyCategory Category, const UE::AssetRegistry::FDependencyQuery& Flags) const
diff --git a/Engine/Source/Runtime/AssetRegistry/Private/DependsNode.h b/Engine/Source/Runtime/AssetRegistry/Private/DependsNode.h
--- a/Engine/Source/Runtime/AssetRegistry/Private/DependsNode.h
+++ b/Engine/Source/Runtime/AssetRegistry/Private/DependsNode.h
@@ -6,6 +6,7 @@
 #include "AssetRegistry/AssetData.h"
 #include "AssetRegistry/IAssetRegistry.h"
 #include "Containers/BitArray.h"
+#include "Containers/Map.h"
 #include "Misc/AssetRegistryInterface.h"
 #include "PropertyCombinationSet.h"
 
@@ -28,6 +29,7 @@ public:
 public:
 	FDependsNode();
 	FDependsNode(const FAssetIdentifier& InIdentifier);
+	void Empty();
 
 	/** Prints the dependencies and referencers for this node to the log */
 	void PrintNode() const;
@@ -89,6 +91,12 @@ public:
 	void RemoveManageReferencesToNode();
 	/** Remove all nodes from referencers and dependencies for which ShouldRemove returns true */
 	void RemoveLinks(const TUniqueFunction<bool(const FDependsNode*)>& ShouldRemove);
+	/**
+	 * Replace all nodes in referencers and dependencies with the newnode for them in OldNodeToNewNode, and remove any
+	 * that have no entry or that remap to null.
+	 */
+	void RemapLinks(const TMap<FDependsNode*, FDependsNode*>& OldNodeToNewNode);
+
 	/** Returns number of connections this node has, both references and dependencies */
 	int32 GetConnectionCount() const;
 	/** Returns amount of memory used by the arrays */
@@ -198,8 +206,6 @@ private:
 	 */
 	void PrintReferencersRecursive(const FString& Indent, TSet<const FDependsNode*>& VisitedNodes) const;
 
-	void ConstructFlags();
-
 	template <bool bLatestVersion>
 	void SerializeLoad_PackageDependencies(FArchive& Ar,
 		const TUniqueFunction<FDependsNode* (int32)>& GetNodeFromSerializeIndex,
@@ -227,14 +233,14 @@ private:
 	TBitArray<> ManageFlags;
 
 	// Transient flags that are not serialized
-	uint32 PackageIsSorted : 1;
-	uint32 SearchableNameIsSorted : 1;
-	uint32 ManageIsSorted : 1;
-	uint32 ReferencersIsSorted : 1;
-	uint32 DependenciesInitialized : 1;
-	uint32 bScriptPath : 1;
+	uint32 PackageIsSorted : 1 = 1;
+	uint32 SearchableNameIsSorted : 1 = 1;
+	uint32 ManageIsSorted : 1 = 1;
+	uint32 ReferencersIsSorted : 1 = 1;
+	uint32 DependenciesInitialized : 1 = 0;
+	uint32 bScriptPath : 1 = 0;
 	/** Used to control FDependsNodeList shrinking behaviour */
-	uint32 bAllowShrinking : 1;
+	uint32 bAllowShrinking : 1 = 1;
 
 	friend FDependsNodeReservations;
 };
@@ -268,15 +274,20 @@ struct FDependsNodeReservations
 
 inline FDependsNode::FDependsNode()
 {
-	ConstructFlags();
 }
 
 inline FDependsNode::FDependsNode(const FAssetIdentifier& InIdentifier)
 {
-	ConstructFlags();
 	SetIdentifier(InIdentifier);
 }
 
+inline void FDependsNode::Empty()
+{
+	SetAllowShrinking(true);
+	ClearDependencies();
+	ClearReferencers();
+}
+
 inline FName FDependsNode::GetPackageName() const
 {
 	return Identifier.PackageName;
@@ -369,13 +380,3 @@ inline UE::AssetRegistry::EDependencyProperty FDependsNode::ByteToManageProperti
 		| (static_cast<uint32>(UE::AssetRegistry::EDependencyProperty::CookRule) * ((Bits & 0x02) != 0))
 		);
 }
-
-inline void FDependsNode::ConstructFlags()
-{
-	PackageIsSorted = 1;
-	SearchableNameIsSorted = 1;
-	ManageIsSorted = 1;
-	ReferencersIsSorted = 1;
-	DependenciesInitialized = 0;
-	bScriptPath = 0;
-}
diff --git a/Engine/Source/Runtime/AssetRegistry/Public/AssetRegistry/AssetRegistryState.h b/Engine/Source/Runtime/AssetRegistry/Public/AssetRegistry/AssetRegistryState.h
--- a/Engine/Source/Runtime/AssetRegistry/Public/AssetRegistry/AssetRegistryState.h
+++ b/Engine/Source/Runtime/AssetRegistry/Public/AssetRegistry/AssetRegistryState.h
@@ -15,6 +15,7 @@
 #include "CoreTypes.h"
 #include "HAL/PlatformCrt.h"
 #include "Misc/AssetRegistryInterface.h"
+#include "Templates/Function.h"
 #include "Templates/UnrealTemplate.h"
 #include "UObject/NameTypes.h"
 #include "UObject/SoftObjectPath.h"
@@ -126,6 +127,11 @@ struct FAssetRegistryLoadOptions
 	FAssetRegistryLoadOptions() = default;
 	explicit FAssetRegistryLoadOptions(const FAssetRegistrySerializationOptions& Options);
 
+	/**
+	 *  If non-empty, called on every FAssetData, AssetPackageData, and FDependsNode loaded from
+	 * disk, keeping them only if ShouldKeepPackage returns true on their packagename.
+	 */
+	TFunction<bool(FName PackageName)> ShouldKeepPackage;
 	/**
 	 * Optional pointer to the FStore previously created by Load, in cases where the same AssetRegistry.bin
 	 * file is being reloaded. Caution: If this Store is not identical to the one that would be loaded from disk,
@@ -137,6 +143,11 @@ struct FAssetRegistryLoadOptions
 	int32 ParallelWorkers = 0;
 	bool bLoadDependencies = true;
 	bool bLoadPackageData = true;
+	/** 
+	 * If true, AssetDatas, AssetPackageDatas and DependsNodes will be individually allocated rather than allocated in a single block,
+	 * so that they can be individually freed when their mountpoint is unloaded.
+	 */
+	bool bMountPointsRequireSeparateAllocationBlocks = false;
 };
 
 /** Extra data about the load of an AssetRegistryState from disk that is optionally returned from Load. */
@@ -513,6 +524,17 @@ public:
 
 	void InitializeFromExisting(const FAssetRegistryState& Existing, const FAssetRegistrySerializationOptions& Options,
 		EInitializationMode InitializationMode = EInitializationMode::Rebuild, FAssetRegistryAppendResult* OutAppendResult = nullptr);
+	/**
+	 * Similar to InitializeFromExisting, but it moves the allocations over from InState into *this rather than
+	 * copying them. It also shares the InState's existing FStore (if any) rather than creating a new one.
+	 * This function is used for faster consumption of states loaded internally by the AssetRegistry.
+	 * This function is less general than InitializeFromExisting to reduce complexity; it only supports
+	 * EInitializationMode::OnlyUpdateNew and does not support filter flags on the input Options; all instances
+	 * of FAssetDatas, FAssetPackageDatas, and FDependsNodes are kept
+	 * The InState is reset to empty by this function.
+	 */
+	void ConsumeExisting(FAssetRegistryState&& Existing, const FAssetRegistrySerializationOptions& Options,
+		FAssetRegistryAppendResult* OutAppendResult = nullptr);
 
 	/** 
 	 * Prunes an asset cache, this removes asset data, nodes, and package data that isn't needed. 
@@ -630,7 +652,7 @@ private:
 	void Load(Archive&& Ar, const FAssetRegistryHeader& Header, const FAssetRegistryLoadOptions& Options);
 
 	/** Initialize the lookup maps */
-	void SetAssetDatas(TArrayView<FAssetData> AssetDatas, const FAssetRegistryLoadOptions& Options);
+	void SetAssetDatas(TConstArrayView<FAssetData*> AssetDatas, const FAssetRegistryLoadOptions& Options);
 
 	/** Find the first non-redirector dependency node starting from InDependency. */
 	FDependsNode* ResolveRedirector(FDependsNode* InDependency, const FAssetDataMap& InAllowedAssets,
@@ -662,7 +684,7 @@ private:
 
 	void LoadDependencies(FArchive& Ar, const FAssetRegistryLoadOptions& Options, FAssetRegistryVersion::Type Version);
 	template <bool bLatestVersion>
-	void LoadDependenciesImpl(FArchive& Ar, FAssetRegistryVersion::Type Version);
+	void LoadDependenciesImpl(FArchive& Ar, const FAssetRegistryLoadOptions& Options, FAssetRegistryVersion::Type Version);
 	void LegacyLoadDependencies_BeforeFlags(FArchive& Ar, bool bSerializeDependencies, FAssetRegistryVersion::Type Version);
 
 	void SetTagsOnExistingAsset(FAssetData* AssetData, FAssetDataTagMap&& NewTags);
