AssetRegistry: Premade: Delay the load of the premade until GConfig is ready, because ActiveMountsLoader requires reading GConfig to know how to load it
#rnx

#Synced-CL 50820990

diff --git a/Engine/Source/Runtime/AssetRegistry/Private/ARLoader.cpp b/Engine/Source/Runtime/AssetRegistry/Private/ARLoader.cpp
--- a/Engine/Source/Runtime/AssetRegistry/Private/ARLoader.cpp
+++ b/Engine/Source/Runtime/AssetRegistry/Private/ARLoader.cpp
@@ -175,7 +175,8 @@ void FPreloader::DelayedInitialize()
 
 	if (TrySetPath())
 	{
-		KickPreload();
+		FScopeLock Lock(&StateLock);
+		KickPreloadOrDelayIfConfigNotReady();
 	}
 	else
 	{
@@ -188,7 +189,7 @@ void FPreloader::DelayedInitialize()
 				FScopeLock Lock(&StateLock);
 				if (LoadState == EState::NotFound && TrySetPath(Pak))
 				{
-					KickPreload();
+					KickPreloadOrDelayIfConfigNotReady();
 					// Remove the callback from OnPakFileMounted2 to avoid wasting time in all future PakFile mounts
 					// Do not access any of the lambda captures after the call to Remove, because deallocating the 
 					// DelegateHandle also deallocates our lambda captures
@@ -200,12 +201,36 @@ void FPreloader::DelayedInitialize()
 	}
 }
 
+void FPreloader::KickPreloadOrDelayIfConfigNotReady()
+{
+	// Called within StateLock.
+	// Only called from DelayedInitialize or its PakMountedDelegate, and only called if ARPath has been found.
+	// KickPreload requires ARPath has been found.
+	check(!ARPath.IsEmpty());
+	if (GConfig)
+	{
+		KickPreload();
+	}
+	else
+	{
+		LoadState = EState::ConfigNotReady;
+		FCoreDelegates::TSConfigReadyForUse().AddLambda([this]()
+			{
+				FScopeLock Lock(&StateLock);
+				if (LoadState == EState::ConfigNotReady)
+				{
+					KickPreload();
+				}
+			});
+	}
+}
+
 void FPreloader::KickPreload()
 {
 	CSV_TRACE_REGION_BEGIN(TEXT("Preload"), TEXT("AssetRegistry"))
 	TRACE_CPUPROFILER_EVENT_SCOPE(FCookedAssetRegistryPreloader::KickPreload);
 	// Called from Within the Lock
-	check(LoadState == EState::NotFound && !ARPath.IsEmpty());
+	check((LoadState == EState::NotFound || LoadState == EState::ConfigNotReady) && !ARPath.IsEmpty());
 	LoadState = EState::Loading;
 	PreloadReady->Reset();
 	Async(EAsyncExecution::TaskGraph, [this]() { TryLoadAsync(); });
@@ -291,8 +316,8 @@ FPreloader::EConsumeResult FPreloader::ConsumeInternal(FConsumeFunction&& Consum
 	}
 
 	// TryAsyncLoad might not yet have set state to Loaded
-	check(LoadState == EState::Loaded || LoadState == EState::Loading || LoadState == EState::NotFound);
-	ELoadResult LocalResult = LoadState == EState::NotFound ? ELoadResult::NotFound : LoadResult;
+	check(LoadState == EState::Loaded || LoadState == EState::Loading || LoadState == EState::NotFound || LoadState == EState::ConfigNotReady);
+	ELoadResult LocalResult = (LoadState == EState::NotFound || LoadState == EState::ConfigNotReady) ? ELoadResult::NotFound : LoadResult;
 	LoadState = EState::Consumed;
 
 	// No further async threads exist that will read/write payload at this point so we can use it outside the lock
diff --git a/Engine/Source/Runtime/AssetRegistry/Private/ARLoader.h b/Engine/Source/Runtime/AssetRegistry/Private/ARLoader.h
--- a/Engine/Source/Runtime/AssetRegistry/Private/ARLoader.h
+++ b/Engine/Source/Runtime/AssetRegistry/Private/ARLoader.h
@@ -76,6 +76,7 @@ private:
 		WillNeverPreload,
 		LoadSynchronous,
 		NotFound,
+		ConfigNotReady,
 		Loading,
 		Loaded,
 		Consumed,
@@ -86,6 +87,7 @@ private:
 	bool TrySetPath(const IPakFile& Pak);
 	ELoadResult TryLoad();
 	void DelayedInitialize();
+	void KickPreloadOrDelayIfConfigNotReady();
 	void KickPreload();
 	void TryLoadAsync();
 	EConsumeResult ConsumeInternal(FConsumeFunction&& ConsumeSynchronous, FConsumeFunction&& ConsumeAsynchronous);
