AssetRegistry: ActiveMountsLoader: Add ActiveMount* functions to unload and reload from AssetRegistry.bin file data about assets inside of specified mount points. This changelist submits the statemachine that manages the load; it is not yet ready to be enabled until some upcoming support changes are submitted in AssetRegistryState serialization.
#rnx

#Synced-CL 50811273

diff --git a/Engine/Config/BaseEngine.ini b/Engine/Config/BaseEngine.ini
--- a/Engine/Config/BaseEngine.ini
+++ b/Engine/Config/BaseEngine.ini
@@ -2901,6 +2901,11 @@ bSerializePackageData=false
 bFilterAssetDataWithNoTags=false
 bFilterDependenciesWithNoTags=false
 bFilterSearchableNames=true
+; ActiveMountLoading: Enable this field on projects with thousands of plugins to save memory by unloading AR data,
+; but requires c++ handling as well, see ActiveMountLoading documentation in IAssetRegistry.h.
+bActiveMountLoadingEnabled_Client=false
+; Configures ActiveMountLoading to use default unloading rules, see ActiveMountUnloading in IAssetRegistry.h.
+bActiveMountLoadingAuto_Client=true
 
 ; The editor can run the game or dedicated server on uncooked content, via UnrealEditor.exe <ProjectName> -game
 ; This is bad for performance but is supported for use in PlayInEditor
diff --git a/Engine/Source/Runtime/AssetRegistry/Private/ARLoader.cpp b/Engine/Source/Runtime/AssetRegistry/Private/ARLoader.cpp
--- a/Engine/Source/Runtime/AssetRegistry/Private/ARLoader.cpp
+++ b/Engine/Source/Runtime/AssetRegistry/Private/ARLoader.cpp
@@ -25,8 +25,10 @@
 #include "Misc/Paths.h"
 #include "Misc/ScopeExit.h"
 #include "Misc/ScopeLock.h"
+#include "Misc/StringBuilder.h"
 #include "ProfilingDebugging/CpuProfilerTrace.h"
 #include "ProfilingDebugging/CsvProfiler.h"
+#include "Serialization/LargeMemoryReader.h"
 #include "Templates/UnrealTemplate.h"
 
 namespace UE::AssetRegistry::Premade
@@ -57,6 +59,17 @@ bool CanLoadAsync()
 	return FPlatformProcess::SupportsMultithreading() && FTaskGraphInterface::IsRunning();
 }
 
+FAssetRegistryLoadOptions GetLoadOptions()
+{
+	FAssetRegistryLoadOptions Options;
+	// If ActiveMounts are enabled we need to be able to unload all Assets in a mountpoint at a time, so 
+	// set bMountPointsRequireSeparateAllocationBlocks=true. Otherwise allocate all assets in a single block.
+	Options.bMountPointsRequireSeparateAllocationBlocks =
+		UE::AssetRegistry::Impl::FActiveMountsLoaderWrapper::IsEnabledByConfig();
+
+	return Options;
+}
+
 /** Returns the paths to possible Premade AssetRegistry files, ordered from highest priority to lowest. */
 static TArray<FString, TInlineAllocator<2>> GetPriorityPaths()
 {
@@ -147,14 +160,30 @@ ELoadResult FPreloader::TryLoad()
 	LLM_SCOPE(ELLMTag::AssetRegistry);
 	checkf(!ARPath.IsEmpty(), TEXT("TryLoad must not be called until after TrySetPath has succeeded."));
 
-	FAssetRegistryLoadOptions Options;
+	FAssetRegistryLoadOptions Options = UE::AssetRegistry::Premade::GetLoadOptions();
 	const int32 ThreadReduction = 2; // This thread + main thread already has work to do 
 	int32 MaxWorkers = CanLoadAsync() ? FPlatformMisc::NumberOfCoresIncludingHyperthreads() - ThreadReduction : 0;
 	Options.ParallelWorkers = FMath::Clamp(MaxWorkers, 0, 16);
-	bool bLoadSucceeded = FAssetRegistryState::LoadFromDisk(*ARPath, Options, Payload);
-	UE_CLOG(!bLoadSucceeded, LogAssetRegistry, Warning, TEXT("Premade AssetRegistry path %s existed but failed to load."), *ARPath);
-	UE_CLOG(bLoadSucceeded, LogAssetRegistry, Log, TEXT("Premade AssetRegistry loaded from '%s'"), *ARPath);
-	LoadResult = bLoadSucceeded ? ELoadResult::Succeeded : ELoadResult::FailedToLoad;
+
+	TUniquePtr<FArchive> FileArchive = UE::AssetRegistry::Impl::OpenReadStateFileArchive(ARPath);
+	if (!FileArchive)
+	{
+		UE_LOG(LogAssetRegistry, Warning, TEXT("Premade AssetRegistry path %s existed but was unreadable."), *ARPath);
+		LoadResult = ELoadResult::FailedToLoad;
+		return LoadResult;
+	}
+	FAssetRegistryLoadResults LoadResults;
+	if (!Payload.ARState.Load(*FileArchive, Options, LoadResults))
+	{
+		UE_LOG(LogAssetRegistry, Warning, TEXT("Premade AssetRegistry path %s existed but failed to load."), *ARPath);
+		LoadResult = ELoadResult::FailedToLoad;
+		return LoadResult;
+	}
+	Payload.StateFilePath = ARPath;
+	Payload.StoreUsedByState = MoveTemp(LoadResults.Store);
+
+	UE_LOG(LogAssetRegistry, Log, TEXT("Premade AssetRegistry loaded from '%s'."), *ARPath);
+	LoadResult = ELoadResult::Succeeded;
 	return LoadResult;
 }
 
@@ -245,7 +274,6 @@ void FPreloader::TryLoadAsync()
 	// Trigger outside the lock so that a locked Consume function that is waiting on PreloadReady can wait inside the lock.
 	PreloadReady->Trigger();
 
-	FString StateFilePath;
 	FConsumeFunction LocalConsumeCallback;
 	{
 		FScopeLock Lock(&StateLock);
@@ -255,7 +283,6 @@ void FPreloader::TryLoadAsync()
 			LoadState = EState::Loaded;
 			if (ConsumeCallback)
 			{
-				StateFilePath = ARPath;
 				LocalConsumeCallback = MoveTemp(ConsumeCallback);
 				ConsumeCallback.Reset();
 				LoadState = EState::Consumed;
@@ -267,7 +294,7 @@ void FPreloader::TryLoadAsync()
 	{
 		// No further threads will read/write payload at this point (until destructor, which is called after all async threads are complete
 		// so we can use it outside the lock
-		LocalConsumeCallback(LocalResult, MoveTemp(Payload), StateFilePath);
+		LocalConsumeCallback(LocalResult, MoveTemp(Payload));
 		Shutdown();
 	}
 }
@@ -283,7 +310,7 @@ FPreloader::EConsumeResult FPreloader::ConsumeInternal(FConsumeFunction&& Consum
 	{
 		Lock.Unlock(); // Unlock before calling external code in Consume callback
 		ELoadResult LocalResult = (LoadState == EState::Consumed || ConsumeCallback) ? ELoadResult::AlreadyConsumed : ELoadResult::Inactive;
-		ConsumeSynchronous(LocalResult, FAssetRegistryState(), FString());
+		ConsumeSynchronous(LocalResult, UE::AssetRegistry::Impl::FAssetProviderInitData());
 		return EConsumeResult::Failed;
 	}
 
@@ -291,9 +318,8 @@ FPreloader::EConsumeResult FPreloader::ConsumeInternal(FConsumeFunction&& Consum
 	{
 		ELoadResult LocalResult = TrySetPath() ? TryLoad() : ELoadResult::NotFound;
 		LoadState = EState::Consumed;
-		FString StateFilePath = ARPath;
 		Lock.Unlock(); // Unlock before calling external code in Consume callback
-		ConsumeSynchronous(LocalResult, MoveTemp(Payload), StateFilePath);
+		ConsumeSynchronous(LocalResult, MoveTemp(Payload));
 		Shutdown(); // Shutdown can be called outside the lock since AsyncThread doesn't exist
 		return LocalResult == ELoadResult::Succeeded ? EConsumeResult::Succeeded : EConsumeResult::Failed;
 	}
@@ -321,9 +347,8 @@ FPreloader::EConsumeResult FPreloader::ConsumeInternal(FConsumeFunction&& Consum
 	LoadState = EState::Consumed;
 
 	// No further async threads exist that will read/write payload at this point so we can use it outside the lock
-	FString StateFilePath = ARPath;
 	Lock.Unlock(); // Unlock before calling external code in Consume callback
-	ConsumeSynchronous(LocalResult, MoveTemp(Payload), StateFilePath);
+	ConsumeSynchronous(LocalResult, MoveTemp(Payload));
 	Shutdown(); // Shutdown can be called outside the lock since we have set state to Consumed and the Async thread will notice and exit
 	return LocalResult == ELoadResult::Succeeded ? EConsumeResult::Succeeded : EConsumeResult::Failed;
 }
@@ -391,12 +416,13 @@ void FAsyncConsumer::Wait(UE::AssetRegistry::FInterfaceWriteScopeLock& ScopeLock
 }
 
 void FAsyncConsumer::Consume(UAssetRegistryImpl& UARI, UE::AssetRegistry::Impl::FEventContext& EventContext,
+	UE::AssetRegistry::Impl::FActiveMountEvents& ActiveMountEvents,
 	UE::AssetRegistry::FInterfaceWriteScopeLock& ProofOfWriteLock, ELoadResult LoadResult,
-	FAssetRegistryState&& ARState, FStringView StateFilePath)
+	UE::AssetRegistry::Impl::FAssetProviderInitData&& StateInitData)
 {
 	// Called within the lock
-	UARI.GuardedData.LoadPremadeAssetRegistry(UARI, EventContext, ProofOfWriteLock, LoadResult,
-		MoveTemp(ARState), StateFilePath);
+	UARI.GuardedData.LoadPremadeAssetRegistry(UARI, EventContext, ActiveMountEvents, ProofOfWriteLock, LoadResult,
+		MoveTemp(StateInitData));
 	check(ReferenceCount >= 1);
 	check(Consumed != nullptr);
 	Consumed->Trigger();
@@ -410,3 +436,63 @@ void FAsyncConsumer::Consume(UAssetRegistryImpl& UARI, UE::AssetRegistry::Impl::
 }
 
 } // namespace UE::AssetRegistry::Premade
+
+namespace UE::AssetRegistry::Impl
+{
+
+TUniquePtr<FArchive> OpenReadStateFileArchive(FStringView FilePath)
+{
+	// Read the entire file into memory, this is faster than serializing repeatedly from a FileReader.
+	TUniquePtr<FArchive> FileReader(IFileManager::Get().CreateFileReader(*WriteToString<512>(FilePath)));
+	if (!FileReader)
+	{
+		return TUniquePtr<FArchive>();
+	}
+	// FLargeMemoryReader can take ownership of the Data buffer; it requires that we allocate with malloc.
+	int64 Size = FileReader->TotalSize();
+	if (Size <= 0)
+	{
+		return TUniquePtr<FArchive>();
+	}
+	uint8* Data = static_cast<uint8*>(FMemory::Malloc(Size));
+	ON_SCOPE_EXIT
+	{
+		if (Data)
+		{
+			FMemory::Free(Data);
+		}
+	};
+
+	FileReader->Serialize(Data, Size);
+	if (FileReader->IsError())
+	{
+		return TUniquePtr<FArchive>();
+	}
+
+	TUniquePtr<FArchive> Result(new FLargeMemoryReader(Data, Size, ELargeMemoryReaderFlags::TakeOwnership));
+	Data = nullptr; // detach Data from our call to FMemory::Free
+	return Result;
+}
+
+FAssetProviderInitData::FAssetProviderInitData(FAssetProviderInitData&& Other)
+{
+	*this = MoveTemp(Other);
+}
+
+FAssetProviderInitData& FAssetProviderInitData::operator=(FAssetProviderInitData&& Other)
+{
+	ARState = MoveTemp(Other.ARState);
+	StateFilePath = MoveTemp(Other.StateFilePath);
+	StoreUsedByState = MoveTemp(Other.StoreUsedByState);
+	Other.Reset();
+	return *this;
+}
+
+void FAssetProviderInitData::Reset()
+{
+	ARState.Reset();
+	StateFilePath.Empty();
+	StoreUsedByState = FixedTagPrivate::FWeakStorePtr();
+}
+
+} // namespace UE::AssetRegistry::Impl
\ No newline at end of file
diff --git a/Engine/Source/Runtime/AssetRegistry/Private/ARLoader.h b/Engine/Source/Runtime/AssetRegistry/Private/ARLoader.h
--- a/Engine/Source/Runtime/AssetRegistry/Private/ARLoader.h
+++ b/Engine/Source/Runtime/AssetRegistry/Private/ARLoader.h
@@ -3,12 +3,21 @@
 #pragma once
 
 #include "ARTypes.h"
+#include "Async/RecursiveMutex.h"
 #include "AssetRegistry/AssetRegistryState.h"
+#include "Containers/ArrayView.h"
 #include "Containers/StringView.h"
 #include "Containers/UnrealString.h"
 #include "Delegates/Delegate.h"
 #include "HAL/CriticalSection.h"
 #include "Misc/Optional.h"
+#include "Misc/PackageName.h"
+#include "Serialization/Archive.h"
+#include "Templates/RefCounting.h"
+#include "Templates/SharedPointer.h"
+#include "Templates/UniquePtr.h"
+
+#include <atomic>
 
 // Preloading a premade AssetRegistry is disabled in editor unless the target specifies it in UBT, because the editor
 // usually gathers to create an up to date AssetRegistry rather than using a previously cooked AssetRegistry.bin.
@@ -17,11 +26,40 @@
 #endif
 
 class FEvent;
+class FPakPlatformFile;
 class IPakFile;
 class UAssetRegistryImpl;
+namespace UE::AssetRegistry { class FAssetRegistryImpl; }
+namespace UE::AssetRegistry::Impl { class FActiveMountEvents; }
+namespace UE::AssetRegistry::Impl { class FActiveMountsLoader; }
+namespace UE::AssetRegistry::Impl { class FAssetProvider; }
+namespace UE::AssetRegistry::Impl { class FMountPointData; }
 namespace UE::AssetRegistry::Impl { struct FEventContext; }
 struct FDelayedAutoRegisterHelper;
 
+namespace UE::AssetRegistry::Impl
+{
+
+/**
+ * Data about the premade AssetRegistryState passed through from its load into LoadPremadeAssetRegistry and the
+ * ActiveMountsLoader 
+ */
+struct FAssetProviderInitData
+{
+	FAssetProviderInitData() = default;
+	FAssetProviderInitData(const FAssetProviderInitData& Other) = delete;
+	FAssetProviderInitData(FAssetProviderInitData&& Other);
+	FAssetProviderInitData& operator=(const FAssetProviderInitData& Other) = delete;
+	FAssetProviderInitData& operator=(FAssetProviderInitData&& Other);
+	void Reset();
+
+	FAssetRegistryState ARState;
+	FString StateFilePath;
+	FixedTagPrivate::FWeakStorePtr StoreUsedByState;
+};
+
+}
+
 namespace UE::AssetRegistry::Premade
 {
 
@@ -45,8 +83,8 @@ public:
 		Failed,
 		Deferred
 	};
-	using FConsumeFunction = TFunction<void(ELoadResult LoadResult, FAssetRegistryState&& ARState,
-		FStringView StateFilePath)>;
+	using FConsumeFunction = TFunction<void(ELoadResult LoadResult,
+		UE::AssetRegistry::Impl::FAssetProviderInitData&& StateInitData)>;
 
 	FPreloader();
 	~FPreloader();
@@ -106,12 +144,13 @@ private:
 	FString ARPath;
 
 	/**
-	 * The ARState loaded from disk. Owned exclusively by either the first Consume or by TryAsyncLoad.
+	 * The set of (ARState, AssetRegistryLoadResults, and other metadata) that was loaded from disk.
+	 * Owned exclusively by either the first Consume or by TryAsyncLoad.
 	 * If LoadState is never set to Loading, this state is read/written only by the first thread to call Consume.
 	 * If LoadState is set to Loading (which happens before threading starts), the thread running TryAsyncLoad
-	 * owns this payload until it triggers PayloadReady, after which ownership returns to the first thread to call Consume.
+	 * owns this payload until it triggers PreloadReady, after which ownership returns to the first thread to call Consume.
 	 */
-	FAssetRegistryState Payload;
+	UE::AssetRegistry::Impl::FAssetProviderInitData Payload;
 
 	FDelegateHandle PakMountedDelegate;
 
@@ -128,6 +167,8 @@ private:
 /** Returns whether the given executable configuration supports AssetRegistry Preloading. Called before Main. */
 extern bool IsEnabled();
 extern bool CanLoadAsync();
+/** Returns the LoadOptions used to load the premade AssetRegistry.bin. */
+FAssetRegistryLoadOptions GetLoadOptions();
 
 extern FPreloader GPreloader;
 
@@ -148,8 +189,9 @@ struct FAsyncConsumer
 	void Wait(FInterfaceWriteScopeLock& ScopeLock);
 	/** Callback from the async thread to consume the Premade ARState */
 	void Consume(UAssetRegistryImpl& UARI, UE::AssetRegistry::Impl::FEventContext& EventContext,
+		UE::AssetRegistry::Impl::FActiveMountEvents& ActiveMountEventContext,
 		UE::AssetRegistry::FInterfaceWriteScopeLock& ProofOfWriteLock,
-		ELoadResult LoadResult, FAssetRegistryState&& ARState, FStringView StateFilePath);
+		ELoadResult LoadResult, UE::AssetRegistry::Impl::FAssetProviderInitData&& StateInitData);
 
 private:
 	/**
@@ -161,3 +203,193 @@ private:
 };
 
 } // namespace UE::AssetRegistry::Premade
+
+namespace UE::AssetRegistry::Impl
+{
+
+/** Open the given AssetRegistry.bin file using the performant archive we use for loading it. */
+TUniquePtr<FArchive> OpenReadStateFileArchive(FStringView FilePath);
+
+/** Which event triggered our need to register the premade AssetRegistry with the ActiveMountLoaders */
+enum class EInitialProviderRegisterReason
+{
+	/* Expected case, OnActiveMountsInitialize was called. */
+	ActiveMountsInitialize,
+	/** Allowed case, AppendState was called before OnActiveMountsInitialize. */
+	AppendState,
+	/** Error case, some other data was added to global AssetRegistry before the premade. */
+	EarlierState,
+};
+
+/**
+ * Holds callbacks and events triggered during ActiveMountLoading functions that need to be broadcast outside of the
+ * AssetRegistry's locks.
+ */
+class FActiveMountEvents
+{
+public:
+	/**
+	 * Add a completion callback from ActiveMountsLoadAsync or ActiveMountsUnloadAsync, and the value of its
+	 * bCanceled parameter.
+	 */
+	void AddCompletion(TFunction<void(bool bCanceled)> InOnComplete, bool bInCanceled);
+	/** Send all of the added events. Must be called outside of the AssetRegistry's locks. */
+	void Broadcast();
+
+private:
+	struct FCompletionData
+	{
+		TFunction<void(bool bCanceled)> OnComplete;
+		bool bCanceled = false;
+	};
+	TArray<FCompletionData> CompletionEvents;
+};
+
+/**
+ * The interface into the ActiveMountsLoader used by the rest of the AssetRegistry. The ActiveMountsLoader stores
+ * filepaths for AssetRegistry.bin files that provide assets to the AssetRegistry, and it manages a state machine to
+ * unload and reload the AssetDatas for mountpoints (such as plugin Content directories) when those mountpoints are
+ * loaded and unloaded. @see the comment on ActiveMountUnloading in IAssetRegistry.h.
+ *
+ * The ActiveMountsLoader has some minor performance costs and so should only be enabled in configurations and projects
+ * that need the memory savings and that can behave correctly when AssetDatas are unloaded for inactive plugins.
+ * 
+ * This interface object exists in all configurations, but it only creates the full ActiveMountsLoader when enabled.
+ * When not enabled it implements only the side effects that are required by the rest of the AssetRegistry, such
+ * as calling AppendState for filenames passed to TryRegisterAndAppendState. This wrapper also manages the separate
+ * critical section(s) used by the ActiveMountsLoader; ActiveMountsLoader functions are always called inside of its
+ * lock.
+ */
+class FActiveMountsLoaderWrapper
+{
+public:
+	/**
+	 * Called during AssetRegistry initialization. Reads configuration settings such as enable flags and creates the
+	 * full FActiveMountsLoader if enabled.
+	 */
+	void InitializeConfig();
+	/**
+	 * Called during AssetRegistry shutdown. Signals any active ActiveMountsLoader threads to terminate, waits for
+	 * them, and shuts down the ActiveMountsLoader.
+	 */
+	void Shutdown(UAssetRegistryImpl& UARI, FActiveMountEvents& EventContext);
+	/** Report whether the full ActiveMountsLoader is currently enabled. */
+	bool IsEnabled();
+
+	/** @see IAssetRegistry::ActiveMountsLoadAsync. Calls OnComplete(bCanceled=true) immediately if disabled. */
+	void LoadAsync(UAssetRegistryImpl& UARI, FActiveMountEvents& EventContext,
+		TConstArrayView<TRefCountPtr<UE::PackageName::IMountPoint>> InMountPoints,
+		TFunction<void(bool bCanceled)>&& OnComplete);
+	/** @see IAssetRegistry::ActiveMountsUnloadAsync. Calls OnComplete(bCanceled=true) immediately if disabled. */
+	void UnloadAsync(UAssetRegistryImpl& UARI, FActiveMountEvents& EventContext,
+		TConstArrayView<TRefCountPtr<UE::PackageName::IMountPoint>> InMountPoints,
+		TFunction<void(bool bCanceled)>&& OnComplete);
+	/**
+	 * Checks the PackageName mounted status of all MountPoints with a registered AssetProvider, and unloads the
+	 * AssetDatas of all of the currently unmounted ones from the AssetRegistry.
+	 */
+	void UnloadAllUnmountedAsync(UAssetRegistryImpl& UARI, FActiveMountEvents& EventContext,
+		TFunction<void(bool bCanceled)>&& OnComplete);
+	/** @see IAssetRegistry::ActiveMountIsLoaded. */
+	bool IsLoaded(UAssetRegistryImpl& UARI, FActiveMountEvents& EventContext,
+		const TRefCountPtr<UE::PackageName::IMountPoint>& MountPoint);
+
+	/** @see IAssetRegistry::ActiveMountsInitialize. */
+	void OnActiveMountsInitialize(UAssetRegistryImpl& UARI, FActiveMountEvents& EventContext,
+		FInterfaceWriteScopeLock& ProofOfWriteLock);
+	/** @see IAssetRegistry::ActiveMountsRegisterAndLoadFilePathSynchronous. */
+	bool TryRegisterAssetProviderAndAppendStateOutsideARLock(UAssetRegistryImpl& UARI, FEventContext& AREventContext,
+		FActiveMountEvents& EventContext, FStringView InStateFilePath);
+	/**
+	 * Called by internal AssetRegistry code to load and register the AssetRegistry files for built-in plugins, called
+	 * immediately after the initial premade AssetRegistry is loaded. If successful, this function also triggers
+	 * the immediate registration of the premade AssetRegistry.
+	 */
+	bool TryLoadStateAndRegisterAssetProvider(UAssetRegistryImpl& UARI,	FActiveMountEvents& EventContext,
+		FInterfaceWriteScopeLock& ProofOfWriteLock, FPakPlatformFile* PakPlatformFile, FStringView InStateFilePath,
+		FAssetRegistryState& OutLoadedState);
+	/**
+	 * Deactivate a file as an AssetProvider, clearing the data about it from memory. AssetDatas in MountPoints
+	 * provided by the AssetProvider can no longer be unloaded or reloaded (until the AssetProvider is registered
+	 * again). Invalid to call on the premade initial provider.
+	 */
+	void UnregisterAssetProvider(UAssetRegistryImpl& UARI, FActiveMountEvents& EventContext,
+		FStringView InStateFilePath);
+	/**
+	 * Called when the premade AssetRegistryState has finished loading and is added to the AssetRegistry. We do not
+	 * yet try to identify the MountPoints used by the AssetDatas within the AssetRegistryState, because identifying
+	 * them might rely on the registration of MountPoints with FPackageName and some projects might not be immediately
+	 * ready for that. This function just stores the filename and marks that the AssetProvider exists.
+	 */
+	void InitialProviderReport(UAssetRegistryImpl& UARI, FActiveMountEvents& EventContext,
+		FAssetProviderInitData& StateInitData);
+	/**
+	 * Called when ActiveMountsInitialize or some other event signals that we are ready to identify MountPoints within
+	 * the initial provider. Just sets a flag, takes no other action.
+	 */
+	void InitialProviderOnReadyForRegister(UAssetRegistryImpl& UARI, FActiveMountEvents& EventContext,
+		EInitialProviderRegisterReason Reason);
+	/**
+	 * Checks whether InitialProviderReport and InitialProviderOnReadyForRegister have both been called and if so
+	 * parses the initial provider's AssetRegistry state to identify its MountPoints.
+	 */
+	void InitialProviderConditionalRegister(UAssetRegistryImpl& UARI, FActiveMountEvents& EventContext,
+		FAssetRegistryState& InitialProviderState);
+
+	/**
+	 * Called when FPackageName reports one or more OnMountPointMounted; triggers the automatic async batched reloading
+	 * of AssetDatas for the MountPoint, if autounload is enabled.
+	 */
+	void OnMountPointsMounted(UAssetRegistryImpl& UARI, FActiveMountEvents& EventContext,
+		TConstArrayView<TRefCountPtr<UE::PackageName::IMountPoint>> InMountPoints);
+	/**
+	 * Called when FPackageName reports one or more OnMountPointDismounted; triggers the automatic async batched
+	 * unloading of AssetDatas for the MountPoint, if autounload is enabled.
+	 */
+	void OnMountPointsDismounted(UAssetRegistryImpl& UARI, FActiveMountEvents& EventContext,
+		FInterfaceWriteScopeLock& ProofOfWriteLock,
+		TConstArrayView<TRefCountPtr<UE::PackageName::IMountPoint>> InMountPoints);
+
+	/**
+	 * Report whether the configuration settings specify that the full ActiveMountsLoader should be enabled. Should only
+	 * be used during InitializeConfig or by functions that need to execute before it, such as
+	 * UE::AssetRegistry::Premade::FPreloader.
+	 */
+	static bool IsEnabledByConfig();
+private:
+	/**
+	 * Called by all public functions, even when the full ActiveMountLoader is enabled. When disabled, operates as
+	 * cheaply as possible, does not lock, and returns nullptr. When enabled, locks and returns a pointer to
+	 * this->Inner.
+	 */
+	FActiveMountsLoader* LockIfEnabled();
+	/** Lock ActiveMountsLock, the lock that must be entered around any full ActiveMountLoader functions. */
+	void Lock();
+	/** Unlock ActiveMountsLock. */
+	void Unlock();
+	/**
+	 * Return the lock that must be entered around AssetDataProvider register/deregister functions. Returns
+	 * nullptr if full ActiveMountLoader is not enabled.
+	 */
+	TSharedPtr<UE::FRecursiveMutex> GetRegistrationLock();
+
+	/**
+	 * Lock that guards the ActiveMountLoader's data, separate from the AssetRegistry's data so that we do not block
+	 * AssetRegistry queries when doing some long operations inside the lock.
+	 * This lock is allowed to be held at the same time as the AssetRegistry's interface lock (avoid when possible).
+	 * When both are held, the AssetRegistry's interface lock must be entered first.
+	 */
+	UE::FRecursiveMutex ActiveMountsLock;
+	/**
+	 * Null if ActiveMounts are disabled. This is an atomic only because it can be set to null during Shutdown and
+	 * the reader might be checking it for null outside the lock when that happens. The reader must always enter the
+	 * lock and check for null again before dereferencing it.
+	 */
+	std::atomic<FActiveMountsLoader*> Inner{ nullptr };
+
+	friend FActiveMountsLoader;
+	friend FAssetProvider;
+	friend FMountPointData;
+};
+
+} // namespace UE::AssetRegistry::Impl
diff --git a/Engine/Source/Runtime/AssetRegistry/Private/ActiveMountsLoader.cpp b/Engine/Source/Runtime/AssetRegistry/Private/ActiveMountsLoader.cpp
new file mode 100644
--- /dev/null
+++ b/Engine/Source/Runtime/AssetRegistry/Private/ActiveMountsLoader.cpp
@@ -0,0 +1,2067 @@
+// Copyright Epic Games, Inc. All Rights Reserved.
+
+#include "ActiveMountsLoader.h"
+
+#include "Algo/Count.h"
+#include "AssetRegistry.h"
+#include "AssetRegistryPrivate.h"
+#include "Async/Async.h"
+#include "Misc/ConfigCacheIni.h"
+#include "Containers/DirectoryTree.h"
+#include "HAL/FileManager.h"
+#include "HAL/PlatformProperties.h"
+#include "Interfaces/IPluginManager.h"
+#include "IPlatformFilePak.h"
+#include "Logging/LogMacros.h"
+#include "Logging/StructuredLog.h"
+#include "Misc/AssertionMacros.h"
+#include "Misc/CommandLine.h"
+#include "Misc/CoreMisc.h"
+#include "Misc/CoreMiscDefines.h"
+#include "Misc/FileHelper.h"
+#include "Misc/Parse.h"
+#include "Misc/ScopeExit.h"
+#include "Misc/ScopeLock.h"
+#include "Misc/StringBuilder.h"
+#include "Stats/Stats.h"
+
+namespace UE::AssetRegistry::Impl
+{
+
+template <typename ContainerType>
+void ConditionalShrinkUsuallyEmptyContainer(ContainerType& Container)
+{
+	if (Container.Num() == 0 ||
+		(Container.Max() > 8 && Container.Max() > 2 * Container.Num()))
+	{
+		Container.Shrink();
+	}
+}
+
+void FActiveMountEvents::AddCompletion(TFunction<void(bool bCanceled)> InOnComplete, bool bInCanceled)
+{
+	if (!InOnComplete)
+	{
+		return;
+	}
+	CompletionEvents.Add(FCompletionData{ MoveTemp(InOnComplete), bInCanceled });
+}
+
+void FActiveMountEvents::Broadcast()
+{
+	for (FCompletionData& Data : CompletionEvents)
+	{
+		check(Data.OnComplete);
+		Data.OnComplete(Data.bCanceled);
+	}
+	CompletionEvents.Empty();
+}
+
+bool FActiveMountsLoaderWrapper::IsEnabled()
+{
+	return Inner.load(std::memory_order_relaxed) != nullptr;
+}
+
+bool FActiveMountsLoaderWrapper::IsEnabledByConfig()
+{
+	static bool bEnabled = []()
+		{
+			bool bLocalEnabled = false;
+			if (FPlatformProperties::RequiresCookedData() && IsRunningGame() && !IsRunningDedicatedServer())
+			{
+				if (!FParse::Bool(FCommandLine::Get(), TEXT("-ActiveMountLoading="), bLocalEnabled))
+				{
+					// We need the value from Config to calculate the proper value, so assert that GConfig is
+					// available by the time we first call IsEnabledByConfig.
+					check(GConfig);
+					GConfig->GetBool(TEXT("AssetRegistry"), TEXT("bActiveMountLoadingEnabled_Client"),
+						bLocalEnabled, GEngineIni);
+				}
+			}
+			return bLocalEnabled;
+		}();
+	return bEnabled;
+}
+
+void FActiveMountsLoaderWrapper::InitializeConfig()
+{
+	if (!IsEnabledByConfig())
+	{
+		return;
+	}
+
+	bool bAutoUnload = false;
+	if (!FParse::Bool(FCommandLine::Get(), TEXT("-ActiveMountLoadingAuto="), bAutoUnload))
+	{
+		check(GConfig); // It was already required by IsEnabledByConfig
+		GConfig->GetBool(TEXT("AssetRegistry"), TEXT("bActiveMountLoadingAuto_Client"),
+			bAutoUnload, GEngineIni);
+	}
+
+	Lock();
+	ON_SCOPE_EXIT
+	{
+		Unlock();
+	};
+	// Races during initialization are unexpected but are also not a problem.
+	if (!Inner.load(std::memory_order_relaxed))
+	{
+		FActiveMountsLoader* LocalInner = new FActiveMountsLoader();
+		LocalInner->RegistrationLock = MakeShared<UE::FRecursiveMutex>();
+		LocalInner->bAutoUnload = bAutoUnload;
+		Inner.store(LocalInner, std::memory_order_relaxed);
+	}
+}
+
+void FActiveMountsLoaderWrapper::Shutdown(UAssetRegistryImpl& UARI, FActiveMountEvents& EventContext)
+{
+	TArray<UE::Tasks::FTask> WaitOnAsync;
+	if (FActiveMountsLoader* LocalInner = LockIfEnabled(); LocalInner)
+	{
+		ON_SCOPE_EXIT
+		{
+			Unlock();
+		};
+		LocalInner->StartShutdown(UARI, EventContext, WaitOnAsync);
+	}
+	else
+	{
+		return;
+	}
+
+	// Waiting has to be done outside the lock because the task threads have to enter the lock
+	// before they can exit.
+	UE::Tasks::Wait(WaitOnAsync);
+
+	if (FActiveMountsLoader* LocalInner = LockIfEnabled(); LocalInner)
+	{
+		ON_SCOPE_EXIT
+		{
+			Unlock();
+		};
+
+		LocalInner->FinishShutdown(UARI, EventContext);
+		delete LocalInner;
+		Inner.store(nullptr, std::memory_order_relaxed);
+	}
+}
+
+FActiveMountsLoader* FActiveMountsLoaderWrapper::LockIfEnabled()
+{
+	if (!Inner.load(std::memory_order_relaxed))
+	{
+		return nullptr;
+	}
+	ActiveMountsLock.Lock();
+	FActiveMountsLoader* LocalInner = Inner.load(std::memory_order_relaxed);
+	if (!LocalInner)
+	{
+		ActiveMountsLock.Unlock();
+		return nullptr;
+	}
+	// Caller is responsible for calling Unlock.
+	return LocalInner;
+}
+
+void FActiveMountsLoaderWrapper::Lock()
+{
+	ActiveMountsLock.Lock();
+}
+
+void FActiveMountsLoaderWrapper::Unlock()
+{
+	ActiveMountsLock.Unlock();
+}
+
+TSharedPtr<UE::FRecursiveMutex> FActiveMountsLoaderWrapper::GetRegistrationLock()
+{
+	if (FActiveMountsLoader* LocalInner = LockIfEnabled(); LocalInner)
+	{
+		ON_SCOPE_EXIT
+		{
+			Unlock();
+		};
+		if (!LocalInner->bShutdown)
+		{
+			return LocalInner->RegistrationLock;
+		}
+	}
+	return TSharedPtr<UE::FRecursiveMutex>();
+}
+
+void FActiveMountsLoaderWrapper::LoadAsync(UAssetRegistryImpl& UARI, FActiveMountEvents& EventContext,
+	TConstArrayView<TRefCountPtr<UE::PackageName::IMountPoint>> InMountPoints,
+	TFunction<void(bool bCanceled)>&& OnComplete)
+{
+	if (FActiveMountsLoader* LocalInner = LockIfEnabled(); LocalInner)
+	{
+		ON_SCOPE_EXIT
+		{
+			Unlock();
+		};
+		if (!LocalInner->bShutdown)
+		{
+			LocalInner->LoadAsync(UARI, EventContext, InMountPoints, MoveTemp(OnComplete));
+			return;
+		}
+	}
+	EventContext.AddCompletion(MoveTemp(OnComplete), true /* bCanceled */);
+}
+
+void FActiveMountsLoaderWrapper::UnloadAsync(UAssetRegistryImpl& UARI, FActiveMountEvents& EventContext,
+	TConstArrayView<TRefCountPtr<UE::PackageName::IMountPoint>> InMountPoints,
+	TFunction<void(bool bCanceled)>&& OnComplete)
+{
+	if (FActiveMountsLoader* LocalInner = LockIfEnabled(); LocalInner)
+	{
+		ON_SCOPE_EXIT
+		{
+			Unlock();
+		};
+		if (!LocalInner->bShutdown)
+		{
+			LocalInner->UnloadAsync(UARI, EventContext, InMountPoints, MoveTemp(OnComplete));
+			return;
+		}
+	}
+	EventContext.AddCompletion(MoveTemp(OnComplete), true /* bCanceled */);
+}
+
+void FActiveMountsLoaderWrapper::UnloadAllUnmountedAsync(UAssetRegistryImpl& UARI, FActiveMountEvents& EventContext,
+	TFunction<void(bool bCanceled)>&& OnComplete)
+{
+	if (FActiveMountsLoader* LocalInner = LockIfEnabled(); LocalInner)
+	{
+		ON_SCOPE_EXIT
+		{
+			Unlock();
+		};
+		if (!LocalInner->bShutdown)
+		{
+			LocalInner->UnloadAllUnmountedAsync(UARI, EventContext, MoveTemp(OnComplete));
+			return;
+		}
+	}
+	EventContext.AddCompletion(MoveTemp(OnComplete), true /* bCanceled */);
+}
+
+bool FActiveMountsLoaderWrapper::IsLoaded(UAssetRegistryImpl& UARI, FActiveMountEvents& EventContext,
+	const TRefCountPtr<UE::PackageName::IMountPoint>& MountPoint)
+{
+	if (FActiveMountsLoader* LocalInner = LockIfEnabled(); LocalInner)
+	{
+		ON_SCOPE_EXIT
+		{
+			Unlock();
+		};
+		if (!LocalInner->bShutdown)
+		{
+			return LocalInner->IsLoaded(UARI, EventContext, MountPoint);
+		}
+	}
+	return false;
+}
+
+void FActiveMountsLoaderWrapper::OnActiveMountsInitialize(UAssetRegistryImpl& UARI, FActiveMountEvents& EventContext,
+	FInterfaceWriteScopeLock& ProofOfWriteLock)
+{
+	if (FActiveMountsLoader* LocalInner = LockIfEnabled(); LocalInner)
+	{
+		ON_SCOPE_EXIT
+		{
+			Unlock();
+		};
+		if (!LocalInner->bShutdown)
+		{
+			LocalInner->OnActiveMountsInitialize(UARI, EventContext, ProofOfWriteLock);
+		}
+	}
+}
+
+bool FActiveMountsLoaderWrapper::TryRegisterAssetProviderAndAppendStateOutsideARLock(UAssetRegistryImpl& UARI,
+	FEventContext& AREventContext, FActiveMountEvents& EventContext, FStringView InStateFilePath)
+{
+	FString StateFilePath = UE::AssetRegistry::CreateStandardFilename(InStateFilePath);
+	TSharedPtr<UE::FRecursiveMutex> RegistrationLock = GetRegistrationLock();
+	if (RegistrationLock)
+	{
+		UE::TUniqueLock RegistrationScopeLock(*RegistrationLock);
+		if (FActiveMountsLoader* LocalInner = LockIfEnabled(); LocalInner)
+		{
+			ON_SCOPE_EXIT
+			{
+				Unlock();
+			};
+			if (!LocalInner->bShutdown)
+			{
+				return LocalInner->TryRegisterAssetProviderAndAppendStateOutsideARLock(UARI, AREventContext,
+					EventContext, StateFilePath);
+			}
+		}
+	}
+
+	TUniquePtr<FArchive> FileArchive = OpenReadStateFileArchive(StateFilePath);
+	if (!FileArchive)
+	{
+		return false;
+	}
+	FAssetRegistryState StateToAppend;
+	if (!StateToAppend.Load(*FileArchive, UE::AssetRegistry::Premade::GetLoadOptions()))
+	{
+		return false;
+	}
+	UARI.AppendState(StateToAppend);
+	return true;
+}
+
+bool FActiveMountsLoaderWrapper::TryLoadStateAndRegisterAssetProvider(UAssetRegistryImpl& UARI,
+	FActiveMountEvents& EventContext, FInterfaceWriteScopeLock& ProofOfWriteLock, FPakPlatformFile* PakPlatformFile,
+	FStringView InStateFilePath, FAssetRegistryState& OutLoadedState)
+{
+	FString StateFilePath = UE::AssetRegistry::CreateStandardFilename(InStateFilePath);
+	TSharedPtr<UE::FRecursiveMutex> RegistrationLock = GetRegistrationLock();
+	if (RegistrationLock)
+	{
+		UE::TUniqueLock RegistrationScopeLock(*RegistrationLock);
+		if (FActiveMountsLoader* LocalInner = LockIfEnabled(); LocalInner)
+		{
+			ON_SCOPE_EXIT
+			{
+				Unlock();
+			};
+			if (!LocalInner->bShutdown)
+			{
+				return LocalInner->TryLoadStateAndRegisterAssetProvider(UARI, EventContext, ProofOfWriteLock,
+					PakPlatformFile, StateFilePath, OutLoadedState);
+			}
+		}
+	}
+
+	// Optimization: if we're using pak files then only search paks (avoid unnecessary fallback to loose)
+	bool bFileExists = PakPlatformFile != nullptr
+		? PakPlatformFile->FindFileInPakFiles(*StateFilePath)
+		: IFileManager::Get().FileExists(*StateFilePath);
+	if (!bFileExists)
+	{
+		return false;
+	}
+	TUniquePtr<FArchive> FileArchive = OpenReadStateFileArchive(StateFilePath);
+	if (!FileArchive)
+	{
+		return false;
+	}
+	return OutLoadedState.Load(*FileArchive, UE::AssetRegistry::Premade::GetLoadOptions());
+}
+
+void FActiveMountsLoaderWrapper::UnregisterAssetProvider(UAssetRegistryImpl& UARI, FActiveMountEvents& EventContext,
+	FStringView InStateFilePath)
+{
+	FString StateFilePath = UE::AssetRegistry::CreateStandardFilename(InStateFilePath);
+	TSharedPtr<UE::FRecursiveMutex> RegistrationLock = GetRegistrationLock();
+	if (RegistrationLock)
+	{
+		UE::TUniqueLock RegistrationScopeLock(*RegistrationLock);
+		TArray<UE::Tasks::FTask> WaitOnAsync;
+		if (FActiveMountsLoader* LocalInner = LockIfEnabled(); LocalInner)
+		{
+			ON_SCOPE_EXIT
+			{
+				Unlock();
+			};
+			if (!LocalInner->bShutdown)
+			{
+				if (!LocalInner->TryStartUnregisterAssetProvider(UARI, EventContext, WaitOnAsync, StateFilePath))
+				{
+					return;
+				}
+			}
+		}
+
+		// Waiting has to be done outside the normal lock because the task threads have to enter the lock
+		// before they can exit. It does not have to be done outside the RegistrationLock, because the
+		// task threads do not use that lock.
+		UE::Tasks::Wait(WaitOnAsync);
+
+		if (FActiveMountsLoader* LocalInner = LockIfEnabled(); LocalInner)
+		{
+			ON_SCOPE_EXIT
+			{
+				Unlock();
+			};
+
+			if (!LocalInner->bShutdown)
+			{
+				LocalInner->FinishUnregisterAssetProvider(UARI, EventContext, StateFilePath);
+			}
+		}
+	}
+}
+
+void FActiveMountsLoaderWrapper::InitialProviderReport(UAssetRegistryImpl& UARI, FActiveMountEvents& EventContext,
+	FAssetProviderInitData& StateInitData)
+{
+	if (FActiveMountsLoader* LocalInner = LockIfEnabled(); LocalInner)
+	{
+		ON_SCOPE_EXIT
+		{
+			Unlock();
+		};
+		if (!LocalInner->bShutdown)
+		{
+			LocalInner->InitialProviderReport(UARI, EventContext, StateInitData);
+		}
+	}
+}
+
+void FActiveMountsLoaderWrapper::InitialProviderOnReadyForRegister(UAssetRegistryImpl& UARI,
+	FActiveMountEvents& EventContext, EInitialProviderRegisterReason Reason)
+{
+	if (FActiveMountsLoader* LocalInner = LockIfEnabled(); LocalInner)
+	{
+		ON_SCOPE_EXIT
+		{
+			Unlock();
+		};
+		if (!LocalInner->bShutdown)
+		{
+			LocalInner->InitialProviderOnReadyForRegister(UARI, EventContext, Reason);
+		}
+	}
+}
+
+void FActiveMountsLoaderWrapper::InitialProviderConditionalRegister(UAssetRegistryImpl& UARI,
+	FActiveMountEvents& EventContext, FAssetRegistryState& InitialProviderState)
+{
+	if (FActiveMountsLoader* LocalInner = LockIfEnabled(); LocalInner)
+	{
+		ON_SCOPE_EXIT
+		{
+			Unlock();
+		};
+		if (!LocalInner->bShutdown)
+		{
+			LocalInner->InitialProviderConditionalRegister(UARI, EventContext, InitialProviderState);
+		}
+	}
+}
+
+void FActiveMountsLoaderWrapper::OnMountPointsMounted(UAssetRegistryImpl& UARI, FActiveMountEvents& EventContext,
+	TConstArrayView<TRefCountPtr<UE::PackageName::IMountPoint>> InMountPoints)
+{
+	if (FActiveMountsLoader* LocalInner = LockIfEnabled(); LocalInner)
+	{
+		ON_SCOPE_EXIT
+		{
+			Unlock();
+		};
+		if (!LocalInner->bShutdown)
+		{
+			LocalInner->OnMountPointsMounted(UARI, EventContext, InMountPoints);
+		}
+	}
+}
+
+void FActiveMountsLoaderWrapper::OnMountPointsDismounted(UAssetRegistryImpl& UARI, FActiveMountEvents& EventContext,
+	FInterfaceWriteScopeLock& ProofOfWriteLock,
+	TConstArrayView<TRefCountPtr<UE::PackageName::IMountPoint>> InMountPoints)
+{
+	if (FActiveMountsLoader* LocalInner = LockIfEnabled(); LocalInner)
+	{
+		ON_SCOPE_EXIT
+		{
+			Unlock();
+		};
+		if (!LocalInner->bShutdown)
+		{
+			LocalInner->OnMountPointsDismounted(UARI, EventContext, ProofOfWriteLock, InMountPoints);
+		}
+	}
+}
+
+FAssetProvider::FAssetProvider(FStringView InFilePath)
+	: FilePath(InFilePath)
+{
+}
+
+FStringView FAssetProvider::GetFilePath() const
+{
+	return FilePath;
+}
+
+bool FAssetProvider::IsInitialProvider() const
+{
+	return bInitialProvider;
+}
+
+void FAssetProvider::SetIsInitialProvider(bool bValue)
+{
+	bInitialProvider = bValue;
+}
+
+const FixedTagPrivate::FWeakStorePtr& FAssetProvider::GetStore() const
+{
+	return Store;
+}
+
+void FAssetProvider::SetStore(FixedTagPrivate::FWeakStorePtr InStore)
+{
+	Store = MoveTemp(InStore);
+}
+
+uint32& FAssetProvider::GetLoadRefCountPtr()
+{
+	return LoadRefCount;
+}
+
+UE::Tasks::FTask& FAssetProvider::GetLoadTask()
+{
+	return LoadTask;
+}
+
+bool FAssetProvider::IsShutdown() const
+{
+	return bShutdown;
+}
+
+void FAssetProvider::SetIsShutdown(bool bValue)
+{
+	bShutdown = bValue;
+}
+
+bool FAssetProvider::HasMissingMountPoints() const
+{
+	return bHasMissingMountPoints;
+}
+
+void FAssetProvider::SetHasMissingMountPoints(bool bValue)
+{
+	bHasMissingMountPoints = bValue;
+}
+
+bool FAssetProvider::IsErrorNoUnload() const
+{
+	return bErrorNoUnload;
+}
+
+void FAssetProvider::SetIsErrorNoUnload(bool bValue)
+{
+	bErrorNoUnload = bValue;
+}
+
+FAssetProviderSetKey::FAssetProviderSetKey(TUniquePtr<FAssetProvider>&& InAssetProvider)
+	: AssetProvider(MoveTemp(InAssetProvider))
+{
+}
+
+bool FAssetProviderSetKey::operator==(const FAssetProviderSetKey& Other) const
+{
+	return AssetProvider.Get() == Other.AssetProvider.Get();
+}
+
+bool FAssetProviderSetKey::operator==(FStringView OtherProviderFilePath) const
+{
+	return AssetProvider->GetFilePath().Equals(OtherProviderFilePath, ESearchCase::IgnoreCase);
+}
+
+FAssetProvider& FAssetProviderSetKey::GetAssetProvider() const
+{
+	return *AssetProvider;
+}
+
+uint32 GetTypeHash(const FAssetProviderSetKey& Key)
+{
+	return GetTypeHash(Key.GetAssetProvider().GetFilePath());
+}
+
+FMountPointData::FMountPointData(TRefCountPtr<UE::PackageName::IMountPoint> InMountPoint, EMountState InInitialState)
+	: MountPoint(MoveTemp(InMountPoint))
+{
+	SetCurrentState(InInitialState);
+	SetTargetState(InInitialState);
+}
+
+bool FMountPointData::operator==(const FMountPointData& Other) const
+{
+	return MountPoint.GetReference() == Other.MountPoint.GetReference();
+}
+
+bool FMountPointData::operator==(UE::PackageName::IMountPoint* InMountPoint) const
+{
+	return MountPoint.GetReference() == InMountPoint;
+}
+
+uint32 GetTypeHash(const FMountPointData& Key)
+{
+	return ::GetTypeHash(Key.GetMountPoint().GetReference());
+}
+
+const TRefCountPtr<UE::PackageName::IMountPoint>& FMountPointData::GetMountPoint() const
+{
+	return MountPoint;
+}
+
+FAssetProvider* FMountPointData::GetAssetProvider() const
+{
+	return AssetProvider;
+}
+
+void FMountPointData::SetAssetProvider(FAssetProvider* InAssetProvider)
+{
+	AssetProvider = InAssetProvider;
+}
+
+EMountState FMountPointData::GetCurrentState() const
+{
+	return static_cast<EMountState>(CurrentStateBits);
+}
+
+void FMountPointData::SetCurrentState(EMountState InState)
+{
+	CurrentStateBits = static_cast<uint32>(InState);
+}
+
+EMountState FMountPointData::GetTargetState() const
+{
+	return static_cast<EMountState>(TargetStateBits);
+}
+
+void FMountPointData::SetTargetState(EMountState InState)
+{
+	TargetStateBits = static_cast<uint32>(InState);
+}
+
+bool FMountPointData::IsErrorNoUnload() const
+{
+	return ErrorNoUnloadBit != 0;
+}
+
+void FMountPointData::SetIsErrorNoUnload(bool bValue)
+{
+	ErrorNoUnloadBit = bValue ? 1 : 0;
+}
+
+bool FMountPointData::IsMultiSegmentName() const
+{
+	return MultiSegmentNameBit != 0;
+}
+
+void FMountPointData::SetIsMultiSegmentName(bool bValue)
+{
+	MultiSegmentNameBit = bValue ? 1 : 0;
+}
+
+bool FMountPointData::IsParentOrChild() const
+{
+	return IsParentOrChildBit != 0;
+}
+
+void FMountPointData::SetIsParentOrChild(bool bValue)
+{
+	IsParentOrChildBit = bValue ? 1 : 0;
+}
+
+bool FMountPointData::HasMultipleProviders() const
+{
+	return HasMultipleProvidersBit != 0;
+}
+
+void FMountPointData::SetHasMultipleProviders(bool bValue)
+{
+	HasMultipleProvidersBit = bValue ? 1 : 0;
+}
+
+void FActiveMountsLoader::StartShutdown(UAssetRegistryImpl& UARI, FActiveMountEvents& EventContext, TArray<UE::Tasks::FTask>& OutWaitOnAsync)
+{
+	// If bShutdown is already true, then another thread has already started shutdown. We can both race to see who
+	// gets to reenter the lock after waiting on async and call FinishShutdown. We still need to wait on this thread
+	// for the tasks to complete before this (or any) thread is allowed to call FinishShutdown.
+	if (!bShutdown)
+	{
+		// Set bShutdown=true to signal to all taskthreads that they should shutdown.
+		bShutdown = true;
+	}
+	// Tell our caller to wait for tasks to complete (outside our lock) before we call FinishShutdown to delete data
+	// out from under them.
+	if (UnloadTask.IsValid())
+	{
+		OutWaitOnAsync.Add(UnloadTask);
+	}
+	for (FAssetProviderSetKey& SetKey : AssetProviders)
+	{
+		UE::Tasks::FTask& LoadTask = SetKey.GetAssetProvider().GetLoadTask();
+		if (LoadTask.IsValid())
+		{
+			OutWaitOnAsync.Add(LoadTask);
+		}
+	}
+}
+
+void FActiveMountsLoader::FinishShutdown(UAssetRegistryImpl& UARI, FActiveMountEvents& EventContext)
+{
+	using namespace UE::PackageName;
+	using FSubscriberArray = TArray<TRefCountPtr<FMountLoadEventSubscriber>>;
+
+	// Caller should have cleared UnloadTask via StartShutdown and WaitOnAsync.
+	check(!UnloadTask.IsValid());
+	for (TPair<TRefCountPtr<IMountPoint>, FSubscriberArray>& LoadPair : LoadSubscribers)
+	{
+		for (TRefCountPtr<FMountLoadEventSubscriber>& Subscriber : LoadPair.Value)
+		{
+			EventContext.AddCompletion(MoveTemp(Subscriber->OnComplete), true /* bWasCanceled */);
+			Subscriber.SafeRelease();
+		}
+	}
+	LoadSubscribers.Empty();
+	for (TPair<TRefCountPtr<IMountPoint>, FSubscriberArray>& UnloadPair : UnloadSubscribers)
+	{
+		for (TRefCountPtr<FMountLoadEventSubscriber>& Subscriber : UnloadPair.Value)
+		{
+			EventContext.AddCompletion(MoveTemp(Subscriber->OnComplete), true /* bWasCanceled */);
+			Subscriber.SafeRelease();
+		}
+	}
+	UnloadSubscribers.Empty();
+
+	// MountPoints have a raw pointer to AssetProviders, AssetProviders do not have a pointer to MountPoints.
+	// MountPoints do not currently access their AssetProvider during their destructor, but in case that changes,
+	// clear them first.
+	MountPoints.Empty();
+#if DO_CHECK
+	for (FAssetProviderSetKey& SetKey : AssetProviders)
+	{
+		FAssetProvider& AssetProvider = SetKey.GetAssetProvider();
+		// Caller should have cleared LoadTask via StartShutdown and WaitOnAsync.
+		check(!AssetProvider.GetLoadTask().IsValid());
+	}
+#endif
+	AssetProviders.Empty(); // TUniquePtr within the FAssetProviderSetKey deletes the AssetProviders.
+}
+
+void FActiveMountsLoader::LoadAsync(UAssetRegistryImpl& UARI, FActiveMountEvents& EventContext,
+	TConstArrayView<TRefCountPtr<UE::PackageName::IMountPoint>> InMountPoints,
+	TFunction<void(bool bCanceled)>&& OnComplete)
+{
+	using namespace UE::PackageName;
+
+	if (!bLoaderIsActive)
+	{
+		// We assume that the initial provider has been registered for a few reasons, such as being able to assume
+		// that no MountPointData implies no AssetProvider has assets in that MountPoint.
+		// To support that assumption we have to early exit here. We know that we can report completion in this case
+		// because we do not unload anything until that point either.
+		EventContext.AddCompletion(MoveTemp(OnComplete), false /* bCanceled */);
+		return;
+	}
+
+	TArray<TRefCountPtr<IMountPoint>> WaitForCompletionMounts;
+	TArray<TRefCountPtr<IMountPoint>> TriggerMounts;
+	for (const TRefCountPtr<IMountPoint>& MountPoint : InMountPoints)
+	{
+		if (MountPoint == nullptr)
+		{
+			continue;
+		}
+
+		FMountPointData* MountPointData = FindMountPointData(MountPoint.GetReference());
+		if (!MountPointData)
+		{
+			continue;
+		}
+		FAssetProvider* AssetProvider = MountPointData->GetAssetProvider();
+		if (!AssetProvider || AssetProvider->IsShutdown())
+		{
+			continue;
+		}
+
+		if (MountPointData->GetCurrentState() != EMountState::Loaded)
+		{
+			if (MountPointData->GetTargetState() != EMountState::Loaded)
+			{
+				TriggerMounts.Add(MountPoint);
+			}
+			if (OnComplete)
+			{
+				WaitForCompletionMounts.Add(MountPoint);
+			}
+		}
+	}
+
+	if (OnComplete)
+	{
+		// Create the subscriber before triggering any loads, because we may find that we can
+		// complete a Load immediately and we will need to decrement the subscriber's RefCount,
+		// so we need to have set it up already.
+		if (WaitForCompletionMounts.IsEmpty())
+		{
+			// We only added to TriggerMounts in the case where we also added to WaitForCompletionMounts, so assert empty.
+			check(TriggerMounts.IsEmpty());
+			EventContext.AddCompletion(MoveTemp(OnComplete), false /* bCanceled */);
+		}
+		else
+		{
+			TRefCountPtr<FMountLoadEventSubscriber> Subscriber = MakeRefCount<FMountLoadEventSubscriber>();
+			Subscriber->OnComplete = MoveTemp(OnComplete);
+			Subscriber->RemainingCount = WaitForCompletionMounts.Num();
+			Subscriber->bIsLoadSubscriber = true;
+			for (TRefCountPtr<IMountPoint>& WaitMount : WaitForCompletionMounts)
+			{
+				LoadSubscribers.FindOrAdd(MoveTemp(WaitMount)).Add(Subscriber);
+			}
+			WaitForCompletionMounts.Empty();
+		}
+	}
+
+	for (const TRefCountPtr<IMountPoint>& TriggerMount : TriggerMounts)
+	{
+		FMountPointData* TriggerData = FindMountPointData(TriggerMount.GetReference());
+		check(TriggerData); // We only added it to TriggerMounts if it existed
+		// We only added the mount to TriggerMounts if GetTargetState() != Loaded, and no other mount states
+		// (in particular, no transition mount states like Loading) are allowed in GetTargetState fields,
+		// so we know the target state must be Unloaded.
+		check(TriggerData->GetTargetState() == EMountState::Unloaded);
+		switch (TriggerData->GetCurrentState())
+		{
+		case EMountState::Unloaded:
+			TriggerData->SetTargetState(EMountState::Loaded);
+			TriggerData->SetCurrentState(EMountState::Loading);
+			check(TriggerData->GetAssetProvider()); // We only added it to TriggerMounts if it had an AssetProvider.
+			AddRefMountPointLoad(UARI, TriggerData);
+			break;
+		case EMountState::Loading:
+			// We previously received a load request that we're still working on, and after that but before now we
+			// received an unload request that superseded it, but that had to wait until the load finishes before
+			// it could operate. And now we have received a load request that supersedes the superseder.
+			// Mark the unload as completed but canceled, and set targetstate back to loaded.
+			TriggerData->SetTargetState(EMountState::Loaded);
+			NotifyMountPointUnloadComplete(EventContext, TriggerMount, true /* bCanceled */);
+			break;
+		case EMountState::Loaded:
+			check(false); // We only added it to TriggerMounts if it was not in Loaded
+			break;
+		case EMountState::Unloading:
+			// We previously received an unload request that we are still working on, but now we have a load
+			// request that supersedes it, but we have to wait for the unload to finish before we can start
+			// working on the load. Mark the unload as completed but canceled, and set targetstate to loaded.
+			TriggerData->SetTargetState(EMountState::Loaded);
+			NotifyMountPointUnloadComplete(EventContext, TriggerMount, true /* bCanceled */);
+			break;
+		default:
+			checkNoEntry();
+			break;
+		}
+	}
+}
+
+void FActiveMountsLoader::UnloadAsync(UAssetRegistryImpl& UARI, FActiveMountEvents& EventContext,
+	TConstArrayView<TRefCountPtr<UE::PackageName::IMountPoint>> InMountPoints,
+	TFunction<void(bool bCanceled)>&& OnComplete)
+{
+	using namespace UE::PackageName;
+
+	if (!bLoaderIsActive)
+	{
+		ensureMsgf(false, TEXT("UnloadAsync called before ActiveMountsInitialized. The unload call will be ignored."));
+		EventContext.AddCompletion(MoveTemp(OnComplete), true /* bCanceled */);
+		return;
+	}
+
+	TArray<TRefCountPtr<IMountPoint>> WaitForCompletionMounts;
+	TArray<TRefCountPtr<IMountPoint>> TriggerMounts;
+	for (const TRefCountPtr<IMountPoint>& MountPoint : InMountPoints)
+	{
+		if (MountPoint == nullptr)
+		{
+			continue;
+		}
+
+		FMountPointData* MountPointData = FindMountPointData(MountPoint.GetReference());
+		if (!MountPointData)
+		{
+			continue;
+		}
+		FAssetProvider* AssetProvider = MountPointData->GetAssetProvider();
+		if (!AssetProvider || AssetProvider->IsShutdown())
+		{
+			continue;
+		}
+		if (MountPointData->IsErrorNoUnload() || AssetProvider->IsErrorNoUnload())
+		{
+			continue;
+		}
+
+		if (MountPointData->GetCurrentState() != EMountState::Unloaded)
+		{
+			if (MountPointData->GetTargetState() != EMountState::Unloaded)
+			{
+				TriggerMounts.Add(MountPoint);
+			}
+			if (OnComplete)
+			{
+				WaitForCompletionMounts.Add(MountPoint);
+			}
+		}
+	}
+
+	if (OnComplete)
+	{
+		// Create the subscriber before triggering any loads, because we may find that we can
+		// complete an unload immediately and we will need to decrement the subscriber's RefCount,
+		// so we need to have set it up already.
+		if (WaitForCompletionMounts.IsEmpty())
+		{
+			// We only added to TriggerMounts in the case where we also added to WaitForCompletionMounts, so assert empty.
+			check(TriggerMounts.IsEmpty());
+			EventContext.AddCompletion(MoveTemp(OnComplete), false /* bCanceled */);
+		}
+		else
+		{
+			TRefCountPtr<FMountLoadEventSubscriber> Subscriber = MakeRefCount<FMountLoadEventSubscriber>();
+			Subscriber->OnComplete = MoveTemp(OnComplete);
+			Subscriber->RemainingCount = WaitForCompletionMounts.Num();
+			Subscriber->bIsLoadSubscriber = false;
+			for (TRefCountPtr<IMountPoint>& WaitMount : WaitForCompletionMounts)
+			{
+				UnloadSubscribers.FindOrAdd(MoveTemp(WaitMount)).Add(Subscriber);
+			}
+			WaitForCompletionMounts.Empty();
+		}
+	}
+
+	for (const TRefCountPtr<IMountPoint>& TriggerMount : TriggerMounts)
+	{
+		FMountPointData* TriggerData = FindMountPointData(TriggerMount.GetReference());
+		check(TriggerData); // We only added it to TriggerMounts if it existed
+		// We only added the mount to TriggerMounts if GetTargetState() != Unloaded, and no other mount states
+		// (in particular, no transition mount states like Loading) are allowed in GetTargetState fields,
+		// so we know the target state must be Loaded.
+		check(TriggerData->GetTargetState() == EMountState::Loaded);
+		switch (TriggerData->GetCurrentState())
+		{
+		case EMountState::Unloaded:
+			check(false); // We only added it to TriggerMounts if it was not in Unloaded
+			break;
+		case EMountState::Loading:
+			// We previously received a load request that we are still working on, but now we have an unload
+			// request that supersedes it, but we have to wait for the load to finish before we can start working
+			// on the unload. Mark the load as completed but canceled, and set targetstate to unloaded.
+			TriggerData->SetTargetState(EMountState::Unloaded);
+			NotifyMountPointLoadComplete(EventContext, TriggerMount, true /* bCanceled */);
+			break;
+		case EMountState::Loaded:
+			TriggerData->SetTargetState(EMountState::Unloaded);
+			TriggerData->SetCurrentState(EMountState::Unloading);
+			check(TriggerData->GetAssetProvider()); // We only added it to TriggerMounts if it had an AssetProvider.
+			AddRefMountPointUnload(UARI, TriggerData);
+			break;
+		case EMountState::Unloading:
+			// We previously received an unload request that we're still working on, and after that but before now
+			// we received a load request that superseded it, but that had to wait until the unload finishes before
+			// it could operate. And now we have received an unload request that supersedes the superseder.
+			// Mark the load as completed but canceled, and set targetstate back to unloaded.
+			TriggerData->SetTargetState(EMountState::Unloaded);
+			NotifyMountPointLoadComplete(EventContext, TriggerMount, true /* bCanceled */);
+			break;
+		default:
+			checkNoEntry();
+			break;
+		}
+	}
+}
+
+void FActiveMountsLoader::UnloadAllUnmountedAsync(UAssetRegistryImpl& UARI, FActiveMountEvents& EventContext,
+	TFunction<void(bool bCanceled)>&& OnComplete)
+{
+	if (!bLoaderIsActive)
+	{
+		EventContext.AddCompletion(MoveTemp(OnComplete), true /* bCanceled */);
+		return;
+	}
+
+	TArray<TRefCountPtr<UE::PackageName::IMountPoint>> MountPointsToUnload;
+	for (FMountPointData& MountData : MountPoints)
+	{
+		FAssetProvider* AssetProvider = MountData.GetAssetProvider();
+		if (!AssetProvider || AssetProvider->IsShutdown())
+		{
+			continue;
+		}
+		if (MountData.IsErrorNoUnload() || AssetProvider->IsErrorNoUnload())
+		{
+			continue;
+		}
+		const TRefCountPtr<UE::PackageName::IMountPoint>& MountPoint = MountData.GetMountPoint();
+
+		// Unload from AssetRegistry the MountPoints that PackageName reports are not currently mounted.
+		if (!MountPoint->IsMounted())
+		{
+			MountPointsToUnload.Add(MountPoint);
+		}
+	}
+	UnloadAsync(UARI, EventContext, MountPointsToUnload, MoveTemp(OnComplete));
+}
+
+bool FActiveMountsLoader::IsLoaded(UAssetRegistryImpl& UARI, FActiveMountEvents& EventContext,
+	const TRefCountPtr<UE::PackageName::IMountPoint>& MountPoint)
+{
+	FMountPointData* MountData = MountPoints.FindByHash(
+		GetTypeHash(MountPoint.GetReference()), MountPoint.GetReference());
+	if (!MountData)
+	{
+		return false;
+	}
+	return MountData->GetCurrentState() == EMountState::Loaded;
+}
+
+void FActiveMountsLoader::OnActiveMountsInitialize(UAssetRegistryImpl& UARI, FActiveMountEvents& EventContext,
+	FInterfaceWriteScopeLock& ProofOfWriteLock)
+{
+	if (bLoaderIsActive)
+	{
+		return;
+	}
+
+	InitialProviderOnReadyForRegister(UARI, EventContext, EInitialProviderRegisterReason::ActiveMountsInitialize);
+	// Note that we require the ProofOfWriteLock around the read of UARI.GuardedData.State.
+	InitialProviderConditionalRegister(UARI, EventContext, UARI.GuardedData.State);
+	bLoaderIsActive = true;
+}
+
+bool FActiveMountsLoader::TryRegisterAssetProviderAndAppendStateOutsideARLock(UAssetRegistryImpl& UARI,
+	FEventContext& AREventContext, FActiveMountEvents& EventContext, FStringView StateFilePath)
+{
+	FAssetProviderSetKey* ExistingSetKey = AssetProviders.FindByHash(GetTypeHash(StateFilePath), StateFilePath);
+	if (ExistingSetKey)
+	{
+		// Handling register/unregister races between different threads without a separate critical section would
+		// be complicated, so we require that our caller enter the registration critical section that is also entered
+		// when unregistering. AssetProvider.IsShutdown will only ever be set when the AssetProvider is also atomically
+		// cleared within that critical section. Therefore we cannot find an IsShutdown() == true in this function.
+		check(!ExistingSetKey->GetAssetProvider().IsShutdown());
+		return true;
+	}
+
+	// Load from disk outside of the lock
+	FAssetRegistryState StateToAppend;
+	FAssetRegistryLoadResults LoadResults;
+	FActiveMountsLoader* pNewValueOfThis;
+	{
+		UARI.ActiveMountsLoader.Unlock();
+		ON_SCOPE_EXIT
+		{
+			UARI.ActiveMountsLoader.Lock();
+		};
+
+		TUniquePtr<FArchive> FileArchive = OpenReadStateFileArchive(StateFilePath);
+		if (!FileArchive)
+		{
+			return false;
+		}
+
+		FAssetRegistryLoadOptions LoadOptions = UE::AssetRegistry::Premade::GetLoadOptions();
+		check(LoadOptions.bMountPointsRequireSeparateAllocationBlocks == true);
+		if (!StateToAppend.Load(*FileArchive, LoadOptions, LoadResults))
+		{
+			return false;
+		}
+
+		{
+			// We have to enter the AssetRegistry lock around the read of UARI.GuardedData.State. Note that
+			// we also have to enter the AssetRegistry lock before entering ActiveMountsLoader lock, to match required
+			// lock order and avoid deadlocks, so it is important that we have already exited the ActiveMountsLoader
+			// lock above.
+			UE::AssetRegistry::FInterfaceWriteScopeLock InterfaceScopeLock(UARI.InterfaceLock);
+
+			// Reenter the ActiveMountsLoader lock to be able to access ActiveMounts data during InitialProviderOnReadyForRegister
+			{
+				UE::TScopeLock InnerScopeLockOfActiveMountsLock(UARI.ActiveMountsLoader.ActiveMountsLock);
+				// After reentering the lock, check again for any early exits or earlier calculations that may have changed.
+				// DELETED THIS-POINTER WARNING: Check pNewValueOfThis!=nullptr first before reading any data, this
+				// may have been deleted while outside the lock.
+				pNewValueOfThis = UARI.ActiveMountsLoader.Inner.load(std::memory_order_relaxed);
+				if (!pNewValueOfThis || pNewValueOfThis->bShutdown)
+				{
+					return false;
+				}
+				if (ExistingSetKey = AssetProviders.FindByHash(GetTypeHash(StateFilePath), StateFilePath); ExistingSetKey)
+				{
+					check(!ExistingSetKey->GetAssetProvider().IsShutdown());
+					return true;
+				}
+
+				// Do the work we can only do while inside both locks
+				InitialProviderOnReadyForRegister(UARI, EventContext, EInitialProviderRegisterReason::AppendState);
+				InitialProviderConditionalRegister(UARI, EventContext, UARI.GuardedData.State);
+
+				// Call AppendState below outside of the ActiveMountsLoader lock for better performance.
+			}
+
+			// Our contract states that we AppendState as well, so do that now. We have to do it only after reading data
+			// from UARI.GuardedData.State above.
+			UARI.GuardedData.AppendState(AREventContext, StateToAppend, EAppendMode::Append, /*bEmitAssetEvents*/ true);
+		}
+	}
+	// After reentering the lock, check again for any early exits or earlier calculations that may have changed.
+	// DELETED THIS-POINTER WARNING: Check pNewValueOfThis!=nullptr first before reading any data, this
+	// may have been deleted while outside the lock.
+	pNewValueOfThis = UARI.ActiveMountsLoader.Inner.load(std::memory_order_relaxed);
+	if (!pNewValueOfThis || pNewValueOfThis->bShutdown)
+	{
+		return false;
+	}
+	if (ExistingSetKey = AssetProviders.FindByHash(GetTypeHash(StateFilePath), StateFilePath); ExistingSetKey)
+	{
+		check(!ExistingSetKey->GetAssetProvider().IsShutdown());
+		return true;
+	}
+
+	TUniquePtr<FAssetProvider> AssetProviderPtr(new FAssetProvider(StateFilePath));
+	FAssetProvider& AssetProvider = *AssetProviderPtr;
+	AssetProvider.SetStore(LoadResults.Store);
+	bool bAlreadyExisted;
+	AssetProviders.Add(FAssetProviderSetKey(MoveTemp(AssetProviderPtr)), &bAlreadyExisted);
+	// Assert !bAlreadyExisted because we would expensively destroy and recreate if it did and would need to not call
+	// ParseInitialState. bAlreadyExisted cannot occur because we tested AssetProvider.FindByHash above.
+	check(!bAlreadyExisted);
+
+	ParseInitialState(AssetProvider, StateToAppend);
+	return true;
+}
+
+bool FActiveMountsLoader::TryLoadStateAndRegisterAssetProvider(UAssetRegistryImpl& UARI,
+	FActiveMountEvents& EventContext, FInterfaceWriteScopeLock& ProofOfWriteLock, FPakPlatformFile* PakPlatformFile,
+	FStringView StateFilePath, FAssetRegistryState& OutLoadedState)
+{
+	// This function uses a legacy behavior that causes poor performance: load the file off disk inside
+	// the AssetRegistry lock (and to make it even easier, also inside the ActiveMounts lock). It is used only for
+	// the load of cooked AssetRegistry.bin for built-in plugins during the startup LoadPremadeAssetRegistry.
+	// If it becomes a performance issue, we will instead need to trigger these loads asynchronously and delay the
+	// AssetRegistry's reporting of bPreloadingComplete=true that happens at the end of LoadPremadeAssetRegistry
+	// until after the asynchronous loads have completed.
+
+	// Optimization: if we're using pak files then only search paks (avoid unnecessary fallback to loose)
+	TStringBuilder<256> FilePathStr(InPlace, StateFilePath);
+	bool bFileExists = PakPlatformFile != nullptr
+		? PakPlatformFile->FindFileInPakFiles(*FilePathStr)
+		: IFileManager::Get().FileExists(*FilePathStr);
+	if (!bFileExists)
+	{
+		return false;
+	}
+	TUniquePtr<FArchive> FileArchive = OpenReadStateFileArchive(StateFilePath);
+	if (!FileArchive)
+	{
+		return false;
+	}
+
+	FAssetRegistryLoadOptions LoadOptions = UE::AssetRegistry::Premade::GetLoadOptions();
+	check(LoadOptions.bMountPointsRequireSeparateAllocationBlocks == true);
+	FAssetRegistryLoadResults LoadResults;
+	if (!OutLoadedState.Load(*FileArchive, LoadOptions, LoadResults))
+	{
+		return false;
+	}
+
+	FAssetProviderSetKey* ExistingSetKey = AssetProviders.FindByHash(GetTypeHash(StateFilePath), StateFilePath);
+	if (ExistingSetKey)
+	{
+		FAssetProvider& ExistingAssetProvider = ExistingSetKey->GetAssetProvider();
+		// Handling register/unregister races between different threads without a separate critical section would
+		// be complicated, so we require that our caller enter the registration critical section that is also entered
+		// when unregistering. AssetProvider.IsShutdown will only ever be set when the AssetProvider is also atomically
+		// cleared within that critical section. Therefore we cannot find an IsShutdown() == true in this function.
+		check(!ExistingAssetProvider.IsShutdown());
+		return true;
+	}
+
+	TUniquePtr<FAssetProvider> AssetProviderPtr(new FAssetProvider(StateFilePath));
+	FAssetProvider& AssetProvider = *AssetProviderPtr;
+	AssetProvider.SetStore(LoadResults.Store);
+	bool bAlreadyExisted;
+	AssetProviders.Add(FAssetProviderSetKey(MoveTemp(AssetProviderPtr)), &bAlreadyExisted);
+	// Assert !bAlreadyExisted because we would expensively destroy and recreate if it did and would need to not call
+	// ParseInitialState. bAlreadyExisted cannot occur because we tested AssetProviders.FindByHash above.
+	check(!bAlreadyExisted);
+
+	ParseInitialState(AssetProvider, OutLoadedState);
+
+	// Since we successfully loaded the AssetRegistry file, we expect our caller to call AppendState on it, so we have
+	// to trigger the registration of the initial provider now to satisfy the requirement that GuardedData.State is
+	// composed only of AssetDatas from the initial provider.
+	InitialProviderOnReadyForRegister(UARI, EventContext, EInitialProviderRegisterReason::AppendState);
+	InitialProviderConditionalRegister(UARI, EventContext, UARI.GuardedData.State);
+
+	return true;
+}
+
+bool FActiveMountsLoader::TryStartUnregisterAssetProvider(UAssetRegistryImpl& UARI, FActiveMountEvents& EventContext,
+	TArray<UE::Tasks::FTask>& OutWaitOnAsync, FStringView StateFilePath)
+{
+	FAssetProviderSetKey* SetKey = AssetProviders.FindByHash(GetTypeHash(StateFilePath), StateFilePath);
+	if (!SetKey)
+	{
+		return false;
+	}
+	FAssetProvider& AssetProvider = SetKey->GetAssetProvider();
+	if (AssetProvider.IsInitialProvider())
+	{
+		ensureMsgf(false, TEXT("ActiveMounts received a request to unregister an AssetProvider, but the filepath is the global premade AssetRegistry, which is not unregisterable. ")
+			TEXT("Ignoring the unregister request. FilePath = %s"),
+			*FString(StateFilePath));
+		return false;
+	}
+
+	// Handling register/unregister races between different threads without a separate critical section would
+	// be complicated, so we require that our caller enter the registration critical section.
+	// AssetProvider.IsShutdown will only ever be set when the AssetProvider is also atomically cleared within that
+	// critical section. Therefore we cannot find an IsShutdown() already true in this function.
+	check(!AssetProvider.IsShutdown());
+	AssetProvider.SetIsShutdown(true);
+
+	if (AssetProvider.GetLoadTask().IsValid())
+	{
+		OutWaitOnAsync.Add(AssetProvider.GetLoadTask());
+	}
+	return true;
+}
+
+void FActiveMountsLoader::FinishUnregisterAssetProvider(UAssetRegistryImpl& UARI, FActiveMountEvents& EventContext,
+	FStringView StateFilePath)
+{
+	FAssetProviderSetKey* SetKey = AssetProviders.FindByHash(GetTypeHash(StateFilePath), StateFilePath);
+	// Handling register/unregister races between different threads without a separate critical section would
+	// be complicated, so we require that our caller enter the registration critical section.
+	// The AssetProvider was previously found in StartUnregisterAssetProvider within the registration critical section,
+	// so it must still exist now.
+	check(SetKey);
+	FAssetProvider& AssetProvider = SetKey->GetAssetProvider();
+
+	for (TSet<FMountPointData>::TIterator Iter(MountPoints); Iter; ++Iter)
+	{
+		FMountPointData& MountData = *Iter;
+		if (MountData.GetAssetProvider() == &AssetProvider)
+		{
+			NotifyMountPointLoadComplete(EventContext, MountData.GetMountPoint(), true /* bCanceled */);
+			NotifyMountPointUnloadComplete(EventContext, MountData.GetMountPoint(), true /* bCanceled */);
+
+			if (MountData.HasMultipleProviders())
+			{
+				MountData.SetAssetProvider(nullptr);
+			}
+			else
+			{
+				Iter.RemoveCurrent();
+				continue;
+			}
+		}
+	}
+	MountPoints.Shrink();
+	LoadSubscribers.Shrink();
+	UnloadSubscribers.Shrink();
+
+	AssetProviders.Remove(*SetKey);
+	AssetProviders.Shrink();
+}
+
+void FActiveMountsLoader::InitialProviderReport(UAssetRegistryImpl& UARI, FActiveMountEvents& EventContext,
+	FAssetProviderInitData& StateInitData)
+{
+	if (bInitialProviderReported)
+	{
+		return;
+	}
+	bInitialProviderReported = true;
+	FString StateFilePath = UE::AssetRegistry::CreateStandardFilename(StateInitData.StateFilePath);
+	TUniquePtr<FAssetProvider> AssetProviderPtr(new FAssetProvider(StateFilePath));
+	FAssetProvider& AssetProvider = *AssetProviderPtr;
+	AssetProvider.SetStore(StateInitData.StoreUsedByState);
+
+	bool bAlreadyExisted = false;
+	AssetProviders.Add(FAssetProviderSetKey(MoveTemp(AssetProviderPtr)), &bAlreadyExisted);
+	if (bAlreadyExisted)
+	{
+		UE_LOGFMT(LogAssetRegistry, Warning, "Premade Assetregistry uses a filename that was already used in another call to AppendState. "
+			"The earlier call prematurely registered it, and AssetRegistry data for assets without MountPoints currently known "
+			"in FPackageName at the time of its earlier registration will remain loaded for the entire process.");
+		return;
+	}
+	AssetProvider.SetIsInitialProvider(true);
+}
+
+void FActiveMountsLoader::InitialProviderOnReadyForRegister(UAssetRegistryImpl& UARI, FActiveMountEvents& EventContext,
+	EInitialProviderRegisterReason Reason)
+{
+	if (bInitialProviderReadyForRegister)
+	{
+		return;
+	}
+
+	switch (Reason)
+	{
+	case EInitialProviderRegisterReason::ActiveMountsInitialize:
+		// Expected flow, no notes
+		break;
+	case EInitialProviderRegisterReason::EarlierState:
+		UE_LOGFMT(LogAssetRegistry, Error,
+			"Assets were added to the AssetRegistry before LoadPremadeAssetRegistry was called. "
+			"We will prematurely register the premade AssetRegistry with ActiveMountsLoader. "
+			"AssetRegistry data for assets without MountPoints currently known in FPackageName will remain loaded for the entire process.");
+		break;
+	case EInitialProviderRegisterReason::AppendState:
+		UE_LOGFMT(LogAssetRegistry, Display,
+			"AppendState was called with another AssetRegistry before ActiveMountsInitialize was called. Triggering ActiveMountsInitialize.");
+		break;
+	default:
+		checkNoEntry();
+		break;
+	}
+	bInitialProviderReadyForRegister = true;
+}
+
+void FActiveMountsLoader::InitialProviderConditionalRegister(UAssetRegistryImpl& UARI,
+	FActiveMountEvents& EventContext, FAssetRegistryState& InitialProviderState)
+{
+	if (bInitialProviderRegistered
+		|| !bInitialProviderReported || !bInitialProviderReadyForRegister)
+	{
+		return;
+	}
+
+	for (FAssetProviderSetKey& SetKey : AssetProviders)
+	{
+		FAssetProvider& AssetProvider = SetKey.GetAssetProvider();
+		if (AssetProvider.IsInitialProvider())
+		{
+			ParseInitialState(AssetProvider, InitialProviderState);
+			break;
+		}
+	}
+	bInitialProviderRegistered = true;
+}
+
+void FActiveMountsLoader::OnMountPointsMounted(UAssetRegistryImpl& UARI, FActiveMountEvents& EventContext,
+	TConstArrayView<TRefCountPtr<UE::PackageName::IMountPoint>> InMountPoints)
+{
+	if (InMountPoints.IsEmpty())
+	{
+		return;
+	}
+	if (bLoaderIsActive && bAutoUnload && bAutoUnloadInitialUnloadTriggered)
+	{
+		LoadAsync(UARI, EventContext, InMountPoints, TFunction<void(bool bCanceled)>());
+	}
+}
+
+void FActiveMountsLoader::OnMountPointsDismounted(UAssetRegistryImpl& UARI, FActiveMountEvents& EventContext,
+	FInterfaceWriteScopeLock& ProofOfWriteLock,
+	TConstArrayView<TRefCountPtr<UE::PackageName::IMountPoint>> InMountPoints)
+{
+	if (InMountPoints.IsEmpty())
+	{
+		return;
+	}
+	if (bAutoUnload)
+	{
+		if (!bAutoUnloadInitialUnloadTriggered)
+		{
+			// Automatically call OnActiveMountsInitialize if not already called.
+			OnActiveMountsInitialize(UARI, EventContext, ProofOfWriteLock);
+			// Automatically unload AssetRegistry data for unmounted MountPoints.
+			UnloadAllUnmountedAsync(UARI, EventContext, TFunction<void(bool bCanceled)>());
+			bAutoUnloadInitialUnloadTriggered = true;
+		}
+	}
+	if (bLoaderIsActive && bAutoUnload)
+	{
+		UnloadAsync(UARI, EventContext, InMountPoints, TFunction<void(bool bCanceled)>());
+	}
+}
+
+void FActiveMountsLoader::ParseInitialState(FAssetProvider& AssetProvider, const FAssetRegistryState& State)
+{
+	using namespace UE::PackageName;
+
+	// If this parsing is too expensive, we can instead do the parsing during cook, and save the list of MountPoints
+	// into the serialized AssetRegistry and pass them from FAssetRegistryState::Load rather than calculating them here.
+	TDirectoryTree<TRefCountPtr<IMountPoint>> UsedMountPoints;
+	TStringBuilder<256> PathStr;
+	TSet<FString> UnmountedRoots;
+	IPluginManager& PluginManager = IPluginManager::Get();
+
+	State.EnumerateAllPaths([&PathStr, &UsedMountPoints, &UnmountedRoots, &AssetProvider, &PluginManager](FName PathName)
+		{
+			PathStr.Reset();
+			PathStr << PathName;
+			if (PathStr.Len() > 0 && !FStringView(PathStr).EndsWith('/'))
+			{
+				PathStr << '/';
+			}
+			FStringView MountPointRoot = FPackageName::SplitPackageNameRoot(PathStr, nullptr,
+				FPackageName::EPathFormatFlags::MountPointSlashes);
+			if (MountPointRoot.IsEmpty())
+			{
+				// Invalid PathName, we don't allow unloading these and don't allow adding them to mount points.
+				UE_LOGFMT(LogAssetRegistry, Warning,
+					"ActiveMountsLoader found an invalid path in a cooked AssetRegistry. AssetRegistry data for assets in this path will always remain loaded. "
+					"AssetRegistryFile: '{AssetRegistryFile}', Path: '{Path}'.",
+					*FString(AssetProvider.GetFilePath()), *PathStr);
+				return;
+			}
+			TRefCountPtr<IMountPoint>* MountPointPtr = UsedMountPoints.FindClosestValue(PathStr);
+			if (!MountPointPtr)
+			{
+				TRefCountPtr<IMountPoint> MountPoint = FPackageName::FindMountPointByChildLongPackageName(PathStr);
+
+				if (!MountPoint)
+				{
+					// Look for a plugin matching the MountPointRoot, if one exists, add it as an unmounted mountpoint.
+					// Get "PluginName" from "/PluginName/" by trimming 1 from begin and end.
+					FStringView PluginName = MountPointRoot.Mid(1, MountPointRoot.Len() - 2);
+					TSharedPtr<IPlugin> Plugin = PluginManager.FindPlugin(PluginName);
+					if (Plugin)
+					{
+						MountPoint = FPackageName::FindOrAddMountPoint(MountPointRoot, Plugin->GetContentDir());
+					}
+				}
+
+				if (MountPoint)
+				{
+					UsedMountPoints.FindOrAdd(MountPoint->GetLongPackageName()) = MoveTemp(MountPoint);
+				}
+				else
+				{
+					if (!UnmountedRoots.FindByHash(GetTypeHash(MountPointRoot), MountPointRoot))
+					{
+						UnmountedRoots.Add(FString(MountPointRoot));
+					}
+				}
+			}
+		});
+
+	if (!UnmountedRoots.IsEmpty())
+	{
+		if (!AssetProvider.HasMissingMountPoints())
+		{
+			AssetProvider.SetHasMissingMountPoints(true);
+			AssetProvider.SetIsErrorNoUnload(true);
+			UE_LOGFMT(LogAssetRegistry, Error,
+				"ActiveMountsLoader found a path in an unrecognized mountpoint in a cooked AssetRegistry. "
+				"We cannot validate the correctness of unloading without knowing the mountpoint, so we are marking this AssetRegistry as not unloadable. "
+				"AssetRegistry data for assets in this cooked AssetRegistry will always remain loaded. "
+				"To fix this, declare all MountPoints before calling ActiveMountsRegisterAndLoadFilePathAsync or ActiveMountsInitialize. "
+				"AssetRegistryFile: '{AssetRegistryFile}', Unrecognized mountpoint : '{MountPoint}'.",
+				*FString(AssetProvider.GetFilePath()), **UnmountedRoots.FindArbitraryElement());
+		}
+	}
+
+	// Record the MountPoints
+	for (const TPair<FStringView, TRefCountPtr<IMountPoint>>& Pair : UsedMountPoints)
+	{
+		IMountPoint* MountPoint = Pair.Value.GetReference();
+		FMountPointData& MountData = FindOrAddMountPointData(MountPoint, EMountState::Loaded);
+
+		if (MountData.GetAssetProvider() != nullptr || MountData.HasMultipleProviders())
+		{
+			if (!MountData.HasMultipleProviders())
+			{
+				MountData.SetHasMultipleProviders(true);
+				MountData.SetIsErrorNoUnload(true);
+				UE_LOGFMT(LogAssetRegistry, Error,
+					"ActiveMountsLoader found assets under MountPoint present in multiple cooked AssetRegistries. "
+					"Unloading is not yet supported when assets are present in multiple registries. "
+					"Assets from this MountPoint will no longer be unloadable. "
+					"MountPoint: '{MountPoint}', AssetRegistryFile1: '{AssetRegistryFile1}', AssetRegistryFile2: '{AssetRegistryFile2}'.",
+					*FString(MountData.GetMountPoint()->GetLongPackageName()),
+					*FString(MountData.GetAssetProvider()->GetFilePath()),
+					*FString(AssetProvider.GetFilePath()));
+			}
+			continue;
+		}
+		MountData.SetAssetProvider(&AssetProvider);
+
+		// Mount points have a starting slash and terminating slash: "/MountPoint/"
+		// If they have more than one path segment - "/Path/Leaf/" we need to use a slower method to find them.
+		SIZE_T NumSlashes = Algo::Count(MountPoint->GetLongPackageName(), '/');
+		if (NumSlashes > 2)
+		{
+			MountData.SetIsMultiSegmentName(true);
+		}
+	}
+
+	// Look for parent/child mounts
+	for (const TPair<FStringView, TRefCountPtr<IMountPoint>>& Pair : UsedMountPoints)
+	{
+		IMountPoint* Parent = Pair.Value.GetReference();
+		TRefCountPtr<IMountPoint>* Closest = UsedMountPoints.FindClosestValue(Parent->GetLongPackageName());
+		if (Closest && Closest->GetReference() != Parent)
+		{
+			IMountPoint* Child = Closest->GetReference();
+			// Don't call FindOrAdd because that could invalidate our FMountPointData pointers, but do assert that
+			// they exist from our FindOrAdd call above.
+			FMountPointData* ParentData = FindMountPointData(Parent);
+			FMountPointData* ChildData = FindMountPointData(Child);
+			check(ParentData);
+			check(ChildData);
+			if (!ParentData->IsParentOrChild() || !ChildData->IsParentOrChild())
+			{
+				ParentData->SetIsParentOrChild(true);
+				ParentData->SetIsErrorNoUnload(true);
+				ChildData->SetIsParentOrChild(true);
+				ChildData->SetIsErrorNoUnload(true);
+				UE_LOGFMT(LogAssetRegistry, Warning,
+					"ActiveMountsLoader found assets in cooked AssetRegistry that use a parent and child MountPoint. "
+					"We do not yet efficiently support unloading and reloading assets when parent and child mount points are present, so we are disabling unload of any assets in these mountpoints. "
+					"AssetRegistryFile: '{AssetRegistryFile}', MountPoint1: '{MountPoint1}', MountPoint2: '{MountPoint2}'.",
+					*FString(AssetProvider.GetFilePath()),
+					*FString(Parent->GetLongPackageName()),
+					*FString(Child->GetLongPackageName()));
+			}
+		}
+	}
+}
+
+FMountPointData& FActiveMountsLoader::FindOrAddMountPointData(UE::PackageName::IMountPoint* MountPoint,
+	EMountState InitialState)
+{
+	using namespace UE::PackageName;
+
+	// We guarantee returning a FMountPointData& and don't have a mechanism for reporting invalid argument,
+	// so caller must guarantee MountPoint != nullptr.
+	check(MountPoint != nullptr);
+
+	FMountPointData* MountPointData = MountPoints.FindByHash(GetTypeHash(MountPoint), MountPoint);
+	if (MountPointData)
+	{
+		return *MountPointData;
+	}
+
+	FSetElementId Id = MountPoints.Add(FMountPointData(TRefCountPtr<IMountPoint>(MountPoint), InitialState));
+	return MountPoints[Id];
+}
+
+FMountPointData* FActiveMountsLoader::FindMountPointData(UE::PackageName::IMountPoint* MountPoint)
+{
+	return MountPoints.FindByHash(GetTypeHash(MountPoint), MountPoint);
+}
+
+void FActiveMountsLoader::NotifyMountPointLoadComplete(FActiveMountEvents& EventContext,
+	const TRefCountPtr<UE::PackageName::IMountPoint>& MountPoint, bool bCanceled)
+{
+	TArray<TRefCountPtr<FMountLoadEventSubscriber>> Subscribers;
+	if (LoadSubscribers.RemoveAndCopyValue(MountPoint, Subscribers))
+	{
+		for (TRefCountPtr<FMountLoadEventSubscriber>& Subscriber : Subscribers)
+		{
+			check(Subscriber); // nullptrs are never added to LoadSubscribers
+			if (Subscriber->RemainingCount <= 0 || !Subscriber->bIsLoadSubscriber)
+			{
+				ensureMsgf(false,
+					TEXT("ActiveMounts found an invalid MountLoadEventSubscriber registered in LoadSubscribers. ")
+					TEXT("We will not send the notification and this may cause a softlock. ")
+					TEXT("MountPoint = %s, RemainingCount = %d, bIsLoadSubscriber = %s, bWasCanceled = %s"),
+					*FString(MountPoint->GetLongPackageName()), Subscriber->RemainingCount,
+					*::LexToString(Subscriber->bIsLoadSubscriber), *::LexToString(Subscriber->bWasCanceled));
+				continue;
+			}
+			Subscriber->bWasCanceled |= bCanceled;
+			Subscriber->RemainingCount--;
+			if (Subscriber->RemainingCount == 0)
+			{
+				EventContext.AddCompletion(MoveTemp(Subscriber->OnComplete), Subscriber->bWasCanceled);
+			}
+		}
+		ConditionalShrinkUsuallyEmptyContainer(LoadSubscribers);
+	}
+}
+
+void FActiveMountsLoader::NotifyMountPointUnloadComplete(FActiveMountEvents& EventContext,
+	const TRefCountPtr<UE::PackageName::IMountPoint>& MountPoint, bool bCanceled)
+{
+	TArray<TRefCountPtr<FMountLoadEventSubscriber>> Subscribers;
+	if (UnloadSubscribers.RemoveAndCopyValue(MountPoint, Subscribers))
+	{
+		for (TRefCountPtr<FMountLoadEventSubscriber>& Subscriber : Subscribers)
+		{
+			check(Subscriber); // nullptrs are never added to UnloadSubscribers
+			if (Subscriber->RemainingCount <= 0 || Subscriber->bIsLoadSubscriber)
+			{
+				ensureMsgf(false,
+					TEXT("ActiveMounts found an invalid MountLoadEventSubscriber registered in UnloadSubscribers. ")
+					TEXT("We will not send the notification and this may cause a softlock. ")
+					TEXT("MountPoint = %s, RemainingCount = %d, bIsLoadSubscriber = %s, bWasCanceled = %s"),
+					*FString(MountPoint->GetLongPackageName()), Subscriber->RemainingCount,
+					*::LexToString(Subscriber->bIsLoadSubscriber), *::LexToString(Subscriber->bWasCanceled));
+				continue;
+			}
+			Subscriber->bWasCanceled |= bCanceled;
+			Subscriber->RemainingCount--;
+			if (Subscriber->RemainingCount == 0)
+			{
+				EventContext.AddCompletion(MoveTemp(Subscriber->OnComplete), Subscriber->bWasCanceled);
+			}
+		}
+		ConditionalShrinkUsuallyEmptyContainer(UnloadSubscribers);
+	}
+}
+
+void FActiveMountsLoader::AddRefMountPointLoad(UAssetRegistryImpl& UARI, FMountPointData* MountPointData)
+{
+	// Caller must not call AddRef if MountPointData is null or does not have an AssetProvider.
+	check(MountPointData && MountPointData->GetAssetProvider());
+	FAssetProvider* AssetProvider = MountPointData->GetAssetProvider();
+	AssetProvider->GetLoadRefCountPtr()++;
+	UE::Tasks::FTask& LoadTask = AssetProvider->GetLoadTask();
+	if (LoadTask.IsValid())
+	{
+		// Nothing to do, the task is still running its loop and will pick up the new MountPoint.
+		return;
+	}
+
+	LoadTask = UE::Tasks::Launch(TEXT("ActiveMountsLoad"), [&UARI, this, AssetProvider]()
+		{
+			RunAssetProviderLoadTask(UARI, *AssetProvider);
+		});
+}
+
+void FActiveMountsLoader::ReleaseMountPointLoad(UAssetRegistryImpl& UARI, FMountPointData* MountPointData)
+{
+	// Caller must not call Release if MountPointData is null or does not have an AssetProvider.
+	check(MountPointData && MountPointData->GetAssetProvider());
+	MountPointData->GetAssetProvider()->GetLoadRefCountPtr()--;
+	// Release is not responsible for shutting down the LoadTask; RunAssetProviderLoadTask polls the
+	// refcount and shuts itself down when it sees 0.
+}
+
+void FActiveMountsLoader::AddRefMountPointUnload(UAssetRegistryImpl& UARI, FMountPointData* MountPointData)
+{
+	// Caller must not call AddRef if MountPointData is null or does not have an AssetProvider.
+	check(MountPointData && MountPointData->GetAssetProvider());
+	MountPointUnloadRefCount++;
+	if (UnloadTask.IsValid())
+	{
+		// Nothing to do, the task is still running its loop and will pick up the new MountPoint.
+		return;
+	}
+	UnloadTask = UE::Tasks::Launch(TEXT("ActiveMountsUnload"), [&UARI, this]()
+		{
+			RunMountPointUnloadTask(UARI);
+		});
+}
+
+void FActiveMountsLoader::ReleaseMountPointUnload(UAssetRegistryImpl& UARI, FMountPointData* MountPointData)
+{
+	// Caller must not call Release if MountPointData is null.
+	check(MountPointData);
+	MountPointUnloadRefCount--;
+	// Release is not responsible for shutting down the UnloadTask; RunMountPointUnloadTaskRunMountPointUnloadTask polls the
+	// refcount and shuts itself down when it sees 0.
+}
+
+void FActiveMountsLoader::RunAssetProviderLoadTask(UAssetRegistryImpl& UARI, FAssetProvider& AssetProvider)
+{
+	LLM_SCOPE(ELLMTag::AssetRegistry);
+
+	auto MarkAsyncThreadComplete = [&AssetProvider]()
+		{
+			// Caller must call only while within the ActiveMounts critical section, to avoid a softlock race
+			// where the main thread sees that the AsyncTask still exists and does not schedule another one,
+			// but the AsyncTask is shutting down and does not plan to check for further work to do.
+			AssetProvider.GetLoadTask() = UE::Tasks::FTask();
+		};
+
+	// Requests might come in unbatched, or in batches that race with our load and come in after we
+	// decide which MountPoints we will load from the disk. Keep repeatedly loading the AssetRegistry.bin file
+	// and creating AssetDatas for the current batch of MountPoints out of it until we have no further load requests.
+	for (int32 NumIterations = 0; /* Loop is terminated by break statements inside a lock */; ++NumIterations)
+	{
+		FString StateFilePath;
+		TArray<TRefCountPtr<UE::PackageName::IMountPoint>> MountPointsToLoad;
+		bool bHasMultiSegmentPath = false;
+		{
+			FActiveMountsLoader* pNewThis = UARI.ActiveMountsLoader.LockIfEnabled();
+			// ActiveMounts were previously enabled otherwise this function would not be called, and Shutdown is not
+			// supposed to disable ActiveMounts until after the thread running this function shuts down. So we
+			// should still be enabled and pNewThis should be non-null and we don't have to handle exiting.
+			check(pNewThis);
+			ON_SCOPE_EXIT
+			{
+				UARI.ActiveMountsLoader.Unlock();
+			};
+
+			// We do however have to handle the shutdown signal, either global or for the AssetProvider, and leave the thread.
+			if (bShutdown || AssetProvider.IsShutdown())
+			{
+				MarkAsyncThreadComplete();
+				return;
+			}
+
+			// If no further MountPoints to load, leave the thread.
+			if (AssetProvider.GetLoadRefCountPtr() == 0)
+			{
+				MarkAsyncThreadComplete();
+				return;
+			}
+
+			// If we have already run once, reschedule our task rather than continuing to run, to give any other
+			// high priority tasks that have come in a chance to take our WorkerThread.
+			if (NumIterations > 0)
+			{
+				AssetProvider.GetLoadTask() = UE::Tasks::Launch(TEXT("ActiveMountsLoad"), [&UARI, this, AssetProvider=&AssetProvider]()
+					{
+						RunAssetProviderLoadTask(UARI, *AssetProvider);
+					});
+				// Do not call MarkAsyncThreadComplete; we are still running the task.
+				return;
+			}
+
+			StateFilePath = AssetProvider.GetFilePath();
+			// Leave the lock to load the file off of disk.
+		}
+
+		// Open the file and preload it, but do not deserialize it yet; we need to know the list of mountpoints first,
+		// and we want to capture those as late as possible in case of stragglers while we are preloading the file.
+		TUniquePtr<FArchive> FileArchive = OpenReadStateFileArchive(StateFilePath);
+		FixedTagPrivate::FWeakStorePtr Store;
+
+		// Reenter the lock to get the list of mountpoints to load from the statefile.
+		{
+			FActiveMountsLoader* pNewThis = UARI.ActiveMountsLoader.LockIfEnabled();
+			// Per the comment above, we should still be enabled and pNewThis should be non-null and we don't have to handle exiting.
+			check(pNewThis);
+			ON_SCOPE_EXIT
+			{
+				UARI.ActiveMountsLoader.Unlock();
+			};
+
+			// We do however have to handle the shutdown signal, either global or for the AssetProvider, and leave the thread.
+			if (bShutdown || AssetProvider.IsShutdown())
+			{
+				MarkAsyncThreadComplete();
+				return;
+			}
+
+			// Check whether LoadRefCountPtr is still non-zero.
+			if (AssetProvider.GetLoadRefCountPtr() == 0)
+			{
+				MarkAsyncThreadComplete();
+				return;
+			}
+
+			MountPointsToLoad.Reserve(AssetProvider.GetLoadRefCountPtr());
+			for (FMountPointData& MountData : MountPoints)
+			{
+				if (MountData.GetCurrentState() == EMountState::Loading
+					&& MountData.GetAssetProvider() == &AssetProvider)
+				{
+					MountPointsToLoad.Add(MountData.GetMountPoint());
+					bHasMultiSegmentPath |= MountData.IsMultiSegmentName();
+				}
+			}
+			if (MountPointsToLoad.Num() == 0)
+			{
+				// Unexpected, because MountPoints are supposed to drop their refcount before leaving the Loading
+				// state.
+				UE_LOGFMT(LogAssetRegistry, Error,
+					"ActiveMounts: Some MountPoint left the loading state without dropping its refcount on the AssetProvider. "
+					"AssetRegistryFile: {AssetRegistryFile}.",
+					*StateFilePath);
+				MarkAsyncThreadComplete();
+				return;
+			}
+			Store = AssetProvider.GetStore();
+			// Leave the lock to deserialize the AssetRegistry state
+		}
+
+		TFunction<bool(FName PackageName)> KeepFunction;
+		TDirectoryTree<bool> KeepTree;
+		TSet<FString> KeepSet;
+		if (!bHasMultiSegmentPath)
+		{
+			for (TRefCountPtr<UE::PackageName::IMountPoint>& MountPoint : MountPointsToLoad)
+			{
+				KeepSet.Add(FString(MountPoint->GetLongPackageName()));
+			}
+			KeepFunction = [&KeepSet](FName PackageName) -> bool
+				{
+					TStringBuilder<256> PathStr(InPlace, PackageName);
+					// We have put MountPoint with slashes - /MountPoint/ - in our set, so get slashes for lookup.
+					FStringView MountPointRoot = FPackageName::SplitPackageNameRoot(PathStr, nullptr,
+						FPackageName::EPathFormatFlags::MountPointSlashes);
+					if (MountPointRoot.IsEmpty())
+					{
+						// We never unload these (we log an error during initial parsing that we keep them loaded)
+						// so we do not need to reload them now.
+						return false;
+					}
+					return KeepSet.ContainsByHash(GetTypeHash(MountPointRoot), MountPointRoot);
+				};
+		}
+		else
+		{
+			// In the case of mountpoint with more than one segment, we have to do a prefix lookup in our list of
+			// mountpoints rather than a faster TSet lookup, because we don't know how much of the packagename to use
+			// for the TSet lookup.
+			for (TRefCountPtr<UE::PackageName::IMountPoint>& MountPoint : MountPointsToLoad)
+			{
+				KeepTree.FindOrAdd(MountPoint->GetLongPackageName()) = true;
+			}
+			KeepFunction = [&KeepTree](FName PackageName) -> bool
+				{
+					TStringBuilder<256> PathStr(InPlace, PackageName);
+					FStringView ParentDir = FPathViews::GetPath(PathStr);
+					if (ParentDir.IsEmpty())
+					{
+						// We never unload these (we log an error during initial parsing that we keep them loaded)
+						// so we do not need to reload them now.
+						return false;
+					}
+					return KeepTree.ContainsPathOrParent(ParentDir);
+				};
+		}
+
+		bool bLoadSuccessful = false;
+		FAssetRegistryState StateToAppend;
+		FAssetRegistryLoadOptions LoadOptions = UE::AssetRegistry::Premade::GetLoadOptions();
+		check(LoadOptions.bMountPointsRequireSeparateAllocationBlocks == true);
+		LoadOptions.ShouldKeepPackage = MoveTemp(KeepFunction);
+		LoadOptions.Store = MoveTemp(Store);
+		FAssetRegistryLoadResults LoadResults;
+
+		if (FileArchive)
+		{
+			bLoadSuccessful = StateToAppend.Load(*FileArchive, LoadOptions, LoadResults);
+		}
+
+		if (bLoadSuccessful)
+		{
+			FEventContext AREventContext;
+			{
+				UE::AssetRegistry::FInterfaceWriteScopeLock InterfaceScopeLock(UARI.InterfaceLock);
+				// Use ConsumeState rather than AppendState so that we don't create a new FStore for the
+				// AssetDatas moved over from the StateToAppend, and we don't wastefully reallocate the AssetDatas.
+				// ConsumeState also uses EAppendMode::OnlyUpdateNew, which is fine for our purpose which is
+				// just reloading data that we previously unloaded.
+				UARI.GuardedData.ConsumeState(AREventContext, MoveTemp(StateToAppend), /*bEmitAssetEvents*/ true);
+			}
+			UARI.Broadcast(AREventContext);
+		}
+
+		// Update state data and send events for all of the MountPointsToLoad
+		FActiveMountEvents EventContext;
+		{
+			FActiveMountsLoader* pNewThis = UARI.ActiveMountsLoader.LockIfEnabled();
+			// Per the comment above, we should still be enabled and pNewThis should be non-null and we don't have to handle exiting.
+			check(pNewThis);
+			ON_SCOPE_EXIT
+			{
+				UARI.ActiveMountsLoader.Unlock();
+			};
+
+			// We do however have to handle the shutdown signal, either global or for the AssetProvider, and leave the thread.
+			if (bShutdown || AssetProvider.IsShutdown())
+			{
+				MarkAsyncThreadComplete();
+				return;
+			}
+
+			if (bLoadSuccessful)
+			{
+				// Refresh the Store pointer on the AssetProvider, in case the old store had been garbage collected and the load
+				// allocated a new store.
+				AssetProvider.SetStore(LoadResults.Store);
+			}
+			else
+			{
+				UE_LOGFMT(LogAssetRegistry, Error,
+					"ActiveMounts: Failed to reload AssetRegistry data from a cooked AssetRegistry after some of its mountpoints were unloaded and then reloaded. "
+					"The assets will be missing from the AssetRegistry, for all of the MountPoints in the cooked AssetRegistry that were previously unloaded. "
+					"AssetRegistryFile: '{AssetRegistryFile}', example missing MountPoint: '{MountPoint}'.",
+					*StateFilePath, *FString(MountPointsToLoad[0]->GetLongPackageName()));
+			}
+
+			for (TRefCountPtr<UE::PackageName::IMountPoint>& MountPoint : MountPointsToLoad)
+			{
+				// The MountPoints should still be in the Loading state, because nothing is allowed to take them out
+				// of the loading state other than this function running on their AssetProvider, and we do not allow
+				// a second instance of this function on the same AssetProvider at a time.
+				FMountPointData* MountData = MountPoints.FindByHash(GetTypeHash(MountPoint), MountPoint);
+				if (!MountData || MountData->GetAssetProvider() == nullptr)
+				{
+					// AssetProvider containing the MountData was unregistered.
+					continue;
+				}
+
+				if (MountData->GetCurrentState() != EMountState::Loading)
+				{
+					UE_LOGFMT(LogAssetRegistry, Error,
+						"ActiveMounts: Failed to reload AssetRegistry data for mountpoint because the mountpoint unexpectedly changed loading state while we were loading it. "
+						"The assets in the MountPoint might now be be missing from the AssetRegistry. "
+						"AssetRegistryFile: '{StateFilePath}', MountPoint: '{MountPoint}', CurrentState: '{CurrentState}.'",
+						*StateFilePath, *FString(MountPoint->GetLongPackageName()), int(MountData->GetCurrentState()));
+					continue;
+				}
+
+				// Even if the load failed, we still have to mark all the MountDatas as loaded, to avoid having this thread
+				// repeatedly continue trying to load them and kill performance.
+				MountData->SetCurrentState(EMountState::Loaded);
+				ReleaseMountPointLoad(UARI, MountData);
+				NotifyMountPointLoadComplete(EventContext, MountPoint, false /* bCanceled */);
+
+				// It is possible that a request came in to unload the MountPoint while we were still working on the
+				// load. The unload requested was recorded in the targetstate if so, and no other action taken. Check
+				// for that now and kick off the unload if so.
+				if (MountData->GetTargetState() != EMountState::Loaded)
+				{
+					// TargetState can only be Loaded or Unloaded, transition states are not possible for it.
+					check(MountData->GetTargetState() == EMountState::Unloaded);
+					MountData->SetCurrentState(EMountState::Unloading);
+					check(MountData->GetAssetProvider()); // We checked above for GetAssetProvider()
+					AddRefMountPointUnload(UARI, MountData);
+				}
+			}
+		}
+	}
+}
+
+void FActiveMountsLoader::RunMountPointUnloadTask(UAssetRegistryImpl& UARI)
+{
+	LLM_SCOPE(ELLMTag::AssetRegistry);
+
+	auto MarkAsyncThreadComplete = [this]()
+		{
+			// Caller must call only while within the ActiveMounts critical section, to avoid a softlock race
+			// where the main thread sees that the AsyncTask still exists and does not schedule another one,
+			// but the AsyncTask is shutting down and does not plan to check for further work to do.
+			UnloadTask = UE::Tasks::FTask();
+		};
+
+	// Requests might come in unbatched, or in batches that race with our unload and come in after we
+	// decide which MountPoints we will unload. Keep repeatedly unloading the current batch until we have
+	// no further unload requests.
+	for (int32 NumIterations = 0; /* Loop is terminated by break statements inside a lock */ ; ++NumIterations)
+	{
+		TArray<TRefCountPtr<UE::PackageName::IMountPoint>> MountPointsToUnload;
+		{
+			FActiveMountsLoader* pNewThis = UARI.ActiveMountsLoader.LockIfEnabled();
+			// ActiveMounts were previously enabled otherwise this function would not be called, and Shutdown is not
+			// supposed to disable ActiveMounts until after the thread running this function shuts down. So we
+			// should still be enabled and pNewThis should be non-null and we don't have to handle exiting.
+			check(pNewThis);
+			ON_SCOPE_EXIT
+			{
+				UARI.ActiveMountsLoader.Unlock();
+			};
+
+			// We do however have to handle the shutdown signal, and leave the thread.
+			if (bShutdown)
+			{
+				MarkAsyncThreadComplete();
+				return;
+			}
+
+			// If no further MountPoints to load, leave the thread.
+			if (MountPointUnloadRefCount == 0)
+			{
+				MarkAsyncThreadComplete();
+				return;
+			}
+
+			// If we have already run once, reschedule our task rather than continuing to run, to give any other
+			// high priority tasks that have come in a chance to take our WorkerThread.
+			if (NumIterations > 0)
+			{
+				UnloadTask = UE::Tasks::Launch(TEXT("ActiveMountsUnload"), [&UARI, this]()
+					{
+						RunMountPointUnloadTask(UARI);
+					});
+				// Do not call MarkAsyncThreadComplete; we are still running the task.
+				return;
+			}
+
+			MountPointsToUnload.Reserve(MountPointUnloadRefCount);
+			for (FMountPointData& MountData : MountPoints)
+			{
+				if (MountData.GetCurrentState() == EMountState::Unloading
+					&& MountData.GetAssetProvider() != nullptr)
+				{
+					MountPointsToUnload.Add(MountData.GetMountPoint());
+				}
+			}
+			if (MountPointsToUnload.Num() == 0)
+			{
+				// Unexpected, because MountPoints are supposed to drop their refcount before leaving the Unloading
+				// state.
+				UE_LOGFMT(LogAssetRegistry, Error,
+					"ActiveMounts: Some MountPoint left the unloading state without dropping its MountPointUnloadRefCount.");
+				MarkAsyncThreadComplete();
+				return;
+			}
+			// Leave the ActiveMounts lock to modify the AssetRegistry state.
+		}
+
+		FEventContext AREventContext;
+		{
+			UE::AssetRegistry::FInterfaceWriteScopeLock InterfaceScopeLock(UARI.InterfaceLock);
+			UARI.GuardedData.UnloadMountPoints(AREventContext, MountPointsToUnload);
+		}
+		UARI.Broadcast(AREventContext);
+
+		// Update state data and send events for all of the unloaded MountPoints.
+		FActiveMountEvents EventContext;
+		{
+			FActiveMountsLoader* pNewThis = UARI.ActiveMountsLoader.LockIfEnabled();
+			// Per the comment above, we should still be enabled and pNewThis should be non-null and we don't have to handle exiting.
+			check(pNewThis);
+			ON_SCOPE_EXIT
+			{
+				UARI.ActiveMountsLoader.Unlock();
+			};
+
+			// We do however have to handle the shutdown signal, either global or for the AssetProvider, and leave the thread.
+			if (bShutdown)
+			{
+				MarkAsyncThreadComplete();
+				return;
+			}
+
+			for (TRefCountPtr<UE::PackageName::IMountPoint>& MountPoint : MountPointsToUnload)
+			{
+				// The MountPoints should still be in the Unloading state, because nothing is allowed to take them out
+				// of the Unloading state other than this function, and we do not allow a second instance of this
+				// function at a time.
+				FMountPointData* MountData = MountPoints.FindByHash(GetTypeHash(MountPoint), MountPoint);
+				if (!MountData || MountData->GetAssetProvider() == nullptr)
+				{
+					// AssetProvider owning the MountData was unregistered
+					continue;
+				}
+				if (MountData->GetCurrentState() != EMountState::Unloading)
+				{
+					UE_LOGFMT(LogAssetRegistry, Error,
+						"ActiveMounts: Failed to unload AssetRegistry data for mountpoint because the mountpoint unexpectedly changed unloading state while we were unloading it. "
+						"The assets in the MountPoint might now be be missing from the AssetRegistry. "
+						"MountPoint: '{MountPoint}', CurrentState: '{CurrentState}'.",
+						*FString(MountPoint->GetLongPackageName()), int(MountData->GetCurrentState()));
+					continue;
+				}
+
+				MountData->SetCurrentState(EMountState::Unloaded);
+				ReleaseMountPointUnload(UARI, MountData);
+				NotifyMountPointUnloadComplete(EventContext, MountPoint, false /* bCanceled */);
+
+				// It is possible that a request came in to load the MountPoint while we were still working on the
+				// unload. The load requested was recorded in the targetstate if so, and no other action taken. Check
+				// for that now and kick off the load if so.
+				if (MountData->GetTargetState() != EMountState::Unloaded)
+				{
+					// TargetState can only be Loaded or Unloaded, transition states are not possible for it.
+					check(MountData->GetTargetState() == EMountState::Loaded);
+					MountData->SetCurrentState(EMountState::Loading);
+					check(MountData->GetAssetProvider()); // We checked above for GetAssetProvider
+					AddRefMountPointLoad(UARI, MountData);
+				}
+			}
+		}
+	}
+}
+
+} // namespace UE::AssetRegistry::Impl
diff --git a/Engine/Source/Runtime/AssetRegistry/Private/ActiveMountsLoader.h b/Engine/Source/Runtime/AssetRegistry/Private/ActiveMountsLoader.h
new file mode 100644
--- /dev/null
+++ b/Engine/Source/Runtime/AssetRegistry/Private/ActiveMountsLoader.h
@@ -0,0 +1,355 @@
+// Copyright Epic Games, Inc. All Rights Reserved.
+
+#pragma once
+
+#include "ARLoader.h"
+#include "AssetRegistry/AssetRegistryState.h"
+#include "Containers/Array.h"
+#include "Containers/ArrayView.h"
+#include "Containers/Map.h"
+#include "Containers/Set.h"
+#include "HAL/PlatformMath.h"
+#include "Misc/Optional.h"
+#include "Tasks/Task.h"
+
+namespace UE::AssetRegistry::Impl
+{
+
+/** Loading states of MountPoints managed by the ActiveMountsLoader state machine. */
+enum class EMountState : uint8
+{
+	Unloaded,
+	Loading,
+	Loaded,
+	Unloading,
+	Count,
+	NumBits = FPlatformMath::ConstExprCeilLogTwo(EMountState::Count),
+};
+
+/**
+ * A source of AssetDatas for the ActiveMountsLoader to apply to the AssetRegistry. Currently the only such source is
+ * a cooked AssetRegistry.bin file on disk. The filepath for a registered AssetRegistry.bin file is stored in an
+ * AssetProvider, and a record of the MountPoints that contain the AssetData paths within the AssetProvider are stored
+ * in the ActiveMountsLoader and associated with the AssetProvider.  The ActiveMountsLoader launches an async task to
+ * serialize the file from disk whenever it needs to reload AssetDatas under any of the MountPoints that contain the
+ * AssetData paths within the AssetProvider.
+ */
+class FAssetProvider
+{
+public:
+	explicit FAssetProvider(FStringView InFilePath);
+
+	/** Return this provider's AssetRegistry.bin filepath that can be deserialized from disk to load the AssetDatas. */
+	FStringView GetFilePath() const;
+	/**
+	 * Report whether this provider is global AssetRegistry's premade AssetRegistry. That AssetRegistry is notable
+	 * because it is the first and the largest, and we have special handling to manage it performantly. It can not be
+	 * unregistered.
+	 */
+	bool IsInitialProvider() const;
+	void SetIsInitialProvider(bool bValue);
+
+	/**
+	 * Return a weak pointer to the opaque FixedTagPrivate::FStore that was created or referenced by
+	 * FAssetRegistryState::Load when this provider was last deserialized from disk. FStore holds deduplicated
+	 * FAssetData.TagsAndValues data for all AssetDatas in the AssetRegistryState. We record this FStore and pass it
+	 * into the FAssetRegistryState::Load when we reload, which then reuses it rather than creating a new one. Other
+	 * than passing it into FAssetRegistryState::Load, we do not interpret it. The pointer is weak and will be
+	 * set to null if all AssetDatas using it are dismounted, in which case it will be recreated during Load.
+	 */
+	const FixedTagPrivate::FWeakStorePtr& GetStore() const;
+	void SetStore(FixedTagPrivate::FWeakStorePtr InStore);
+	/**
+	 * A count of how many FMountPoints have been requested to reload and not yet reloaded. Used to manage the
+	 * async task that loads them. This value contractually equals the number of FMountPointDatas associated with this
+	 * AssetProvider that are in the loading state.
+	 */
+	uint32& GetLoadRefCountPtr();
+	/**
+	 * Return a mutable reference to the LoadTask; we create it on demand if not yet created and it will keep running
+	 * until it finds no more work to do and exits. When the task exits it clears this reference to itself.
+	 */
+	UE::Tasks::FTask& GetLoadTask();
+
+	/**
+	 * Whether this AssetProvider is shutting down, which happens when it is unregistered by the game project
+	 * to save on memory, or when the ActiveMountsLoader is shutting down. This field is used to block new activity while
+	 * we wait for async tasks to shutdown.
+	 */
+	bool IsShutdown() const;
+	void SetIsShutdown(bool bValue);
+	/**
+	 * Information discovered during parsing of this provider's AssetRegistryState: true if there are
+	 * AssetDatas in this provider for which we could not find the FPackageName::IMountPoint and therefore cannot
+	 * unload/reload.
+	 */
+	bool HasMissingMountPoints() const;
+	void SetHasMissingMountPoints(bool bValue);
+	/**
+	 * Whether this AssetProvider encountered an error on load of its AssetRegistryState and therefore we
+	 * should not unload AssetDatas in any of its MountPoints.
+	 */
+	bool IsErrorNoUnload() const;
+	void SetIsErrorNoUnload(bool bValue);
+
+private:
+	const FString FilePath;
+	UE::Tasks::FTask LoadTask;
+	FixedTagPrivate::FWeakStorePtr Store;
+	uint32 LoadRefCount = 0;
+	bool bInitialProvider : 1 = false;
+	bool bShutdown : 1 = false;
+	bool bHasMissingMountPoints : 1 = false;
+	bool bErrorNoUnload : 1 = false;
+};
+
+/**
+ * A wrapper structure for FAssetProviders in a TSet that allow them to be looked up by FilePath. This structure
+ * also maintains the TUniquePtr for the FAssetProvider, which need to be at a fixed pointer over their lifetime
+ * and therefore must be TUniquePtr rather than direct elements in the TSet.
+ */
+class FAssetProviderSetKey
+{
+public:
+	explicit FAssetProviderSetKey(TUniquePtr<FAssetProvider>&& InAssetProvider);
+
+	bool operator==(const FAssetProviderSetKey& Other) const;
+	bool operator==(FStringView OtherProviderFilePath) const;
+
+	FAssetProvider& GetAssetProvider() const;
+
+private:
+	TUniquePtr<FAssetProvider> AssetProvider;
+};
+extern uint32 GetTypeHash(const FAssetProviderSetKey& Key);
+
+/**
+ * Data held by the ActiveMountsLoader about a MountPoint that contains AssetDatas in one of its registered
+ * AssetProviders. This data includes the associated AssetProvider for the MountPoint and the state machine data for
+ * the unload/reload of the assets in the MountPoint.
+ */
+class FMountPointData
+{
+public:
+	explicit FMountPointData(TRefCountPtr<UE::PackageName::IMountPoint> InMountPoint, EMountState InInitialState);
+
+	bool operator==(const FMountPointData& Other) const;
+	bool operator==(UE::PackageName::IMountPoint* MountPoint) const;
+
+	const TRefCountPtr<UE::PackageName::IMountPoint>& GetMountPoint() const;
+	FAssetProvider* GetAssetProvider() const;
+	void SetAssetProvider(FAssetProvider* InAssetProvider);
+	/**
+	 * The current state of the MountPoints unload/reload. This only changes when a request changes it into a transition
+	 * state or when a completed unload/reload changes it back to a stable state.
+	 */
+	EMountState GetCurrentState() const;
+	void SetCurrentState(EMountState InState);
+	/**
+	 * The currently requested unloaded or loaded state of the MountPoint. We track this requested state so that we
+	 * can handle new requests that override old requests that are still in progress.
+	 */
+	EMountState GetTargetState() const;
+	void SetTargetState(EMountState InState);
+
+	/**
+	 * Whether this MountPoint encountered an error on load of (any of) its associated AssetProviders and therefore we
+	 * should not unload AssetDatas under it. Notably, being associated with multiple AssetProviders is an error
+	 * itself; @see HasMultipleProviders.
+	 */
+	bool IsErrorNoUnload() const;
+	void SetIsErrorNoUnload(bool bValue);
+	/**
+	 * Whether this MountPoint has a multisegment IMountPoint->GetLongPackageName(), such as /Game/Weapons/, and so
+	 * requires more expensive calculations to determine whether a given AssetData LongPackageName is under it.
+	 */
+	bool IsMultiSegmentName() const;
+	void SetIsMultiSegmentName(bool bValue);
+	/**
+	 * Whether this MountPoint is a either a child or parent of another existing IMountPoint (e.g. /Game/Weapons/ is
+	 * a child of /Game/, and /Game/ is a parent of /Game/Weapons/) and so we don't currently support unload AssetDatas
+	 * under it.
+	 */
+	bool IsParentOrChild() const;
+	void SetIsParentOrChild(bool bValue);
+	/**
+	 * Whether AssetDatas under this MountPoint exist in multiple AssetProviders, and therefore we should not unload
+	 * AssetDatas under it. We do not currently support MountPoints in multiple AssetProviders because we want to save
+	 * memory by avoiding tracking multiple providers.
+	 */
+	bool HasMultipleProviders() const;
+	void SetHasMultipleProviders(bool bValue);
+
+private:
+	TRefCountPtr<UE::PackageName::IMountPoint> MountPoint;
+	FAssetProvider* AssetProvider = nullptr;
+	uint32 ErrorNoUnloadBit : 1 = false;
+	uint32 MultiSegmentNameBit : 1 = false;
+	uint32 IsParentOrChildBit : 1 = false;
+	uint32 HasMultipleProvidersBit : 1 = false;
+	uint32 CurrentStateBits : static_cast<uint8>(EMountState::NumBits) = static_cast<uint32>(EMountState::Unloaded);
+	uint32 TargetStateBits : static_cast<uint8>(EMountState::NumBits) = static_cast<uint32>(EMountState::Unloaded);
+};
+extern uint32 GetTypeHash(const FMountPointData& Key);
+
+/**
+ * Storage of an OnComplete function that was passed into LoadAsync or UnloadAsync. We hold this storage until
+ * all requests in the batch have completed and we can therefore call the OnComplete function. This struct is updated
+ * every time we complete a request for one of its MountPoints.
+ */
+struct FMountLoadEventSubscriber : public FRefCountedObject
+{
+	TFunction<void(bool bCanceled)> OnComplete;
+	// We don't store the list of remaining MountPoints to save memory/time on storing it, and we
+	// can do that but still be able to debug if it ever gets stuck at non-zero value because that
+	// information already is present (stored in the other direction) in FActiveMountsLoader.Subscribers.
+	// If RemainingCount > 0, there should be a pointer to this somewhere in LoadSubscribers or UnloadSubscribers.
+	int32 RemainingCount = 0;
+	bool bWasCanceled = false;
+	/**
+	 * True if listening to loads (and referenced from LoadSubscribers), false if listening to unloads
+	 * (and referenced from UnloadSubscribers).
+	 */
+	bool bIsLoadSubscriber = true;
+};
+
+/**
+ * The implementation of ActiveMountsLoading. This implementation class is created only when the feature is enabled
+ * in the current process. It holds data about how to reload each mountpoint and manages their state in an
+ * unload/reload state machine.
+ */
+class FActiveMountsLoader
+{
+public:
+	// Documentation for public functions on this class is on FActiveMountsLoaderWrapper functions of the same name
+	// (with the possible addition of Start and Finish prefixes for functions that require WaitOnAsync outside the
+	// lock). @see FActiveMountsLoaderWrapper.
+
+	/**
+	 * Called by FActiveMountsLoaderWrapper::Shutdown. Signals threads to shutdown and reports all EventRefs that
+	 * must be waited on for those threads to terminate. Followed by a call to FinishShutdown after waiting.
+	 */
+	void StartShutdown(UAssetRegistryImpl& UARI, FActiveMountEvents& EventContext,
+		TArray<UE::Tasks::FTask>& OutWaitOnAsync);
+	/**
+	 * Called by FActiveMountsLoaderWrapper::Shutdown after waiting on EventRefs from StartShutdown. Clears all
+	 * internal data, FActiveMountsLoader can now be deleted.
+	 */
+	void FinishShutdown(UAssetRegistryImpl& UARI, FActiveMountEvents& EventContext);
+
+	void LoadAsync(UAssetRegistryImpl& UARI, FActiveMountEvents& EventContext,
+		TConstArrayView<TRefCountPtr<UE::PackageName::IMountPoint>> InMountPoints,
+		TFunction<void(bool bCanceled)>&& OnComplete);
+	void UnloadAsync(UAssetRegistryImpl& UARI, FActiveMountEvents& EventContext,
+		TConstArrayView<TRefCountPtr<UE::PackageName::IMountPoint>> InMountPoints,
+		TFunction<void(bool bCanceled)>&& OnComplete);
+	void UnloadAllUnmountedAsync(UAssetRegistryImpl& UARI, FActiveMountEvents& EventContext,
+		TFunction<void(bool bCanceled)>&& OnComplete);
+	bool IsLoaded(UAssetRegistryImpl& UARI, FActiveMountEvents& EventContext,
+		const TRefCountPtr<UE::PackageName::IMountPoint>& MountPoint);
+
+	void OnActiveMountsInitialize(UAssetRegistryImpl& UARI, FActiveMountEvents& EventContext,
+		FInterfaceWriteScopeLock& ProofOfWriteLock);
+	bool TryRegisterAssetProviderAndAppendStateOutsideARLock(UAssetRegistryImpl& UARI, FEventContext& AREventContext,
+		FActiveMountEvents& EventContext, FStringView StateFilePath);
+	bool TryLoadStateAndRegisterAssetProvider(UAssetRegistryImpl& UARI, FActiveMountEvents& EventContext,
+		FInterfaceWriteScopeLock& ProofOfWriteLock, FPakPlatformFile* PakPlatformFile, FStringView StateFilePath,
+		FAssetRegistryState& OutLoadedState);
+	/**
+	 * Called by FActiveMountsLoaderWrapper::UnregisterAssetProvider. Signals the AssetProvider's load thread to
+	 * shutdown if it exists and reports the EventRefs that must be waited on for it to terminate.
+	 * Followed by a call to FinishUnregisterAssetProvider after waiting. Returns false if the AssetProvider does
+	 * not exist; FinishUnregisterAssetProvider should not be called in that case.
+	 */
+	bool TryStartUnregisterAssetProvider(UAssetRegistryImpl& UARI, FActiveMountEvents& EventContext,
+		TArray<UE::Tasks::FTask>& OutWaitOnAsync, FStringView StateFilePath);
+	/**
+	 * Called by FActiveMountsLoaderWrapper::UnregisterAssetProvider after waiting on EventRefs from StartShutdown.
+	 * Clears the AssetProvider's data and the data of any MountPoints it owned. 
+	 */
+	void FinishUnregisterAssetProvider(UAssetRegistryImpl& UARI, FActiveMountEvents& EventContext,
+		FStringView StateFilePath);
+
+	void InitialProviderReport(UAssetRegistryImpl& UARI, FActiveMountEvents& EventContext,
+		FAssetProviderInitData& StateInitData);
+	void InitialProviderOnReadyForRegister(UAssetRegistryImpl& UARI, FActiveMountEvents& EventContext,
+		EInitialProviderRegisterReason Reason);
+	void InitialProviderConditionalRegister(UAssetRegistryImpl& UARI, FActiveMountEvents& EventContext,
+		FAssetRegistryState& AppendedState);
+
+	void OnMountPointsMounted(UAssetRegistryImpl& UARI, FActiveMountEvents& EventContext,
+		TConstArrayView<TRefCountPtr<UE::PackageName::IMountPoint>> InMountPoints);
+	void OnMountPointsDismounted(UAssetRegistryImpl& UARI, FActiveMountEvents& EventContext,
+		FInterfaceWriteScopeLock& ProofOfWriteLock,
+		TConstArrayView<TRefCountPtr<UE::PackageName::IMountPoint>> InMountPoints);
+
+private:
+	/**
+	 * Inspect all of the AssetDatas in the AssetRegistryState loaded by an AssetProvider to find which MountPoints
+	 * should be associated with it.
+	 */
+	void ParseInitialState(FAssetProvider& AssetProvider, const FAssetRegistryState& State);
+	FMountPointData& FindOrAddMountPointData(UE::PackageName::IMountPoint* MountPoint, EMountState InitialState);
+	FMountPointData* FindMountPointData(UE::PackageName::IMountPoint* MountPoint);
+	void NotifyMountPointLoadComplete(FActiveMountEvents& EventContext,
+		const TRefCountPtr<UE::PackageName::IMountPoint>& MountPoint, bool bCanceled);
+	void NotifyMountPointUnloadComplete(FActiveMountEvents& EventContext,
+		const TRefCountPtr<UE::PackageName::IMountPoint>& MountPoint, bool bCanceled);
+	void AddRefMountPointLoad(UAssetRegistryImpl& UARI, FMountPointData* MountPointData);
+	void ReleaseMountPointLoad(UAssetRegistryImpl& UARI, FMountPointData* MountPointData);
+	void AddRefMountPointUnload(UAssetRegistryImpl& UARI, FMountPointData* MountPointData);
+	void ReleaseMountPointUnload(UAssetRegistryImpl& UARI, FMountPointData* MountPointData);
+	/** The DoTask function for an AssetProvider's LoadTask for any requested MountPoint loads in it. */
+	void RunAssetProviderLoadTask(UAssetRegistryImpl& UARI, FAssetProvider& AssetProvider);
+	/** The DoTask function for the UnloadTask for any requested MountPoint unloads. */
+	void RunMountPointUnloadTask(UAssetRegistryImpl& UARI);
+
+	TSet<FAssetProviderSetKey> AssetProviders;
+	TSet<FMountPointData> MountPoints;
+	/** The list of OnComplete functions that are waiting on a MountPoint to load. */
+	TMap<TRefCountPtr<UE::PackageName::IMountPoint>, TArray<TRefCountPtr<FMountLoadEventSubscriber>>> LoadSubscribers;
+	/** The list of OnComplete functions that are waiting on a MountPoint to unload. */
+	TMap<TRefCountPtr<UE::PackageName::IMountPoint>, TArray<TRefCountPtr<FMountLoadEventSubscriber>>> UnloadSubscribers;
+	/**
+	 * An additional lock entered into by registration and unregistration functions. Unregistration functions have to
+	 * wait outside of the usual ActiveMountsLoader lock on FActiveMountsLoaderWrapper::ActiveMountsLock because the
+	 * tasks they wait on need to enter that lock. But we need to force only a single register/unregister to happen at
+	 * once, so we have this lock. This lock must always be entered before entering the ActiveMountsLock. That creates
+	 * a complex detail since it is possible the FActiveMountsLoader pointer holding this RegistrationLock is deleted
+	 * during shutdown while we are about to try entering the RegistrationLock. We handle that detail by storing a
+	 * shared pointer to the lock that might outlive the FActiveMountsLoader that created it.
+	 */
+	TSharedPtr<UE::FRecursiveMutex> RegistrationLock;
+	/**
+	 * The UnloadTask is created on demand if not yet created and it will keep running until it finds no more work to
+	 * do and exits. When the task exits it clears this reference to itself.
+	 */
+	UE::Tasks::FTask UnloadTask;
+	/**
+	 * A count of how many FMountPoints have been requested to unload and not yet unloaded. Used to manage the
+	 * async task that unloads them. This value contractually equals the number of FMountPointDatas that are in the
+	 * unloading state.
+	 */
+	uint32 MountPointUnloadRefCount = 0;
+	/**
+	 * Configuration-specified variable for whether we should automatically request unload/reload when
+	 * FPackageName::IMountPoints are dismounted/mounted.
+	 */
+	bool bAutoUnload : 1 = false;
+	/** True only after ActiveMountsInitialize is called and so we are allowed to service Unload/Load requests. */
+	bool bLoaderIsActive : 1 = false;
+	/** InitialProviderOnReadyForRegister has been called. */
+	bool bInitialProviderReadyForRegister : 1 = false;
+	/** InitialProviderReport has been called. */
+	bool bInitialProviderReported : 1 = false;
+	/** Gate value to ensure Call_Once on the initial provider's registration. */
+	bool bInitialProviderRegistered : 1 = false;
+	/** Gate value to ensure Call_Once on bAutoUnload's trigger of UnloadAllUnmountedAsync. */
+	bool bAutoUnloadInitialUnloadTriggered : 1 = false;
+	/** Set to true when AssetRegistry is shutting down and we are waiting for async; blocks any external requests. */
+	bool bShutdown : 1 = false;
+
+	friend FActiveMountsLoaderWrapper;
+};
+
+} // namespace UE::AssetRegistry
\ No newline at end of file
diff --git a/Engine/Source/Runtime/AssetRegistry/Private/AssetRegistry.cpp b/Engine/Source/Runtime/AssetRegistry/Private/AssetRegistry.cpp
--- a/Engine/Source/Runtime/AssetRegistry/Private/AssetRegistry.cpp
+++ b/Engine/Source/Runtime/AssetRegistry/Private/AssetRegistry.cpp
@@ -374,7 +374,7 @@ void FAssetRegistryImpl::ConditionalLoadPremadeAssetRegistry(UE::AssetRegistry::
 }
 
 void FAssetRegistryImpl::ConsumeOrDeferPreloadedPremade(UAssetRegistryImpl& UARI, Impl::FEventContext& EventContext,
-	FInterfaceWriteScopeLock& ProofOfWriteLock)
+	Impl::FActiveMountEvents& ActiveMountEvents, FInterfaceWriteScopeLock& ProofOfWriteLock)
 {
 	// Called from inside WriteLock on InterfaceLock
 	using namespace UE::AssetRegistry::Premade;
@@ -389,23 +389,25 @@ void FAssetRegistryImpl::ConsumeOrDeferPreloadedPremade(UAssetRegistryImpl& UARI
 	if (Premade::CanLoadAsync())
 	{
 		FPreloader::FConsumeFunction ConsumeFromAsyncThread = [this, &UARI](Premade::ELoadResult LoadResult,
-			FAssetRegistryState&& ARState, FStringView StateFilePath)
+			UE::AssetRegistry::Impl::FAssetProviderInitData&& StateInitData)
 		{
 			Impl::FEventContext EventContext;
+			Impl::FActiveMountEvents ActiveMountEvents;
 			{
 				UE::AssetRegistry::FInterfaceWriteScopeLock InterfaceScopeLock(UARI.InterfaceLock);
-				AsyncConsumer.Consume(UARI, EventContext, InterfaceScopeLock, LoadResult,
-					MoveTemp(ARState), StateFilePath);
+				AsyncConsumer.Consume(UARI, EventContext, ActiveMountEvents, InterfaceScopeLock, LoadResult,
+					MoveTemp(StateInitData));
 			}
 			UARI.Broadcast(EventContext);
+			ActiveMountEvents.Broadcast();
 		};
 		auto ConsumeOnCurrentThread = [ConsumeFromAsyncThread](Premade::ELoadResult LoadResult,
-			FAssetRegistryState&& ARState, FStringView InStateFilePath) mutable
+			UE::AssetRegistry::Impl::FAssetProviderInitData&& StateInitData) mutable
 		{
-			Async(EAsyncExecution::TaskGraph, [LoadResult, ARState=MoveTemp(ARState),
-				ConsumeFromAsyncThread=MoveTemp(ConsumeFromAsyncThread), StateFilePath=FString(InStateFilePath)]() mutable
+			Async(EAsyncExecution::TaskGraph, [LoadResult, StateInitData=MoveTemp(StateInitData),
+				ConsumeFromAsyncThread=MoveTemp(ConsumeFromAsyncThread)]() mutable
 			{
-				ConsumeFromAsyncThread(LoadResult, MoveTemp(ARState), StateFilePath);
+				ConsumeFromAsyncThread(LoadResult, MoveTemp(StateInitData));
 			});
 		};
 
@@ -415,11 +417,11 @@ void FAssetRegistryImpl::ConsumeOrDeferPreloadedPremade(UAssetRegistryImpl& UARI
 	else
 	{
 		GPreloader.Consume(
-			[this, &UARI, &EventContext, &ProofOfWriteLock](Premade::ELoadResult LoadResult,
-			FAssetRegistryState&& ARState, FStringView StateFilePath)
+			[this, &UARI, &EventContext, &ProofOfWriteLock, &ActiveMountEvents](Premade::ELoadResult LoadResult,
+				UE::AssetRegistry::Impl::FAssetProviderInitData&& StateInitData)
 			{
-				LoadPremadeAssetRegistry(UARI, EventContext, ProofOfWriteLock, LoadResult,
-					MoveTemp(ARState), StateFilePath);
+				LoadPremadeAssetRegistry(UARI, EventContext, ActiveMountEvents, ProofOfWriteLock, LoadResult,
+					MoveTemp(StateInitData));
 			});
 	}
 }
@@ -497,6 +499,7 @@ struct FInitializeContext
 {
 	UAssetRegistryImpl& UARI;
 	FEventContext Events;
+	FActiveMountEvents ActiveMountEvents;
 	FClassInheritanceContext InheritanceContext;
 	FClassInheritanceBuffer InheritanceBuffer;
 	TArray<FString> RootContentPaths;
@@ -531,6 +534,7 @@ UAssetRegistryImpl::UAssetRegistryImpl(const FObjectInitializer& ObjectInitializ
 		InitializeEvents(Context);
 	}
 	Broadcast(Context.Events);
+	Context.ActiveMountEvents.Broadcast();
 }
 
 bool UAssetRegistryImpl::IsPathBeautificationNeeded(const FString& InAssetPath) const
@@ -546,8 +550,8 @@ FAssetRegistryImpl::FAssetRegistryImpl()
 }
 
 void FAssetRegistryImpl::LoadPremadeAssetRegistry(UAssetRegistryImpl& UARI, Impl::FEventContext& EventContext,
-	FInterfaceWriteScopeLock& ProofOfWriteLock,
-	Premade::ELoadResult LoadResult, FAssetRegistryState&& ARState, FStringView StateFilePath)
+	Impl::FActiveMountEvents& ActiveMountEvents, FInterfaceWriteScopeLock& ProofOfWriteLock,
+	Premade::ELoadResult LoadResult, Impl::FAssetProviderInitData&& StateInitData)
 {
 	SCOPED_BOOT_TIMING("LoadPremadeAssetRegistry");
 	UE_SCOPED_ENGINE_ACTIVITY("Loading premade asset registry");
@@ -559,10 +563,16 @@ void FAssetRegistryImpl::LoadPremadeAssetRegistry(UAssetRegistryImpl& UARI, Impl
 		SCOPED_BOOT_TIMING("LoadPremadeAssetRegistry_Main");
 		if (LoadResult == Premade::ELoadResult::Succeeded)
 		{
+			UARI.ActiveMountsLoader.InitialProviderReport(UARI, ActiveMountEvents, StateInitData);
+			
 			if (State.GetNumAssets() == 0)
 			{
+				// ActiveMounts: Handle the possibility that ActiveMountsInitialize has already been called and if so
+				// parse the premade's state now.
+				UARI.ActiveMountsLoader.InitialProviderConditionalRegister(UARI, ActiveMountEvents, StateInitData.ARState);
+
 				// Add the premade state into the global AssetRegistry
-				State = MoveTemp(ARState);
+				State = MoveTemp(StateInitData.ARState);
 				CachePathsFromState(EventContext, State);
 				if (bEmitAssetEvents)
 				{
@@ -579,12 +589,18 @@ void FAssetRegistryImpl::LoadPremadeAssetRegistry(UAssetRegistryImpl& UARI, Impl
 				// This is unexpected, we should not have called AppendState or in any other way added Assets to the
 				// AssetRegistry before loading the premade AssetRegistry. However, we can handle it with only minor
 				// performance issues.
+				// ActiveMounts usually parses the premade's ARState during ActiveMountsInitialize by reading the
+				// global AssetRegistry's GuardedData.State, but that state already contains some other data, so we
+				// need to instead parse the premade's state directly.
+				UARI.ActiveMountsLoader.InitialProviderOnReadyForRegister(UARI, ActiveMountEvents,
+					Impl::EInitialProviderRegisterReason::EarlierState);
+				UARI.ActiveMountsLoader.InitialProviderConditionalRegister(UARI, ActiveMountEvents, StateInitData.ARState);
 
 				// Add the premade state into the global AssetRegistry
-				if (State.GetNumAssets() < ARState.GetNumAssets())
+				if (State.GetNumAssets() < StateInitData.ARState.GetNumAssets())
 				{
 					FAssetRegistryState ExistingState = MoveTemp(State);
-					State = MoveTemp(ARState);
+					State = MoveTemp(StateInitData.ARState);
 					CachePathsFromState(EventContext, State);
 					if (bEmitAssetEvents)
 					{
@@ -599,7 +615,8 @@ void FAssetRegistryImpl::LoadPremadeAssetRegistry(UAssetRegistryImpl& UARI, Impl
 				}
 				else
 				{
-					AppendState(EventContext, ARState, FAssetRegistryState::EInitializationMode::OnlyUpdateNew, bEmitAssetEvents);
+					AppendState(EventContext, StateInitData.ARState,
+						FAssetRegistryState::EInitializationMode::OnlyUpdateNew, bEmitAssetEvents);
 				}
 			}
 			UpdatePersistentMountPoints();
@@ -620,17 +637,12 @@ void FAssetRegistryImpl::LoadPremadeAssetRegistry(UAssetRegistryImpl& UARI, Impl
 		{
 			if (ContentPlugin->CanContainContent())
 			{
-				FArrayReader SerializedAssetData;
+				FAssetRegistryState PluginState;
 				FString PluginAssetRegistry = ContentPlugin->GetBaseDir() / TEXT("AssetRegistry.bin");
 
-				// Optimization: if we're using pak files then only search paks (avoid unnecessary fallback to loose)
-				bool bFileExists = PakPlatformFile != nullptr ? PakPlatformFile->FindFileInPakFiles(*PluginAssetRegistry) : IFileManager::Get().FileExists(*PluginAssetRegistry);
-				if (bFileExists && FFileHelper::LoadFileToArray(SerializedAssetData, *PluginAssetRegistry))
+				if (UARI.ActiveMountsLoader.TryLoadStateAndRegisterAssetProvider(UARI, ActiveMountEvents,
+					ProofOfWriteLock, PakPlatformFile, PluginAssetRegistry, PluginState))
 				{
-					SerializedAssetData.Seek(0);
-					FAssetRegistryState PluginState;
-					PluginState.Load(SerializedAssetData);
-
 #if WITH_EDITOR
 					/*
 					 * Only update the new assets when using a premade asset registry in editor.
@@ -749,7 +761,9 @@ void FAssetRegistryImpl::Initialize(UE::AssetRegistry::FInterfaceWriteScopeLock&
 		}
 	}
 
-	ConsumeOrDeferPreloadedPremade(Context.UARI, Context.Events, ScopeLock);
+	Context.UARI.ActiveMountsLoader.InitializeConfig();
+
+	ConsumeOrDeferPreloadedPremade(Context.UARI, Context.Events, Context.ActiveMountEvents, ScopeLock);
 
 	// Report startup time. This does not include DirectoryWatcher startup time.
 	double StartupDuration = FPlatformTime::Seconds() - StartupStartTime;
@@ -1704,6 +1718,9 @@ void UAssetRegistryImpl::OnPreExit()
 	// Now that we are no longer holding the lock, we can destroy the gatherer
 	GlobalGatherer.Reset();
 
+	UE::AssetRegistry::Impl::FActiveMountEvents ActiveMountEvents;
+	ActiveMountsLoader.Shutdown(*this, ActiveMountEvents);
+	ActiveMountEvents.Broadcast();
 }
 
 void UAssetRegistryImpl::FinishDestroy()
@@ -6039,6 +6056,135 @@ void UAssetRegistryImpl::AppendState(const FAssetRegistryState& InState, UE::Ass
 	Broadcast(EventContext);
 }
 
+bool UAssetRegistryImpl::ActiveMountsIsEnabled()
+{
+	return ActiveMountsLoader.IsEnabled();
+}
+
+void UAssetRegistryImpl::ActiveMountsInitialize()
+{
+	using namespace UE::AssetRegistry;
+
+	UE::AssetRegistry::Impl::FActiveMountEvents ActiveMountEvents;
+	{
+		LLM_SCOPE(ELLMTag::AssetRegistry);
+		UE::AssetRegistry::FInterfaceWriteScopeLock InterfaceScopeLock(InterfaceLock);
+		ActiveMountsLoader.OnActiveMountsInitialize(*this, ActiveMountEvents, InterfaceScopeLock);
+	}
+	ActiveMountEvents.Broadcast();
+}
+
+bool UAssetRegistryImpl::ActiveMountsRegisterAndLoadFilePathSynchronous(FStringView StateFilePath)
+{
+	if (StateFilePath.IsEmpty())
+	{
+		ensureMsgf(false, TEXT("AppendStateForActiveMounts called with an invalid empty StateFilePath."));
+		return false;
+	}
+
+	UE::AssetRegistry::Impl::FEventContext EventContext;
+	UE::AssetRegistry::Impl::FActiveMountEvents ActiveMountEvents;
+	LLM_SCOPE(ELLMTag::AssetRegistry);
+	bool bResult = ActiveMountsLoader.TryRegisterAssetProviderAndAppendStateOutsideARLock(*this, EventContext,
+		ActiveMountEvents, StateFilePath);
+	Broadcast(EventContext);
+	ActiveMountEvents.Broadcast();
+	return bResult;
+}
+
+void UAssetRegistryImpl::ActiveMountsRegisterAndLoadFilePathAsync(FStringView StateFilePath,
+	TFunction<void(bool bSucceeded)> OnComplete, UE::Tasks::ETaskPriority TaskPriority)
+{
+	if (StateFilePath.IsEmpty())
+	{
+		ensureMsgf(false, TEXT("AppendStateForActiveMounts called with an invalid empty StateFilePath."));
+		if (OnComplete)
+		{
+			OnComplete(false);
+		}
+		return;
+	}
+
+	UE::Tasks::Launch(UE_SOURCE_LOCATION,
+		[StateFilePath = FString(StateFilePath), OnComplete = MoveTemp(OnComplete)]
+		{
+			bool bSuccess = false;
+			IAssetRegistry* GlobalAR = IAssetRegistry::Get();
+			if (GlobalAR)
+			{
+				bSuccess = GlobalAR->ActiveMountsRegisterAndLoadFilePathSynchronous(StateFilePath);
+			}
+			if (OnComplete)
+			{
+				OnComplete(bSuccess);
+			}
+		}, TaskPriority);
+}
+
+void UAssetRegistryImpl::ActiveMountsLoadAsync(TConstArrayView<FStringView> MountPointLongPackageNames,
+	TFunction<void(bool bCanceled)> OnComplete)
+{
+	using namespace UE::AssetRegistry;
+
+	LLM_SCOPE(ELLMTag::AssetRegistry);
+	TArray<TRefCountPtr<UE::PackageName::IMountPoint>> MountPoints;
+	MountPoints.Reserve(MountPointLongPackageNames.Num());
+	for (FStringView MountPointName : MountPointLongPackageNames)
+	{
+		TRefCountPtr<UE::PackageName::IMountPoint> MountPoint =
+			FPackageName::FindMountPointByRootPackageName(MountPointName);
+		if (MountPoint)
+		{
+			MountPoints.Add(MoveTemp(MountPoint));
+		}
+	}
+
+	UE::AssetRegistry::Impl::FActiveMountEvents ActiveMountEvents;
+	ActiveMountsLoader.LoadAsync(*this, ActiveMountEvents, MountPoints, MoveTemp(OnComplete));
+	ActiveMountEvents.Broadcast();
+}
+
+void UAssetRegistryImpl::ActiveMountsUnloadAsync(TConstArrayView<FStringView> MountPointLongPackageNames,
+	TFunction<void(bool bCanceled)> OnComplete)
+{
+	using namespace UE::AssetRegistry;
+
+	LLM_SCOPE(ELLMTag::AssetRegistry);
+	TArray<TRefCountPtr<UE::PackageName::IMountPoint>> MountPoints;
+	MountPoints.Reserve(MountPointLongPackageNames.Num());
+	for (FStringView MountPointName : MountPointLongPackageNames)
+	{
+		TRefCountPtr<UE::PackageName::IMountPoint> MountPoint =
+			FPackageName::FindMountPointByRootPackageName(MountPointName);
+		if (MountPoint)
+		{
+			MountPoints.Add(MoveTemp(MountPoint));
+		}
+	}
+
+	UE::AssetRegistry::Impl::FActiveMountEvents ActiveMountEvents;
+	ActiveMountsLoader.UnloadAsync(*this, ActiveMountEvents, MountPoints, MoveTemp(OnComplete));
+	ActiveMountEvents.Broadcast();
+}
+
+bool UAssetRegistryImpl::ActiveMountIsLoaded(FStringView MountPointLongPackageName)
+{
+	using namespace UE::AssetRegistry;
+
+	TRefCountPtr<UE::PackageName::IMountPoint> MountPoint =
+		FPackageName::FindMountPointByRootPackageName(MountPointLongPackageName);
+	if (!MountPoint)
+	{
+		return false;
+	}
+
+	LLM_SCOPE(ELLMTag::AssetRegistry);
+	UE::AssetRegistry::Impl::FActiveMountEvents ActiveMountEvents;
+	bool bResult = ActiveMountsLoader.IsLoaded(*this, ActiveMountEvents, MountPoint);
+	ActiveMountEvents.Broadcast();
+	return bResult;
+}
+
 namespace UE::AssetRegistry
 {
 
@@ -6062,7 +6208,50 @@ void FAssetRegistryImpl::AppendState(Impl::FEventContext& EventContext, const FA
 		Mode,
 		AppendResultPtr);
 
-	CachePathsFromState(EventContext, InState);
+	CachePathsFromState(EventContext, InState, &Options);
+	CachedPathTree.Shrink();
+
+	if (AppendResultPtr)
+	{
+		for (const FAssetData* AssetData : AppendResultPtr->AddedAssets)
+		{
+			EventContext.AssetEvents.Emplace(*AssetData, UE::AssetRegistry::Impl::FEventContext::EEvent::Added);
+		}
+		for (const FAssetData* AssetData : AppendResultPtr->UpdatedAssets)
+		{
+			EventContext.AssetEvents.Emplace(*AssetData, UE::AssetRegistry::Impl::FEventContext::EEvent::Updated);
+		}
+	}
+}
+
+void FAssetRegistryImpl::ConsumeState(Impl::FEventContext& EventContext, FAssetRegistryState&& InState,
+	bool bEmitAssetEvents)
+{
+	FAssetRegistryAppendResult LocalAppendResult;
+	FAssetRegistryAppendResult* AppendResultPtr = bEmitAssetEvents ? &LocalAppendResult : nullptr;
+
+#if WITH_EDITOR
+	const FAssetRegistrySerializationOptions& InOptions = DevelopmentSerializationOptions;
+	//Development options have all filters disabled, so there is no need to resolve runtime class filter options
+#else
+	FAssetRegistrySerializationOptions& InOptions = SerializationOptions;
+	const bool bSuppressWarnings = !UE::AssetRegistry::Utils::HasEngineModuleLoaded();
+	Utils::UpdateSerializationOptions(InOptions, CachedParsedFilterRules, bSuppressWarnings);
+#endif
+
+	// Since we're moving assets rather than copying them, we do not filter them by tags and we do not remove any
+	// any entries from their tags.
+	FAssetRegistrySerializationOptions ConsumeOptions = InOptions;
+	ConsumeOptions.bFilterAssetDataWithNoTags = false;
+	ConsumeOptions.CookFilterlistTagsByClass.Empty();
+
+	CachePathsFromState(EventContext, InState, &ConsumeOptions);
+
+	State.ConsumeExisting(
+		MoveTemp(InState),
+		ConsumeOptions,
+		AppendResultPtr);
+
 	CachedPathTree.Shrink();
 
 	if (AppendResultPtr)
@@ -6078,18 +6267,23 @@ void FAssetRegistryImpl::AppendState(Impl::FEventContext& EventContext, const FA
 	}
 }
 
-void FAssetRegistryImpl::CachePathsFromState(Impl::FEventContext& EventContext, const FAssetRegistryState& InState)
+void FAssetRegistryImpl::CachePathsFromState(Impl::FEventContext& EventContext, const FAssetRegistryState& InState,
+	const FAssetRegistrySerializationOptions* InOptions)
 {
 	SCOPED_BOOT_TIMING("FAssetRegistryImpl::CachePathsFromState");
 
 	// Refreshes ClassGeneratorNames if out of date due to module load
 	CollectCodeGeneratorClasses();
 
+	if (!InOptions)
+	{
 #if WITH_EDITOR
-	const FAssetRegistrySerializationOptions& Options = DevelopmentSerializationOptions;
+		InOptions = &DevelopmentSerializationOptions;
 #else
-	const FAssetRegistrySerializationOptions& Options = SerializationOptions;
+		InOptions = &SerializationOptions;
 #endif
+	}
+	const FAssetRegistrySerializationOptions& Options = *InOptions;
 
 	// Add paths to cache
 	InState.EnumerateAllAssets([this, &Options, &EventContext](const FAssetData& AssetData)
@@ -8807,13 +9001,14 @@ void UAssetRegistryImpl::OnMultipleMountPointsMounted(TConstArrayView<TRefCountP
 #endif
 
 	UE::AssetRegistry::Impl::FEventContext EventContext;
+	UE::AssetRegistry::Impl::FActiveMountEvents ActiveMountEvents;
 	{
 		LLM_SCOPE(ELLMTag::AssetRegistry);
 		UE::AssetRegistry::FInterfaceWriteScopeLock InterfaceScopeLock(InterfaceLock);
 		UE::AssetRegistry::Impl::FClassInheritanceContext InheritanceContext;
 		UE::AssetRegistry::Impl::FClassInheritanceBuffer InheritanceBuffer;
 		GetInheritanceContextWithRequiredLock(InterfaceScopeLock, InheritanceContext, InheritanceBuffer);
-		GuardedData.OnMultipleMountPointsMounted(*this, EventContext,
+		GuardedData.OnMultipleMountPointsMounted(*this, EventContext, ActiveMountEvents,
 			InheritanceContext, InterfaceScopeLock, MountedMountPoints);
 
 		// Listen for directory changes in this content path
@@ -8842,12 +9037,14 @@ void UAssetRegistryImpl::OnMultipleMountPointsMounted(TConstArrayView<TRefCountP
 	}
 
 	Broadcast(EventContext);
+	ActiveMountEvents.Broadcast();
 }
 
 namespace UE::AssetRegistry
 {
 
 void FAssetRegistryImpl::OnMultipleMountPointsMounted(UAssetRegistryImpl& UARI, Impl::FEventContext& EventContext,
+	Impl::FActiveMountEvents& ActiveMountEvents,
 	Impl::FClassInheritanceContext& InheritanceContext,
 	UE::AssetRegistry::FInterfaceWriteScopeLock& ScopeLock, 
 	TConstArrayView<TRefCountPtr<UE::PackageName::IMountPoint>> MountedMountPoints)
@@ -8885,6 +9082,8 @@ void FAssetRegistryImpl::OnMultipleMountPointsMounted(UAssetRegistryImpl& UARI,
 			}
 		}
 	}
+
+	UARI.ActiveMountsLoader.OnMountPointsMounted(UARI, ActiveMountEvents, MountedMountPoints);
 }
 
 }
@@ -8914,11 +9113,12 @@ void UAssetRegistryImpl::OnMultipleMountPointsDismounted(TConstArrayView<TRefCou
 #endif
 
 	UE::AssetRegistry::Impl::FEventContext EventContext;
+	UE::AssetRegistry::Impl::FActiveMountEvents ActiveMountEvents;
 	{
 		LLM_SCOPE(ELLMTag::AssetRegistry);
 		UE::AssetRegistry::FInterfaceWriteScopeLock InterfaceScopeLock(InterfaceLock);
-		GuardedData.OnMultipleMountPointsDismounted(*this, EventContext,
-			DismountedMountPoints);
+		GuardedData.OnMultipleMountPointsDismounted(*this, EventContext, ActiveMountEvents,
+			InterfaceScopeLock, DismountedMountPoints);
 
 		// Stop listening for directory changes in this content path
 #if WITH_EDITOR
@@ -8939,15 +9139,19 @@ void UAssetRegistryImpl::OnMultipleMountPointsDismounted(TConstArrayView<TRefCou
 #endif // WITH_EDITOR
 	}
 	Broadcast(EventContext);
+	ActiveMountEvents.Broadcast();
 }
 
 namespace UE::AssetRegistry
 {
 
 void FAssetRegistryImpl::OnMultipleMountPointsDismounted(UAssetRegistryImpl& UARI,
-	Impl::FEventContext& EventContext,
+	Impl::FEventContext& EventContext, Impl::FActiveMountEvents& ActiveMountEvents,
+	UE::AssetRegistry::FInterfaceWriteScopeLock& ProofOfWriteLock,
 	TConstArrayView<TRefCountPtr<UE::PackageName::IMountPoint>> DismountedMountPoints)
 {
+	UARI.ActiveMountsLoader.OnMountPointsDismounted(UARI, ActiveMountEvents, ProofOfWriteLock, DismountedMountPoints);
+
 	// For the GlobalGatherer we remove all input MountPoints; we don't need to keep information about their gather.
 	// But for the FAssetData stored inside our AssetRegistryState, for some mount points
 	// we are configured to keep those assets in memory even when dismounted.
@@ -9297,7 +9501,10 @@ void UAssetRegistryImpl::ResumePathsMounting()
 			FScopeLock Guard(&MountingCriticalSection);
 			LocalList = MoveTemp(ListOfMountedMountPoints);
 		}
-		OnMultipleMountPointsMounted(LocalList);
+		if (!LocalList.IsEmpty())
+		{
+			OnMultipleMountPointsMounted(LocalList);
+		}
 	}
 }
 
@@ -9315,7 +9522,10 @@ void UAssetRegistryImpl::ResumePathsDismounting()
 			FScopeLock Guard(&MountingCriticalSection);
 			LocalList = MoveTemp(ListOfDismountedMountPoints);
 		}
-		OnMultipleMountPointsDismounted(LocalList);
+		if (!LocalList.IsEmpty())
+		{
+			OnMultipleMountPointsDismounted(LocalList);
+		}
 	}
 }
 
diff --git a/Engine/Source/Runtime/AssetRegistry/Private/AssetRegistry.h b/Engine/Source/Runtime/AssetRegistry/Private/AssetRegistry.h
--- a/Engine/Source/Runtime/AssetRegistry/Private/AssetRegistry.h
+++ b/Engine/Source/Runtime/AssetRegistry/Private/AssetRegistry.h
@@ -2,6 +2,7 @@
 
 #pragma once
 
+#include "ARLoader.h"
 #include "AssetRegistryImpl.h"
 #include "Containers/Ticker.h"
 #include "Misc/PackageName.h"
@@ -10,6 +11,7 @@
 #include "AssetRegistry.generated.h"
 
 class FAssetDataGatherer;
+namespace UE::AssetRegistry::Impl { class FActiveMountsLoader; }
 namespace UE::AssetRegistry::Premade { struct FAsyncConsumer; }
 
 /**
@@ -114,6 +116,17 @@ public:
 	virtual void Serialize(FArchive& Ar) override;
 	virtual void AppendState(const FAssetRegistryState& InState,
 		UE::AssetRegistry::EAppendMode AppendMode = UE::AssetRegistry::EAppendMode::Append) override;
+	virtual bool ActiveMountsIsEnabled() override;
+	virtual void ActiveMountsInitialize() override;
+	virtual bool ActiveMountsRegisterAndLoadFilePathSynchronous(FStringView StateFilePath) override;
+	virtual void ActiveMountsRegisterAndLoadFilePathAsync(FStringView StateFilePath,
+		TFunction<void(bool bSucceeded)> OnComplete = TFunction<void(bool bSucceeded)>(),
+		UE::Tasks::ETaskPriority TaskPriority = UE::Tasks::ETaskPriority::Default) override;
+	virtual void ActiveMountsLoadAsync(TConstArrayView<FStringView> MountPointLongPackageNames,
+		TFunction<void(bool bCanceled)> OnComplete = TFunction<void(bool bCanceled)>()) override;
+	virtual void ActiveMountsUnloadAsync(TConstArrayView<FStringView> MountPointLongPackageNames,
+		TFunction<void(bool bCanceled)> OnComplete = TFunction<void(bool bCanceled)>()) override;
+	virtual bool ActiveMountIsLoaded(FStringView MountPointLongPackageName) override;
 	virtual SIZE_T GetAllocatedSize(bool bLogDetailed = false) const override;
 	virtual void LoadPackageRegistryData(FArchive& Ar, FLoadPackageRegistryData& InOutData) const override;
 	virtual void LoadPackageRegistryData(const FString& PackageFilename, FLoadPackageRegistryData& InOutData) const override;
@@ -345,6 +358,7 @@ private:
 private:
 
 	UE::AssetRegistry::FAssetRegistryImpl GuardedData;
+	UE::AssetRegistry::Impl::FActiveMountsLoaderWrapper ActiveMountsLoader;
 
 	/** Lock guarding the GuardedData */
 	mutable UE::AssetRegistry::Private::FRWLockWithPriority InterfaceLock;
@@ -456,6 +470,7 @@ private:
 
 	friend FAssetDataGatherer;
 	friend UE::AssetRegistry::FAssetRegistryImpl;
+	friend UE::AssetRegistry::Impl::FActiveMountsLoader;
 	friend UE::AssetRegistry::Premade::FAsyncConsumer;
 };
 PRAGMA_ENABLE_DEPRECATION_WARNINGS
diff --git a/Engine/Source/Runtime/AssetRegistry/Private/AssetRegistryImpl.h b/Engine/Source/Runtime/AssetRegistry/Private/AssetRegistryImpl.h
--- a/Engine/Source/Runtime/AssetRegistry/Private/AssetRegistryImpl.h
+++ b/Engine/Source/Runtime/AssetRegistry/Private/AssetRegistryImpl.h
@@ -33,6 +33,8 @@ class UAssetRegistryImpl;
 namespace UE::AssetRegistry::DependsNode { struct FGatherDependencyDataScratchpad; }
 namespace UE::AssetRegistry::Premade { enum class ELoadResult : uint8; }
 namespace UE::AssetRegistry::Premade { struct FAsyncConsumer; }
+namespace UE::AssetRegistry::Impl { class FActiveMountEvents; }
+namespace UE::AssetRegistry::Impl { class FActiveMountsLoader; }
 namespace UE::AssetRegistry::Impl { struct FEventContext; }
 namespace UE::AssetRegistry::Impl { struct FInitializeContext; }
 namespace UE::AssetRegistry::Impl { struct FScanPathContext; }
@@ -242,6 +244,8 @@ public:
 	void Serialize(FArchive& Ar, Impl::FEventContext& EventContext);
 	void AppendState(Impl::FEventContext& EventContext, const FAssetRegistryState& InState,
 		FAssetRegistryState::EInitializationMode Mode = FAssetRegistryState::EInitializationMode::Append, bool bEmitAssetEvents = false);
+	/** Similar to AppendState, but it moves rather than copies. @see FAssetRegistryState::ConsumeExisting. */
+	void ConsumeState(Impl::FEventContext& EventContext, FAssetRegistryState&& InState, bool bEmitAssetEvents = false);
 	void GetAllocatedSize(bool bLogDetailed, SIZE_T& StateSize, SIZE_T& StaticSize, SIZE_T& SearchSize) const;
 	bool IsLoadingAssets() const;
 	bool IsGathering() const;
@@ -257,10 +261,12 @@ public:
 	void AddLoadedAssetToProcess(const UObject& AssetLoaded);
 #endif
 	void OnMultipleMountPointsMounted(UAssetRegistryImpl& UARI, Impl::FEventContext& EventContext,
+		Impl::FActiveMountEvents& ActiveMountEvents,
 		Impl::FClassInheritanceContext& InheritanceContext,
 		UE::AssetRegistry::FInterfaceWriteScopeLock& ScopeLock,
 		TConstArrayView<TRefCountPtr<UE::PackageName::IMountPoint>> MountedMountPoints);
 	void OnMultipleMountPointsDismounted(UAssetRegistryImpl& UARI, Impl::FEventContext& EventContext,
+		Impl::FActiveMountEvents& ActiveMountEvents, UE::AssetRegistry::FInterfaceWriteScopeLock& ProofOfWriteLock,
 		TConstArrayView<TRefCountPtr<UE::PackageName::IMountPoint>> DismountedMountPoints);
 
 	// Other helper functions called by UAssetRegistryImpl
@@ -526,7 +532,8 @@ private:
 		TSet<FTopLevelAssetPath>& SubClassNames, TSet<FTopLevelAssetPath>& ProcessedClassNames, const TSet<FTopLevelAssetPath>& ExcludedClassNames) const;
 
 	/** Internal helper which processes a given state and adds its contents to the current registry */
-	void CachePathsFromState(Impl::FEventContext& EventContext, const FAssetRegistryState& InState);
+	void CachePathsFromState(Impl::FEventContext& EventContext, const FAssetRegistryState& InState,
+		const FAssetRegistrySerializationOptions* Options = nullptr);
 
 	/** Adds ClassPath-> ParentPath to CachedBPInheritanceMap. Applies CoreRedirects to NotYetRedirectedParentPath. */
 	void AddCachedBPClassParent(const FTopLevelAssetPath& ClassPath, const FTopLevelAssetPath& NotYetRedirectedParentPath);
@@ -536,11 +543,11 @@ private:
 
 	/** Conditionally loads the Premade AssetRegistry from disk, or queues it for asynchronous loading if not yet loaded. */
 	void ConsumeOrDeferPreloadedPremade(UAssetRegistryImpl& UARI, Impl::FEventContext& EventContext,
-		FInterfaceWriteScopeLock& ProofOfWriteLock);
+		Impl::FActiveMountEvents& ActiveMountEvents, FInterfaceWriteScopeLock& ProofOfWriteLock);
 	/** Moves a premade asset registry state into this AR */
 	void LoadPremadeAssetRegistry(UAssetRegistryImpl& UARI, Impl::FEventContext& Context,
-		FInterfaceWriteScopeLock& ProofOfWriteLock,
-		Premade::ELoadResult LoadResult, FAssetRegistryState&& ARState, FStringView StateFilepath);
+		Impl::FActiveMountEvents& ActiveMountEvents, FInterfaceWriteScopeLock& ProofOfWriteLock,
+		Premade::ELoadResult LoadResult, Impl::FAssetProviderInitData&& StateInitData);
 	/** Add MountPoints of all AssetDatas currently registered in this->State to the list of PersistentMountPoints. */
 	void UpdatePersistentMountPoints();
 	void OnInitialSearchCompleted(Impl::FEventContext& EventContext);
@@ -736,6 +743,7 @@ private:
 #endif
 	UE::AssetRegistry::Premade::FAsyncConsumer AsyncConsumer;
 	friend UE::AssetRegistry::Premade::FAsyncConsumer;
+	friend UE::AssetRegistry::Impl::FActiveMountsLoader;
 	friend UE::AssetRegistry::Impl::FClassInheritanceContext;
 	friend UE::AssetRegistry::Impl::FTickContext;
 
diff --git a/Engine/Source/Runtime/AssetRegistry/Public/AssetRegistry/IAssetRegistry.h b/Engine/Source/Runtime/AssetRegistry/Public/AssetRegistry/IAssetRegistry.h
--- a/Engine/Source/Runtime/AssetRegistry/Public/AssetRegistry/IAssetRegistry.h
+++ b/Engine/Source/Runtime/AssetRegistry/Public/AssetRegistry/IAssetRegistry.h
@@ -9,6 +9,7 @@
 #include "Containers/StringFwd.h"
 #include "Misc/AssetRegistryInterface.h"
 #include "Misc/Optional.h"
+#include "Tasks/Task.h"
 #include "Templates/Function.h"
 #include "Templates/PimplPtr.h"
 #include "UObject/Interface.h"
@@ -1039,7 +1040,90 @@ public:
 	virtual void Serialize(FStructuredArchive::FRecord Record) = 0;
 
 	/** Append the assets from the incoming state into our own */
-	virtual void AppendState(const FAssetRegistryState& InState, UE::AssetRegistry::EAppendMode AppendMode = UE::AssetRegistry::EAppendMode::Append) = 0;
+	virtual void AppendState(const FAssetRegistryState& InState,
+		UE::AssetRegistry::EAppendMode AppendMode = UE::AssetRegistry::EAppendMode::Append) = 0;
+
+	// ActiveMountUnloading: save memory when using cooked AssetRegistries by removing AssetRegistryData
+	// for MountPoints that are unmounted. Only beneficial if your game unloads mountpoints during some periods,
+	// and only usable if your game does not require knowledge of the assets in those mountpoints when unloaded.
+	// Currently only enablable for ClientOnly, since ClientOnly is the configuration that runs on the most memory
+	// constrained platforms. Enable by setting Engine.ini:[AssetRegistry]:bActiveMountLoadingEnabled_Client=true,
+	// and calling IAssetRegistry->ActiveMountsInitialize() when ready to start unloading at startup.
+	// By default, unloading is triggered by calling ActiveMountsUnloadAsync, but alternatively you can 
+	// activate unloading by default, whenever FPackageName::OnMountPointDismounted, by setting
+	// Engine.ini:[AssetRegistry]:bActiveMountLoadingAuto_Client=true.
+
+	/**
+	 * Returns true if ActiveMountsIsEnabled in the current configuration. Returns true even if ActiveMountsInitialize
+	 * has not yet been called. If not enabled, functions with AssetRegistry side effects will such as
+	 * ActiveMountsRegister* will still execute the side-effects, but all behavior specific to unloading and reloading
+	 * will be skipped by the ActiveMounts functions.
+	 */
+	virtual bool ActiveMountsIsEnabled() = 0;
+	/**
+	 * Called once per process to activate ActiveMountUnloading. Calls to ActiveMountsLoadAsync/ActiveMountsUnloadAsync
+	 * before this call are ignored.
+	 * This function also forces the parsing of the premade AssetRegistry: MountPoints in the premade AssetRegistry
+	 * state will be searched at this point in FPackageName or FPlugins; any mountpoints used by the premade that are
+	 * not mounted and not findable as a plugin should be registered before this via FPackageName::FindOrAddMountPoint.
+	 * Any mountpoints not findable will log an error and AssetDatas under them will always be loaded.
+	 */
+	virtual void ActiveMountsInitialize() = 0;
+
+	/**
+	 * Registers an AssetRegistry.bin file on disk for ActiveMount unloading, and also immediately loads its
+	 * data into the AssetRegistry with AppendState. The default premade AssetRegistry.bin is automatically
+	 * registered by the AssetRegistry and does not need this call.
+	 * If ActiveMountLoading is disabled, this function behaves the same as AppendState.
+	 * Returns with no action if AssetRegistryState can not be loaded from the given file.
+	 * This function also forces the parsing of the premade AssetRegistry, @see ActiveMountsInitialize.
+	 * 
+	 * @return True if the AssetRegistryState could be loaded from the given file, else false.
+	 */
+	virtual bool ActiveMountsRegisterAndLoadFilePathSynchronous(FStringView StateFilePath) = 0;
+	/**
+	 * Runs ActiveMountsRegisterAndLoadFilePathSynchronous on an async thread and calls OnComplete when done.
+	 * OnComplete is called from an async thread, and in game error cases might be called after IAssetRegistry
+	 * is shutdown. Caller should handle thread safety and should check for IAssetRegistry::Get != nullptr
+	 * if it is used in their callback.
+	*/
+	virtual void ActiveMountsRegisterAndLoadFilePathAsync(FStringView StateFilePath,
+		TFunction<void(bool bSucceeded)> OnComplete = TFunction<void(bool bSucceeded)>(),
+		UE::Tasks::ETaskPriority TaskPriority = UE::Tasks::ETaskPriority::Default) = 0;
+	/**
+	 * Requests that AssetDatas under the given MountPoints (looked up via FPackageName::FindOrAddMountPoint) be
+	 * reloaded into the AssetRegistry. This only affects MountPoints previously unloaded via auto unload or 
+	 * ActiveMountsUnloadAsync, and is only possible for AssetDatas in the premade AssetRegistry or in an AssetRegistry
+	 * file registered via an ActiveMountsRegister function. Calls OnComplete once the load completes.
+	 * 
+	 * @param OnComplete Optional delegate: void OnComplete(bool bCanceled). Called when the load is complete.
+	 *        bCanceled is set to true iff ActiveMountUnloading is not enabled or for any requested MountPoint the
+	 *        load failed or was overridden by an unload request. If true, test ActiveMountsIsEnabled and call
+	 *        ActiveMountIsLoaded to test state of each mountpoint.
+	 *        OnComplete is called from an async thread, caller should handle thread safety in their callback.
+	 */
+	virtual void ActiveMountsLoadAsync(TConstArrayView<FStringView> MountPointLongPackageNames,
+		TFunction<void(bool bCanceled)> OnComplete = TFunction<void(bool bCanceled)>()) = 0;
+	/**
+	 * Requests that AssetDatas under the given MountPoints (looked up via FPackageName::FindOrAddMountPoint) be
+	 * unloaded from the AssetRegistry. This is only possible for AssetDatas in the premade AssetRegistry or in an
+	 * AssetRegistry file registered via an ActiveMountsRegister function. Calls OnComplete once the unload completes.
+	 *
+	 * @param OnComplete Optional delegate: void OnComplete(bool bCanceled). Called when the unload is complete.
+	 *        bCanceled is set to true iff ActiveMountUnloading is not enabled or for any requested MountPoint the
+	 *        unload failed or was overridden by a load request. If true, test ActiveMountsIsEnabled and call
+	 *        ActiveMountIsLoaded to test state of each mountpoint.
+	 *        OnComplete is called from an async thread, caller should handle thread safety in their callback.
+	 */
+	virtual void ActiveMountsUnloadAsync(TConstArrayView<FStringView> MountPointLongPackageNames,
+		TFunction<void(bool bCanceled)> OnComplete = TFunction<void(bool bCanceled)>()) = 0;
+	/**
+	 * Returns whether the given MountPoint (looked up via FPackageName::FindOrAddMountPoint) is loaded in the
+	 * ActiveMounts system. Returns false if currently loading or unloading or if unloaded. Only valid if
+	 * ActiveMountsLoader is enabled and active and the MountPoint is used by the PreMade AssetRegistry or a
+	 * registered AssetRegistry file. Returns false for all other cases.
+	 */
+	virtual bool ActiveMountIsLoaded(FStringView MountPointLongPackageName) = 0;
 
 	/** Returns memory size of entire registry, optionally logging sizes */
 	virtual SIZE_T GetAllocatedSize(bool bLogDetailed = false) const = 0;
