AssetRegistry: ActiveMountsLoading: Usability tweaks.
Change configuration and comments to use consistent naming.
Add log statements to show when the system is active.
Add ActiveMountsIsInProgress and ActiveMountsReportWhenIdle.
Fix FSuspendPathsMounting to handle AR singleton being cleared during shutdown without crashing.
#rnx

#Synced-CL 50811273

diff --git a/Engine/Config/BaseEngine.ini b/Engine/Config/BaseEngine.ini
--- a/Engine/Config/BaseEngine.ini
+++ b/Engine/Config/BaseEngine.ini
@@ -2901,11 +2901,11 @@ bSerializePackageData=false
 bFilterAssetDataWithNoTags=false
 bFilterDependenciesWithNoTags=false
 bFilterSearchableNames=true
-; ActiveMountLoading: Enable this field on projects with thousands of plugins to save memory by unloading AR data,
-; but requires c++ handling as well, see ActiveMountLoading documentation in IAssetRegistry.h.
-bActiveMountLoadingEnabled_Client=false
-; Configures ActiveMountLoading to use default unloading rules, see ActiveMountUnloading in IAssetRegistry.h.
-bActiveMountLoadingAuto_Client=true
+; ActiveMountsLoading: Enable this field on projects with thousands of plugins to save memory by unloading AR data,
+; but requires c++ handling as well, see ActiveMountsLoading documentation in IAssetRegistry.h.
+bActiveMountsLoadingEnabled_Client=false
+; Configures ActiveMountsLoading to use default unloading rules, see ActiveMountsLoading in IAssetRegistry.h.
+bActiveMountsLoadingAuto_Client=true
 
 ; The editor can run the game or dedicated server on uncooked content, via UnrealEditor.exe <ProjectName> -game
 ; This is bad for performance but is supported for use in PlayInEditor
diff --git a/Engine/Source/Runtime/AssetRegistry/Private/ARLoader.h b/Engine/Source/Runtime/AssetRegistry/Private/ARLoader.h
--- a/Engine/Source/Runtime/AssetRegistry/Private/ARLoader.h
+++ b/Engine/Source/Runtime/AssetRegistry/Private/ARLoader.h
@@ -210,7 +210,7 @@ namespace UE::AssetRegistry::Impl
 /** Open the given AssetRegistry.bin file using the performant archive we use for loading it. */
 TUniquePtr<FArchive> OpenReadStateFileArchive(FStringView FilePath);
 
-/** Which event triggered our need to register the premade AssetRegistry with the ActiveMountLoaders */
+/** Which event triggered our need to register the premade AssetRegistry with the ActiveMountsLoading. */
 enum class EInitialProviderRegisterReason
 {
 	/* Expected case, OnActiveMountsInitialize was called. */
@@ -222,7 +222,7 @@ enum class EInitialProviderRegisterReason
 };
 
 /**
- * Holds callbacks and events triggered during ActiveMountLoading functions that need to be broadcast outside of the
+ * Holds callbacks and events triggered during ActiveMountsLoading functions that need to be broadcast outside of the
  * AssetRegistry's locks.
  */
 class FActiveMountEvents
@@ -233,6 +233,8 @@ public:
 	 * bCanceled parameter.
 	 */
 	void AddCompletion(TFunction<void(bool bCanceled)> InOnComplete, bool bInCanceled);
+	/** Add an OnIdle callback from ActiveMountsReportWhenIdle. */
+	void AddIdle(TFunction<void()> InOnIdle);
 	/** Send all of the added events. Must be called outside of the AssetRegistry's locks. */
 	void Broadcast();
 
@@ -243,13 +245,14 @@ private:
 		bool bCanceled = false;
 	};
 	TArray<FCompletionData> CompletionEvents;
+	TArray<TFunction<void()>> IdleCallbacks;
 };
 
 /**
  * The interface into the ActiveMountsLoader used by the rest of the AssetRegistry. The ActiveMountsLoader stores
  * filepaths for AssetRegistry.bin files that provide assets to the AssetRegistry, and it manages a state machine to
  * unload and reload the AssetDatas for mountpoints (such as plugin Content directories) when those mountpoints are
- * loaded and unloaded. @see the comment on ActiveMountUnloading in IAssetRegistry.h.
+ * loaded and unloaded. @see the comment on ActiveMountsLoading in IAssetRegistry.h.
  *
  * The ActiveMountsLoader has some minor performance costs and so should only be enabled in configurations and projects
  * that need the memory savings and that can behave correctly when AssetDatas are unloaded for inactive plugins.
@@ -290,9 +293,13 @@ public:
 	 */
 	void UnloadAllUnmountedAsync(UAssetRegistryImpl& UARI, FActiveMountEvents& EventContext,
 		TFunction<void(bool bCanceled)>&& OnComplete);
-	/** @see IAssetRegistry::ActiveMountIsLoaded. */
+	/** @see IAssetRegistry::ActiveMountsIsLoaded. */
 	bool IsLoaded(UAssetRegistryImpl& UARI, FActiveMountEvents& EventContext,
 		const TRefCountPtr<UE::PackageName::IMountPoint>& MountPoint);
+	/** @see IAssetRegistry::ActiveMountsIsInProgress. */
+	bool IsInProgress(UAssetRegistryImpl& UARI);
+	/** @see IAssetRegistry::ActiveMountsReportWhenIdle. */
+	void ReportWhenIdle(UAssetRegistryImpl& UARI, FActiveMountEvents& EventContext, TFunction<void()> OnIdle);
 
 	/** @see IAssetRegistry::ActiveMountsInitialize. */
 	void OnActiveMountsInitialize(UAssetRegistryImpl& UARI, FActiveMountEvents& EventContext,
@@ -358,23 +365,23 @@ public:
 	static bool IsEnabledByConfig();
 private:
 	/**
-	 * Called by all public functions, even when the full ActiveMountLoader is enabled. When disabled, operates as
+	 * Called by all public functions, even when the full ActiveMountsLoader is enabled. When disabled, operates as
 	 * cheaply as possible, does not lock, and returns nullptr. When enabled, locks and returns a pointer to
 	 * this->Inner.
 	 */
 	FActiveMountsLoader* LockIfEnabled();
-	/** Lock ActiveMountsLock, the lock that must be entered around any full ActiveMountLoader functions. */
+	/** Lock ActiveMountsLock, the lock that must be entered around any full ActiveMountsLoader functions. */
 	void Lock();
 	/** Unlock ActiveMountsLock. */
 	void Unlock();
 	/**
 	 * Return the lock that must be entered around AssetDataProvider register/deregister functions. Returns
-	 * nullptr if full ActiveMountLoader is not enabled.
+	 * nullptr if full ActiveMountsLoader is not enabled.
 	 */
 	TSharedPtr<UE::FRecursiveMutex> GetRegistrationLock();
 
 	/**
-	 * Lock that guards the ActiveMountLoader's data, separate from the AssetRegistry's data so that we do not block
+	 * Lock that guards the ActiveMountsLoader's data, separate from the AssetRegistry's data so that we do not block
 	 * AssetRegistry queries when doing some long operations inside the lock.
 	 * This lock is allowed to be held at the same time as the AssetRegistry's interface lock (avoid when possible).
 	 * When both are held, the AssetRegistry's interface lock must be entered first.
diff --git a/Engine/Source/Runtime/AssetRegistry/Private/ActiveMountsLoader.cpp b/Engine/Source/Runtime/AssetRegistry/Private/ActiveMountsLoader.cpp
--- a/Engine/Source/Runtime/AssetRegistry/Private/ActiveMountsLoader.cpp
+++ b/Engine/Source/Runtime/AssetRegistry/Private/ActiveMountsLoader.cpp
@@ -47,6 +47,15 @@ void FActiveMountEvents::AddCompletion(TFunction<void(bool bCanceled)> InOnCompl
 	CompletionEvents.Add(FCompletionData{ MoveTemp(InOnComplete), bInCanceled });
 }
 
+void FActiveMountEvents::AddIdle(TFunction<void()> InOnIdle)
+{
+	if (!InOnIdle)
+	{
+		return;
+	}
+	IdleCallbacks.Add(MoveTemp(InOnIdle));
+}
+
 void FActiveMountEvents::Broadcast()
 {
 	for (FCompletionData& Data : CompletionEvents)
@@ -55,6 +64,12 @@ void FActiveMountEvents::Broadcast()
 		Data.OnComplete(Data.bCanceled);
 	}
 	CompletionEvents.Empty();
+	for (TFunction<void()>& OnIdle : IdleCallbacks)
+	{
+		check(OnIdle);
+		OnIdle();
+	}
+	IdleCallbacks.Empty();
 }
 
 bool FActiveMountsLoaderWrapper::IsEnabled()
@@ -69,12 +84,12 @@ bool FActiveMountsLoaderWrapper::IsEnabledByConfig()
 			bool bLocalEnabled = false;
 			if (FPlatformProperties::RequiresCookedData() && IsRunningGame() && !IsRunningDedicatedServer())
 			{
-				if (!FParse::Bool(FCommandLine::Get(), TEXT("-ActiveMountLoading="), bLocalEnabled))
+				if (!FParse::Bool(FCommandLine::Get(), TEXT("-ActiveMountsLoading="), bLocalEnabled))
 				{
 					// We need the value from Config to calculate the proper value, so assert that GConfig is
 					// available by the time we first call IsEnabledByConfig.
 					check(GConfig);
-					GConfig->GetBool(TEXT("AssetRegistry"), TEXT("bActiveMountLoadingEnabled_Client"),
+					GConfig->GetBool(TEXT("AssetRegistry"), TEXT("bActiveMountsLoadingEnabled_Client"),
 						bLocalEnabled, GEngineIni);
 				}
 			}
@@ -87,14 +102,15 @@ void FActiveMountsLoaderWrapper::InitializeConfig()
 {
 	if (!IsEnabledByConfig())
 	{
+		UE_LOGFMT(LogAssetRegistry, Display, "ActiveMountsLoading is disabled.");
 		return;
 	}
 
 	bool bAutoUnload = false;
-	if (!FParse::Bool(FCommandLine::Get(), TEXT("-ActiveMountLoadingAuto="), bAutoUnload))
+	if (!FParse::Bool(FCommandLine::Get(), TEXT("-ActiveMountsLoadingAuto="), bAutoUnload))
 	{
 		check(GConfig); // It was already required by IsEnabledByConfig
-		GConfig->GetBool(TEXT("AssetRegistry"), TEXT("bActiveMountLoadingAuto_Client"),
+		GConfig->GetBool(TEXT("AssetRegistry"), TEXT("bActiveMountsLoadingAuto_Client"),
 			bAutoUnload, GEngineIni);
 	}
 
@@ -106,6 +122,7 @@ void FActiveMountsLoaderWrapper::InitializeConfig()
 	// Races during initialization are unexpected but are also not a problem.
 	if (!Inner.load(std::memory_order_relaxed))
 	{
+		UE_LOGFMT(LogAssetRegistry, Display, "ActiveMountsLoading is enabled.");
 		FActiveMountsLoader* LocalInner = new FActiveMountsLoader();
 		LocalInner->RegistrationLock = MakeShared<UE::FRecursiveMutex>();
 		LocalInner->bAutoUnload = bAutoUnload;
@@ -262,6 +279,43 @@ bool FActiveMountsLoaderWrapper::IsLoaded(UAssetRegistryImpl& UARI, FActiveMount
 	return false;
 }
 
+bool FActiveMountsLoaderWrapper::IsInProgress(UAssetRegistryImpl& UARI)
+{
+	if (FActiveMountsLoader* LocalInner = LockIfEnabled(); LocalInner)
+	{
+		ON_SCOPE_EXIT
+		{
+			Unlock();
+		};
+		if (!LocalInner->bShutdown)
+		{
+			return LocalInner->IsInProgress(UARI);
+		}
+	}
+	return false;
+}
+
+void FActiveMountsLoaderWrapper::ReportWhenIdle(UAssetRegistryImpl& UARI, FActiveMountEvents& EventContext,
+	TFunction<void()> OnIdle)
+{
+	if (!OnIdle)
+	{
+		return;
+	}
+	if (FActiveMountsLoader* LocalInner = LockIfEnabled(); LocalInner)
+	{
+		ON_SCOPE_EXIT
+		{
+			Unlock();
+		};
+		if (!LocalInner->bShutdown)
+		{
+			return LocalInner->ReportWhenIdle(UARI, EventContext, MoveTemp(OnIdle));
+		}
+	}
+	EventContext.AddIdle(MoveTemp(OnIdle));
+}
+
 void FActiveMountsLoaderWrapper::OnActiveMountsInitialize(UAssetRegistryImpl& UARI, FActiveMountEvents& EventContext,
 	FInterfaceWriteScopeLock& ProofOfWriteLock)
 {
@@ -721,6 +775,11 @@ void FActiveMountsLoader::FinishShutdown(UAssetRegistryImpl& UARI, FActiveMountE
 		}
 	}
 	UnloadSubscribers.Empty();
+	for (TFunction<void()>& OnIdle : OnIdleSubscribers)
+	{
+		EventContext.AddIdle(MoveTemp(OnIdle));
+	}
+	OnIdleSubscribers.Empty();
 
 	// MountPoints have a raw pointer to AssetProviders, AssetProviders do not have a pointer to MountPoints.
 	// MountPoints do not currently access their AssetProvider during their destructor, but in case that changes,
@@ -1012,6 +1071,34 @@ bool FActiveMountsLoader::IsLoaded(UAssetRegistryImpl& UARI, FActiveMountEvents&
 	return MountData->GetCurrentState() == EMountState::Loaded;
 }
 
+bool FActiveMountsLoader::IsInProgress(UAssetRegistryImpl& UARI)
+{
+	if (UnloadTask.IsValid())
+	{
+		return true;
+	}
+	for (FAssetProviderSetKey& SetKey : AssetProviders)
+	{
+		FAssetProvider& AssetProvider = SetKey.GetAssetProvider();
+		if (AssetProvider.GetLoadTask().IsValid())
+		{
+			return true;
+		}
+	}
+	return false;
+}
+
+void FActiveMountsLoader::ReportWhenIdle(UAssetRegistryImpl& UARI, FActiveMountEvents& EventContext,
+	TFunction<void()> OnIdle)
+{
+	if (!IsInProgress(UARI))
+	{
+		EventContext.AddIdle(OnIdle);
+		return;
+	}
+	OnIdleSubscribers.Add(MoveTemp(OnIdle));
+}
+
 void FActiveMountsLoader::OnActiveMountsInitialize(UAssetRegistryImpl& UARI, FActiveMountEvents& EventContext,
 	FInterfaceWriteScopeLock& ProofOfWriteLock)
 {
@@ -1603,6 +1690,23 @@ void FActiveMountsLoader::NotifyMountPointUnloadComplete(FActiveMountEvents& Eve
 	}
 }
 
+void FActiveMountsLoader::ConditionalNotifyIdle(UAssetRegistryImpl& UARI, FActiveMountEvents& EventContext)
+{
+	if (OnIdleSubscribers.IsEmpty())
+	{
+		return;
+	}
+	if (IsInProgress(UARI))
+	{
+		return;
+	}
+	for (TFunction<void()>& OnIdle : OnIdleSubscribers)
+	{
+		EventContext.AddIdle(MoveTemp(OnIdle));
+	}
+	OnIdleSubscribers.Empty();
+}
+
 void FActiveMountsLoader::AddRefMountPointLoad(UAssetRegistryImpl& UARI, FMountPointData* MountPointData)
 {
 	// Caller must not call AddRef if MountPointData is null or does not have an AssetProvider.
@@ -1659,13 +1763,19 @@ void FActiveMountsLoader::ReleaseMountPointUnload(UAssetRegistryImpl& UARI, FMou
 void FActiveMountsLoader::RunAssetProviderLoadTask(UAssetRegistryImpl& UARI, FAssetProvider& AssetProvider)
 {
 	LLM_SCOPE(ELLMTag::AssetRegistry);
+	FActiveMountEvents OnIdleEvents;
+	ON_SCOPE_EXIT
+	{
+		OnIdleEvents.Broadcast();
+	};
 
-	auto MarkAsyncThreadComplete = [&AssetProvider]()
+	auto MarkAsyncThreadComplete = [this, &AssetProvider, &UARI, &OnIdleEvents]()
 		{
 			// Caller must call only while within the ActiveMounts critical section, to avoid a softlock race
 			// where the main thread sees that the AsyncTask still exists and does not schedule another one,
 			// but the AsyncTask is shutting down and does not plan to check for further work to do.
 			AssetProvider.GetLoadTask() = UE::Tasks::FTask();
+			ConditionalNotifyIdle(UARI, OnIdleEvents);
 		};
 
 	// Requests might come in unbatched, or in batches that race with our load and come in after we
@@ -1878,6 +1988,9 @@ void FActiveMountsLoader::RunAssetProviderLoadTask(UAssetRegistryImpl& UARI, FAs
 					*StateFilePath, *FString(MountPointsToLoad[0]->GetLongPackageName()));
 			}
 
+			UE_LOGFMT(LogAssetRegistry, Display, "ActiveMountsLoading reloaded AssetRegistry data for {NumMountPoints} MountPoints.",
+				MountPointsToLoad.Num());
+
 			for (TRefCountPtr<UE::PackageName::IMountPoint>& MountPoint : MountPointsToLoad)
 			{
 				// The MountPoints should still be in the Loading state, because nothing is allowed to take them out
@@ -1925,13 +2038,19 @@ void FActiveMountsLoader::RunAssetProviderLoadTask(UAssetRegistryImpl& UARI, FAs
 void FActiveMountsLoader::RunMountPointUnloadTask(UAssetRegistryImpl& UARI)
 {
 	LLM_SCOPE(ELLMTag::AssetRegistry);
+	FActiveMountEvents OnIdleEvents;
+	ON_SCOPE_EXIT
+	{
+		OnIdleEvents.Broadcast();
+	};
 
-	auto MarkAsyncThreadComplete = [this]()
+	auto MarkAsyncThreadComplete = [this, &UARI, &OnIdleEvents]()
 		{
 			// Caller must call only while within the ActiveMounts critical section, to avoid a softlock race
 			// where the main thread sees that the AsyncTask still exists and does not schedule another one,
 			// but the AsyncTask is shutting down and does not plan to check for further work to do.
 			UnloadTask = UE::Tasks::FTask();
+			ConditionalNotifyIdle(UARI, OnIdleEvents);
 		};
 
 	// Requests might come in unbatched, or in batches that race with our unload and come in after we
@@ -2023,6 +2142,9 @@ void FActiveMountsLoader::RunMountPointUnloadTask(UAssetRegistryImpl& UARI)
 				return;
 			}
 
+			UE_LOGFMT(LogAssetRegistry, Display, "ActiveMountsLoading unloaded AssetRegistry data for {NumMountPoints} MountPoints.",
+				MountPointsToUnload.Num());
+
 			for (TRefCountPtr<UE::PackageName::IMountPoint>& MountPoint : MountPointsToUnload)
 			{
 				// The MountPoints should still be in the Unloading state, because nothing is allowed to take them out
diff --git a/Engine/Source/Runtime/AssetRegistry/Private/ActiveMountsLoader.h b/Engine/Source/Runtime/AssetRegistry/Private/ActiveMountsLoader.h
--- a/Engine/Source/Runtime/AssetRegistry/Private/ActiveMountsLoader.h
+++ b/Engine/Source/Runtime/AssetRegistry/Private/ActiveMountsLoader.h
@@ -247,6 +247,8 @@ public:
 		TFunction<void(bool bCanceled)>&& OnComplete);
 	bool IsLoaded(UAssetRegistryImpl& UARI, FActiveMountEvents& EventContext,
 		const TRefCountPtr<UE::PackageName::IMountPoint>& MountPoint);
+	bool IsInProgress(UAssetRegistryImpl& UARI);
+	void ReportWhenIdle(UAssetRegistryImpl& UARI, FActiveMountEvents& EventContext, TFunction<void()> OnIdle);
 
 	void OnActiveMountsInitialize(UAssetRegistryImpl& UARI, FActiveMountEvents& EventContext,
 		FInterfaceWriteScopeLock& ProofOfWriteLock);
@@ -295,6 +297,7 @@ private:
 		const TRefCountPtr<UE::PackageName::IMountPoint>& MountPoint, bool bCanceled);
 	void NotifyMountPointUnloadComplete(FActiveMountEvents& EventContext,
 		const TRefCountPtr<UE::PackageName::IMountPoint>& MountPoint, bool bCanceled);
+	void ConditionalNotifyIdle(UAssetRegistryImpl& UARI, FActiveMountEvents& EventContext);
 	void AddRefMountPointLoad(UAssetRegistryImpl& UARI, FMountPointData* MountPointData);
 	void ReleaseMountPointLoad(UAssetRegistryImpl& UARI, FMountPointData* MountPointData);
 	void AddRefMountPointUnload(UAssetRegistryImpl& UARI, FMountPointData* MountPointData);
@@ -310,6 +313,8 @@ private:
 	TMap<TRefCountPtr<UE::PackageName::IMountPoint>, TArray<TRefCountPtr<FMountLoadEventSubscriber>>> LoadSubscribers;
 	/** The list of OnComplete functions that are waiting on a MountPoint to unload. */
 	TMap<TRefCountPtr<UE::PackageName::IMountPoint>, TArray<TRefCountPtr<FMountLoadEventSubscriber>>> UnloadSubscribers;
+	/** The list of callbacks that are waiting on ReportWhenIdle. */
+	TArray<TFunction<void()>> OnIdleSubscribers;
 	/**
 	 * An additional lock entered into by registration and unregistration functions. Unregistration functions have to
 	 * wait outside of the usual ActiveMountsLoader lock on FActiveMountsLoaderWrapper::ActiveMountsLock because the
diff --git a/Engine/Source/Runtime/AssetRegistry/Private/AssetRegistry.cpp b/Engine/Source/Runtime/AssetRegistry/Private/AssetRegistry.cpp
--- a/Engine/Source/Runtime/AssetRegistry/Private/AssetRegistry.cpp
+++ b/Engine/Source/Runtime/AssetRegistry/Private/AssetRegistry.cpp
@@ -6167,7 +6167,7 @@ void UAssetRegistryImpl::ActiveMountsUnloadAsync(TConstArrayView<FStringView> Mo
 	ActiveMountEvents.Broadcast();
 }
 
-bool UAssetRegistryImpl::ActiveMountIsLoaded(FStringView MountPointLongPackageName)
+bool UAssetRegistryImpl::ActiveMountsIsLoaded(FStringView MountPointLongPackageName)
 {
 	using namespace UE::AssetRegistry;
 
@@ -6185,6 +6185,18 @@ bool UAssetRegistryImpl::ActiveMountIsLoaded(FStringView MountPointLongPackageNa
 	return bResult;
 }
 
+bool UAssetRegistryImpl::ActiveMountsIsInProgress()
+{
+	return ActiveMountsLoader.IsInProgress(*this);
+}
+
+void UAssetRegistryImpl::ActiveMountsReportWhenIdle(TFunction<void()> OnIdle)
+{
+	UE::AssetRegistry::Impl::FActiveMountEvents ActiveMountEvents;
+	ActiveMountsLoader.ReportWhenIdle(*this, ActiveMountEvents, MoveTemp(OnIdle));
+	ActiveMountEvents.Broadcast();
+}
+
 namespace UE::AssetRegistry
 {
 
diff --git a/Engine/Source/Runtime/AssetRegistry/Private/AssetRegistry.h b/Engine/Source/Runtime/AssetRegistry/Private/AssetRegistry.h
--- a/Engine/Source/Runtime/AssetRegistry/Private/AssetRegistry.h
+++ b/Engine/Source/Runtime/AssetRegistry/Private/AssetRegistry.h
@@ -126,7 +126,9 @@ public:
 		TFunction<void(bool bCanceled)> OnComplete = TFunction<void(bool bCanceled)>()) override;
 	virtual void ActiveMountsUnloadAsync(TConstArrayView<FStringView> MountPointLongPackageNames,
 		TFunction<void(bool bCanceled)> OnComplete = TFunction<void(bool bCanceled)>()) override;
-	virtual bool ActiveMountIsLoaded(FStringView MountPointLongPackageName) override;
+	virtual bool ActiveMountsIsLoaded(FStringView MountPointLongPackageName) override;
+	virtual bool ActiveMountsIsInProgress() override;
+	virtual void ActiveMountsReportWhenIdle(TFunction<void()> OnIdle) override;
 	virtual SIZE_T GetAllocatedSize(bool bLogDetailed = false) const override;
 	virtual void LoadPackageRegistryData(FArchive& Ar, FLoadPackageRegistryData& InOutData) const override;
 	virtual void LoadPackageRegistryData(const FString& PackageFilename, FLoadPackageRegistryData& InOutData) const override;
diff --git a/Engine/Source/Runtime/AssetRegistry/Public/AssetRegistry/IAssetRegistry.h b/Engine/Source/Runtime/AssetRegistry/Public/AssetRegistry/IAssetRegistry.h
--- a/Engine/Source/Runtime/AssetRegistry/Public/AssetRegistry/IAssetRegistry.h
+++ b/Engine/Source/Runtime/AssetRegistry/Public/AssetRegistry/IAssetRegistry.h
@@ -1043,15 +1043,16 @@ public:
 	virtual void AppendState(const FAssetRegistryState& InState,
 		UE::AssetRegistry::EAppendMode AppendMode = UE::AssetRegistry::EAppendMode::Append) = 0;
 
-	// ActiveMountUnloading: save memory when using cooked AssetRegistries by removing AssetRegistryData
+	// ActiveMountsLoading: save memory when using cooked AssetRegistries by removing AssetRegistryData
 	// for MountPoints that are unmounted. Only beneficial if your game unloads mountpoints during some periods,
 	// and only usable if your game does not require knowledge of the assets in those mountpoints when unloaded.
 	// Currently only enablable for ClientOnly, since ClientOnly is the configuration that runs on the most memory
-	// constrained platforms. Enable by setting Engine.ini:[AssetRegistry]:bActiveMountLoadingEnabled_Client=true,
+	// constrained platforms. Enable by setting Engine.ini:[AssetRegistry]:bActiveMountsLoadingEnabled_Client=true,
 	// and calling IAssetRegistry->ActiveMountsInitialize() when ready to start unloading at startup.
-	// By default, unloading is triggered by calling ActiveMountsUnloadAsync, but alternatively you can 
-	// activate unloading by default, whenever FPackageName::OnMountPointDismounted, by setting
-	// Engine.ini:[AssetRegistry]:bActiveMountLoadingAuto_Client=true.
+	// By default, ActiveMountsLoading runs in auto mode and unloading is triggered whenever
+	// FPackageName::OnMountPointDismounted is called, but alternatively you can set 
+	// Engine.ini:[AssetRegistry]:bActiveMountsLoadingAuto_Client=false and call ActiveMountsUnloadAsync and
+	// ActiveMountsLoadAsync explicitly.
 
 	/**
 	 * Returns true if ActiveMountsIsEnabled in the current configuration. Returns true even if ActiveMountsInitialize
@@ -1061,7 +1062,7 @@ public:
 	 */
 	virtual bool ActiveMountsIsEnabled() = 0;
 	/**
-	 * Called once per process to activate ActiveMountUnloading. Calls to ActiveMountsLoadAsync/ActiveMountsUnloadAsync
+	 * Called once per process to activate ActiveMountsLoading. Calls to ActiveMountsLoadAsync/ActiveMountsUnloadAsync
 	 * before this call are ignored.
 	 * This function also forces the parsing of the premade AssetRegistry: MountPoints in the premade AssetRegistry
 	 * state will be searched at this point in FPackageName or FPlugins; any mountpoints used by the premade that are
@@ -1071,10 +1072,10 @@ public:
 	virtual void ActiveMountsInitialize() = 0;
 
 	/**
-	 * Registers an AssetRegistry.bin file on disk for ActiveMount unloading, and also immediately loads its
+	 * Registers an AssetRegistry.bin file on disk for ActiveMountsLoading, and also immediately loads its
 	 * data into the AssetRegistry with AppendState. The default premade AssetRegistry.bin is automatically
 	 * registered by the AssetRegistry and does not need this call.
-	 * If ActiveMountLoading is disabled, this function behaves the same as AppendState.
+	 * If ActiveMountsLoading is disabled, this function behaves the same as AppendState.
 	 * Returns with no action if AssetRegistryState can not be loaded from the given file.
 	 * This function also forces the parsing of the premade AssetRegistry, @see ActiveMountsInitialize.
 	 * 
@@ -1097,9 +1098,9 @@ public:
 	 * file registered via an ActiveMountsRegister function. Calls OnComplete once the load completes.
 	 * 
 	 * @param OnComplete Optional delegate: void OnComplete(bool bCanceled). Called when the load is complete.
-	 *        bCanceled is set to true iff ActiveMountUnloading is not enabled or for any requested MountPoint the
+	 *        bCanceled is set to true iff ActiveMountsLoading is not enabled or for any requested MountPoint the
 	 *        load failed or was overridden by an unload request. If true, test ActiveMountsIsEnabled and call
-	 *        ActiveMountIsLoaded to test state of each mountpoint.
+	 *        ActiveMountsIsLoaded to test state of each mountpoint.
 	 *        OnComplete is called from an async thread, caller should handle thread safety in their callback.
 	 */
 	virtual void ActiveMountsLoadAsync(TConstArrayView<FStringView> MountPointLongPackageNames,
@@ -1110,9 +1111,9 @@ public:
 	 * AssetRegistry file registered via an ActiveMountsRegister function. Calls OnComplete once the unload completes.
 	 *
 	 * @param OnComplete Optional delegate: void OnComplete(bool bCanceled). Called when the unload is complete.
-	 *        bCanceled is set to true iff ActiveMountUnloading is not enabled or for any requested MountPoint the
+	 *        bCanceled is set to true iff ActiveMountsLoading is not enabled or for any requested MountPoint the
 	 *        unload failed or was overridden by a load request. If true, test ActiveMountsIsEnabled and call
-	 *        ActiveMountIsLoaded to test state of each mountpoint.
+	 *        ActiveMountsIsLoaded to test state of each mountpoint.
 	 *        OnComplete is called from an async thread, caller should handle thread safety in their callback.
 	 */
 	virtual void ActiveMountsUnloadAsync(TConstArrayView<FStringView> MountPointLongPackageNames,
@@ -1120,10 +1121,25 @@ public:
 	/**
 	 * Returns whether the given MountPoint (looked up via FPackageName::FindOrAddMountPoint) is loaded in the
 	 * ActiveMounts system. Returns false if currently loading or unloading or if unloaded. Only valid if
-	 * ActiveMountsLoader is enabled and active and the MountPoint is used by the PreMade AssetRegistry or a
-	 * registered AssetRegistry file. Returns false for all other cases.
+	 * ActiveMountsLoader is enabled and the MountPoint is used by the PreMade AssetRegistry or a registered
+	 * AssetRegistry file. Returns false for all other cases.
 	 */
-	virtual bool ActiveMountIsLoaded(FStringView MountPointLongPackageName) = 0;
+	virtual bool ActiveMountsIsLoaded(FStringView MountPointLongPackageName) = 0;
+	/**
+	 * Returns whether ActiveMountsLoading is in progress for any MountPoint. Returns false if ActiveMountsLoading is
+	 * not enabled.
+	 */
+	virtual bool ActiveMountsIsInProgress() = 0;
+	/**
+	 * Call the given callback when all requested ActiveMountsLoading loads and unloads complete and the system
+	 * becomes idle. The callback is called outside of the lock, so by the time it is called, a new request may
+	 * have come in and the system might be in progress again. Intended use of this function is to guard entry into a
+	 * state that relies on AssetRegistry queries after issuing a set of ActiveMounts loads or unloads.
+	 * If the system is already idle (ActiveMountsIsInProgress returns false), the callback is called immediately.
+	 * If ActiveMountsLoading is not enabled, the callback is called immediately.
+	 * The callback might be called from an async thread; caller should handle thread safety in their callback.
+	 */
+	virtual void ActiveMountsReportWhenIdle(TFunction<void()> OnIdle) = 0;
 
 	/** Returns memory size of entire registry, optionally logging sizes */
 	virtual SIZE_T GetAllocatedSize(bool bLogDetailed = false) const = 0;
@@ -1242,13 +1258,13 @@ public:
 	/**
 	 * Request to pause paths mounting.
 	 * Can be used by the calling code to temporarily pause loading of data inside the AssetRegistry in response
-	 * to mounted MountPoints (e.g. ActiveMountLoading). All mounting requests will be executed in one batch at the
+	 * to mounted MountPoints (e.g. ActiveMountsLoading). All mounting requests will be executed in one batch at the
 	 * end of that scope.
 	 */
 	struct FSuspendPathsMounting
 	{
 		UE_NODISCARD_CTOR FSuspendPathsMounting() { IAssetRegistry::GetChecked().SuspendPathsMounting(); }
-		~FSuspendPathsMounting() { IAssetRegistry::GetChecked().ResumePathsMounting(); }
+		~FSuspendPathsMounting() { IAssetRegistry* AR = IAssetRegistry::Get(); if (AR) AR->ResumePathsMounting(); }
 	};
 
 	/**
