[AssetRegistry]

Move LoadedAssetsToProcess out of the GuaradedData and allow threadsafe access via a MPSC queue rather than relying on the InterfaceLock

#rnx

#Synced-CL 50826644

diff --git a/Engine/Source/Runtime/AssetRegistry/Private/AssetRegistry.cpp b/Engine/Source/Runtime/AssetRegistry/Private/AssetRegistry.cpp
--- a/Engine/Source/Runtime/AssetRegistry/Private/AssetRegistry.cpp
+++ b/Engine/Source/Runtime/AssetRegistry/Private/AssetRegistry.cpp
@@ -4539,9 +4539,7 @@ void UAssetRegistryImpl::AssetTagsFinalized(const UObject& FinalizedAsset)
 		return;
 	}
 	LLM_SCOPE(ELLMTag::AssetRegistry);
-
-	UE::AssetRegistry::FInterfaceWriteScopeLock InterfaceScopeLock(InterfaceLock);
-	GuardedData.AddLoadedAssetToProcess(FinalizedAsset);
+	AddLoadedAssetToProcess(FinalizedAsset);
 #endif
 }
 
@@ -8672,8 +8670,12 @@ void FAssetRegistryImpl::OnDirectoryRescanRequired(FInterfaceWriteScopeLock& Sco
 void UAssetRegistryImpl::OnAssetLoaded(UObject *AssetLoaded)
 {
 	LLM_SCOPE(ELLMTag::AssetRegistry);
-	UE::AssetRegistry::FInterfaceWriteScopeLock InterfaceScopeLock(InterfaceLock);
-	GuardedData.AddLoadedAssetToProcess(*AssetLoaded);
+	AddLoadedAssetToProcess(*AssetLoaded);
+}
+
+void UAssetRegistryImpl::AddLoadedAssetToProcess(const UObject& AssetLoaded)
+{
+	LoadedAssetsToProcess.Enqueue(&AssetLoaded);
 }
 
 void UAssetRegistryImpl::ProcessLoadedAssetsToUpdateCache(UE::AssetRegistry::Impl::FEventContext& EventContext,
@@ -8699,7 +8701,7 @@ void UAssetRegistryImpl::ProcessLoadedAssetsToUpdateCache(UE::AssetRegistry::Imp
 	{
 		LLM_SCOPE(ELLMTag::AssetRegistry);
 		UE::AssetRegistry::FInterfaceWriteScopeLock InterfaceScopeLock(InterfaceLock);
-		GuardedData.GetProcessLoadedAssetsBatch(BatchObjects, BatchSize, bUpdateDiskCacheAfterLoad);
+		GuardedData.GetProcessLoadedAssetsBatch(LoadedAssetsToProcess, BatchObjects, BatchSize, bUpdateDiskCacheAfterLoad);
 		if (BatchObjects.Num() == 0)
 		{
 			return;
@@ -8736,35 +8738,20 @@ void UAssetRegistryImpl::ProcessLoadedAssetsToUpdateCache(UE::AssetRegistry::Imp
 
 		LLM_SCOPE(ELLMTag::AssetRegistry);
 		UE::AssetRegistry::FInterfaceWriteScopeLock InterfaceScopeLock(InterfaceLock);
-		GuardedData.PushProcessLoadedAssetsBatch(EventContext, BatchAssetDatas,
+		GuardedData.PushProcessLoadedAssetsBatch(EventContext, LoadedAssetsToProcess, BatchAssetDatas,
 			TArrayView<const UObject*>(BatchObjects).Slice(Index, CurrentBatchSize-Index));
 		if (bTimedOut)
 		{
 			break;
 		}
-		GuardedData.GetProcessLoadedAssetsBatch(BatchObjects, BatchSize, bUpdateDiskCacheAfterLoad);
+		GuardedData.GetProcessLoadedAssetsBatch(LoadedAssetsToProcess, BatchObjects, BatchSize, bUpdateDiskCacheAfterLoad);
 	}
 }
 
 namespace UE::AssetRegistry
 {
 
-void FAssetRegistryImpl::AddLoadedAssetToProcess(const UObject& AssetLoaded)
-{
-	// Make sure the loaded asset is from a monitored path
-	if (GlobalGatherer.IsValid())
-	{
-		FString LocalPath;
-		if (!FPackageName::TryConvertLongPackageNameToFilename(AssetLoaded.GetPackage()->GetName(), LocalPath))
-		{
-			return;
-		}
-	}
-
-	LoadedAssetsToProcess.Add(&AssetLoaded);
-}
-
-void FAssetRegistryImpl::GetProcessLoadedAssetsBatch(TArray<const UObject*>& OutLoadedAssets, uint32 BatchSize,
+void FAssetRegistryImpl::GetProcessLoadedAssetsBatch(TQueue<TWeakObjectPtr<const UObject>, EQueueMode::Mpsc>& LoadedAssetsToProcess, TArray<const UObject*>& OutLoadedAssets, uint32 BatchSize,
 	bool bUpdateDiskCacheAfterLoad)
 {
 	if (!GlobalGatherer.IsValid() || !bUpdateDiskCacheAfterLoad)
@@ -8774,9 +8761,10 @@ void FAssetRegistryImpl::GetProcessLoadedAssetsBatch(TArray<const UObject*>& Out
 	}
 
 	OutLoadedAssets.Reset(BatchSize);
-	while (!LoadedAssetsToProcess.IsEmpty() && OutLoadedAssets.Num() < static_cast<int32>(BatchSize))
+	TWeakObjectPtr<const UObject> AssetToProcess;
+	while (OutLoadedAssets.Num() < static_cast<int32>(BatchSize) && LoadedAssetsToProcess.Dequeue(AssetToProcess))
 	{
-		const UObject* LoadedAsset = LoadedAssetsToProcess.PopFrontValue().Get();
+		const UObject* LoadedAsset = AssetToProcess.Get();
 		if (!LoadedAsset)
 		{
 			// This could be null, in which case it already got freed, ignore
@@ -8814,6 +8802,7 @@ void FAssetRegistryImpl::GetProcessLoadedAssetsBatch(TArray<const UObject*>& Out
 }
 
 void FAssetRegistryImpl::PushProcessLoadedAssetsBatch(Impl::FEventContext& EventContext,
+	TQueue<TWeakObjectPtr<const UObject>, EQueueMode::Mpsc>& LoadedAssetsToProcess,
 	TArrayView<FAssetData> LoadedAssetDatas, TArrayView<const UObject*> UnprocessedFromBatch)
 {
 	// Add or update existing for all of the AssetDatas created by the batch
@@ -8848,9 +8837,9 @@ void FAssetRegistryImpl::PushProcessLoadedAssetsBatch(Impl::FEventContext& Event
 	}
 
 	// Push back any objects from the batch that were not processed due to timing out
-	for (const UObject* Obj : ReverseIterate(UnprocessedFromBatch))
+	for (const UObject* Obj : UnprocessedFromBatch)
 	{
-		LoadedAssetsToProcess.EmplaceFront(Obj);
+		LoadedAssetsToProcess.Enqueue(Obj);
 	}
 }
 
diff --git a/Engine/Source/Runtime/AssetRegistry/Private/AssetRegistry.h b/Engine/Source/Runtime/AssetRegistry/Private/AssetRegistry.h
--- a/Engine/Source/Runtime/AssetRegistry/Private/AssetRegistry.h
+++ b/Engine/Source/Runtime/AssetRegistry/Private/AssetRegistry.h
@@ -299,6 +299,7 @@ private:
 		UE::AssetRegistry::EScanFlags InScanFlags);
 
 #if WITH_EDITOR
+	void AddLoadedAssetToProcess(const UObject& AssetLoaded);
 	/** Create FAssetData from any loaded UObject assets and store the updated AssetData in the state */
 	void ProcessLoadedAssetsToUpdateCache(UE::AssetRegistry::Impl::FEventContext& EventContext,
 		UE::AssetRegistry::Impl::EGatherStatus Status, UE::AssetRegistry::Impl::FInterruptionContext& InOutInterruptionContext);
@@ -375,6 +376,12 @@ private:
 	/** Handles to all registered OnDirectoryChanged delegates */
 	TMap<FString, FDelegateHandle> OnDirectoryChangedDelegateHandles;
 	TArray<FString> DirectoryWatchRoots;
+
+	/** List of objects that need to be processed because they were loaded or saved.
+	*   This object lives outside of the GuardedData as it doesn't require the InterfaceLock
+	*   to be written to. Reads must be coordinated such that only one thread reads at a time.
+	*/
+	TQueue<TWeakObjectPtr<const UObject>, EQueueMode::Mpsc> LoadedAssetsToProcess;
 #endif
 
 #if WITH_EDITORONLY_DATA
diff --git a/Engine/Source/Runtime/AssetRegistry/Private/AssetRegistryImpl.h b/Engine/Source/Runtime/AssetRegistry/Private/AssetRegistryImpl.h
--- a/Engine/Source/Runtime/AssetRegistry/Private/AssetRegistryImpl.h
+++ b/Engine/Source/Runtime/AssetRegistry/Private/AssetRegistryImpl.h
@@ -10,8 +10,9 @@
 #include "AssetRegistry/IAssetRegistry.h"
 #include "AssetRegistry/PathTree.h"
 #include "Containers/ArrayView.h"
-#include "Containers/RingBuffer.h"
 #include "Containers/Map.h"
+#include "Containers/Queue.h"
+#include "Containers/RingBuffer.h"
 #include "Containers/Set.h"
 #include "HAL/CriticalSection.h"
 #include "Math/NumericLimits.h"
@@ -258,7 +259,6 @@ public:
 		TArray<FFileChangeData>& FileChangesProcessed);
 	void OnDirectoryRescanRequired(FInterfaceWriteScopeLock& ScopeLock, Impl::FEventContext& EventContext,
 		Impl::FClassInheritanceContext& InheritanceContext, FString& DirPath, int64 BeforeTimeStamp);
-	void AddLoadedAssetToProcess(const UObject& AssetLoaded);
 #endif
 	void OnMultipleMountPointsMounted(UAssetRegistryImpl& UARI, Impl::FEventContext& EventContext,
 		Impl::FActiveMountEvents& ActiveMountEvents,
@@ -323,9 +323,10 @@ public:
 	void ClearGathererCache();
 #if WITH_EDITOR
 	void AssetsSaved(UE::AssetRegistry::Impl::FEventContext& EventContext, TArray<FAssetData>&& Assets);
-	void GetProcessLoadedAssetsBatch(TArray<const UObject*>& OutLoadedAssets, uint32 BatchSize,
-		bool bUpdateDiskCacheAfterLoad);
+	void GetProcessLoadedAssetsBatch(TQueue<TWeakObjectPtr<const UObject>, EQueueMode::Mpsc>& LoadedAssetsToProcess,
+		TArray<const UObject*>& OutLoadedAssets, uint32 BatchSize, bool bUpdateDiskCacheAfterLoad);
 	void PushProcessLoadedAssetsBatch(Impl::FEventContext& EventContext,
+		TQueue<TWeakObjectPtr<const UObject>, EQueueMode::Mpsc>& LoadedAssetsToProcess,
 		TArrayView<FAssetData> LoadedAssetDatas, TArrayView<const UObject*> UnprocessedFromBatch);
 	bool IsLoadedAndTrackedAsset(const FSoftObjectPath& InObjectPath) const;
 	/** Call LoadCalculatedDependencies on each Package updated after the last LoadCalculatedDependencies. */
@@ -706,9 +707,6 @@ private:
 	TSet<FName> PackagesNeedingDependencyCalculation;
 	TSet<FName> PackagesNeedingDependencyCalculationOnGameThread;
 
-	/** List of objects that need to be processed because they were loaded or saved */
-	TRingBuffer<TWeakObjectPtr<const UObject>> LoadedAssetsToProcess;
-
 	/** The set of object paths that have had their disk cache updated from the in memory version */
 	TSet<FSoftObjectPath> AssetDataObjectPathsUpdatedOnLoad;
 
