[AssetRegistry][PackageName]

Extends TryGetAssetPackageData to allow specifying that we allow failure if the AR is currently not available for reading from the current thread.

This allows long stalls to be eliminated when calling TryGetAssetPackageData while the AssetRegistry is gathering when functions can provide their own alternative to calling into the AR for PackageData (e.g. DoesPackageExist can gracefully fall back to an OS call when the AR is contended)

#rnx

#Synced-CL 50826644

diff --git a/Engine/Source/Runtime/AssetRegistry/Private/ARTypes.h b/Engine/Source/Runtime/AssetRegistry/Private/ARTypes.h
--- a/Engine/Source/Runtime/AssetRegistry/Private/ARTypes.h
+++ b/Engine/Source/Runtime/AssetRegistry/Private/ARTypes.h
@@ -34,6 +34,11 @@ public:
 		Mutex.WriteLock();
 	}
 
+	FORCEINLINE bool TryWriteLock()
+	{
+		return Mutex.TryWriteLock();
+	}
+
 	FORCEINLINE void WriteUnlock()
 	{
 		Mutex.WriteUnlock();
@@ -44,6 +49,11 @@ public:
 		Mutex.ReadLock();
 	}
 
+	FORCEINLINE bool TryReadLock()
+	{
+		return Mutex.TryReadLock();
+	}
+
 	FORCEINLINE void ReadUnlock()
 	{
 		Mutex.ReadUnlock();
diff --git a/Engine/Source/Runtime/AssetRegistry/Private/AssetRegistry.cpp b/Engine/Source/Runtime/AssetRegistry/Private/AssetRegistry.cpp
--- a/Engine/Source/Runtime/AssetRegistry/Private/AssetRegistry.cpp
+++ b/Engine/Source/Runtime/AssetRegistry/Private/AssetRegistry.cpp
@@ -252,20 +252,20 @@ public:
 		return AssetRegistry->TryGetAssetByObjectPath(ObjectPath, OutAssetData);
 	}
 
-	virtual UE::AssetRegistry::EExists TryGetAssetPackageData(FName PackageName, class FAssetPackageData& OutPackageData) const override
+	virtual UE::AssetRegistry::EExists TryGetAssetPackageData(FName PackageName, class FAssetPackageData& OutPackageData, bool bFailIfLockHeld) const override
 	{
 		FName OutCorrectCasePackageName;
-		return TryGetAssetPackageData(PackageName, OutPackageData, OutCorrectCasePackageName);
+		return TryGetAssetPackageData(PackageName, OutPackageData, OutCorrectCasePackageName, bFailIfLockHeld);
 	}
 	
-	virtual UE::AssetRegistry::EExists TryGetAssetPackageData(FName PackageName, class FAssetPackageData& OutPackageData, FName& OutCorrectCasePackageName) const override
+	virtual UE::AssetRegistry::EExists TryGetAssetPackageData(FName PackageName, class FAssetPackageData& OutPackageData, FName& OutCorrectCasePackageName, bool bFailIfLockHeld) const override
 	{
 		IAssetRegistry* AssetRegistry = IAssetRegistry::Get();
 		if (!AssetRegistry)
 		{
 			return UE::AssetRegistry::EExists::Unknown;
 		}
-		return AssetRegistry->TryGetAssetPackageData(PackageName, OutPackageData, OutCorrectCasePackageName);
+		return AssetRegistry->TryGetAssetPackageData(PackageName, OutPackageData, OutCorrectCasePackageName, bFailIfLockHeld);
 	}
 
 	virtual bool EnumerateAssets(const FARFilter& Filter, TFunctionRef<bool(const FAssetData&)> Callback,
@@ -2847,15 +2847,27 @@ UE::AssetRegistry::EExists UAssetRegistryImpl::TryGetAssetByObjectPath(const FSo
 	return UE::AssetRegistry::EExists::Exists;
 }
 
-UE::AssetRegistry::EExists UAssetRegistryImpl::TryGetAssetPackageData(const FName PackageName, FAssetPackageData& OutAssetPackageData) const
+UE::AssetRegistry::EExists UAssetRegistryImpl::TryGetAssetPackageData(const FName PackageName, FAssetPackageData& OutAssetPackageData, bool bFailIfLockHeld) const
 {
 	FName OutCorrectCasePackageName;
-	return TryGetAssetPackageData(PackageName, OutAssetPackageData, OutCorrectCasePackageName);
+	return TryGetAssetPackageData(PackageName, OutAssetPackageData, OutCorrectCasePackageName, bFailIfLockHeld);
 }
 
-UE::AssetRegistry::EExists UAssetRegistryImpl::TryGetAssetPackageData(const FName PackageName, FAssetPackageData& OutAssetPackageData, FName& OutCorrectCasePackageName) const
+UE::AssetRegistry::EExists UAssetRegistryImpl::TryGetAssetPackageData(const FName PackageName, FAssetPackageData& OutAssetPackageData, FName& OutCorrectCasePackageName, bool bFailIfLockHeld) const
 {
-	UE::AssetRegistry::FInterfaceReadScopeLock InterfaceScopeLock(InterfaceLock);
+	if (!InterfaceLock.TryReadLock())
+	{
+		if (bFailIfLockHeld)
+		{
+			return UE::AssetRegistry::EExists::Unknown;
+		}
+		InterfaceLock.ReadLock();
+	}
+	ON_SCOPE_EXIT
+	{
+		InterfaceLock.ReadUnlock();
+	};
+
 	const bool bLocalIsGathering = GuardedData.IsGathering();
 	const bool bAssetRegistryHasFullKnowledge = GuardedData.IsSearchAllAssets() && !bLocalIsGathering;
 	const FAssetRegistryState& State = GuardedData.GetState();
diff --git a/Engine/Source/Runtime/AssetRegistry/Private/AssetRegistry.h b/Engine/Source/Runtime/AssetRegistry/Private/AssetRegistry.h
--- a/Engine/Source/Runtime/AssetRegistry/Private/AssetRegistry.h
+++ b/Engine/Source/Runtime/AssetRegistry/Private/AssetRegistry.h
@@ -53,8 +53,8 @@ public:
 	virtual FAssetData GetAssetByObjectPath( const FName ObjectPath, bool bIncludeOnlyOnDiskAssets = false ) const override;
 	virtual FAssetData GetAssetByObjectPath(const FSoftObjectPath& ObjectPath, bool bIncludeOnlyOnDiskAssets = false, bool bSkipARFilteredAssets = true) const override;
 	virtual UE::AssetRegistry::EExists TryGetAssetByObjectPath(const FSoftObjectPath& ObjectPath, FAssetData& OutAssetData) const override;
-	virtual UE::AssetRegistry::EExists TryGetAssetPackageData(const FName PackageName, FAssetPackageData& OutAssetPackageData) const override;
-	virtual UE::AssetRegistry::EExists TryGetAssetPackageData(const FName PackageName, FAssetPackageData& OutAssetPackageData, FName& OutCorrectCasePackageName) const override;
+	virtual UE::AssetRegistry::EExists TryGetAssetPackageData(const FName PackageName, FAssetPackageData& OutAssetPackageData, bool bFailIfLockHeld = false) const override;
+	virtual UE::AssetRegistry::EExists TryGetAssetPackageData(const FName PackageName, FAssetPackageData& OutAssetPackageData, FName& OutCorrectCasePackageName, bool bFailIfLockHeld = false) const override;
 	virtual bool GetAllAssets(TArray<FAssetData>& OutAssetData, bool bIncludeOnlyOnDiskAssets = false) const override;
 	virtual bool EnumerateAllAssets(TFunctionRef<bool(const FAssetData&)> Callback,
 		UE::AssetRegistry::EEnumerateAssetsFlags InEnumerateFlags = UE::AssetRegistry::EEnumerateAssetsFlags::None) const override;
diff --git a/Engine/Source/Runtime/AssetRegistry/Public/AssetRegistry/AssetRegistryModule.h b/Engine/Source/Runtime/AssetRegistry/Public/AssetRegistry/AssetRegistryModule.h
--- a/Engine/Source/Runtime/AssetRegistry/Public/AssetRegistry/AssetRegistryModule.h
+++ b/Engine/Source/Runtime/AssetRegistry/Public/AssetRegistry/AssetRegistryModule.h
@@ -99,20 +99,20 @@ public:
 		return AssetRegistry->TryGetAssetByObjectPath(ObjectPath, OutAssetData);
 	}
 
-	virtual UE::AssetRegistry::EExists TryGetAssetPackageData(FName PackageName, FAssetPackageData& OutAssetPackageData) const override
+	virtual UE::AssetRegistry::EExists TryGetAssetPackageData(FName PackageName, FAssetPackageData& OutAssetPackageData, bool bFailIfLockHeld) const override
 	{
 		FName OutCorrectCasePackageName;
-		return TryGetAssetPackageData(PackageName, OutAssetPackageData, OutCorrectCasePackageName);
+		return TryGetAssetPackageData(PackageName, OutAssetPackageData, OutCorrectCasePackageName, bFailIfLockHeld);
 	}
 	
-	virtual UE::AssetRegistry::EExists TryGetAssetPackageData(FName PackageName, FAssetPackageData& OutAssetPackageData, FName& OutCorrectCasePackageName) const override
+	virtual UE::AssetRegistry::EExists TryGetAssetPackageData(FName PackageName, FAssetPackageData& OutAssetPackageData, FName& OutCorrectCasePackageName, bool bFailIfLockHeld) const override
 	{
 		IAssetRegistry* AssetRegistry = IAssetRegistry::Get();
 		if (!AssetRegistry)
 		{
 			return UE::AssetRegistry::EExists::Unknown;
 		}
-		return AssetRegistry->TryGetAssetPackageData(PackageName, OutAssetPackageData, OutCorrectCasePackageName);
+		return AssetRegistry->TryGetAssetPackageData(PackageName, OutAssetPackageData, OutCorrectCasePackageName, bFailIfLockHeld);
 	}
 
 	virtual bool EnumerateAssets(const FARFilter& Filter, TFunctionRef<bool(const FAssetData&)> Callback,
diff --git a/Engine/Source/Runtime/AssetRegistry/Public/AssetRegistry/IAssetRegistry.h b/Engine/Source/Runtime/AssetRegistry/Public/AssetRegistry/IAssetRegistry.h
--- a/Engine/Source/Runtime/AssetRegistry/Public/AssetRegistry/IAssetRegistry.h
+++ b/Engine/Source/Runtime/AssetRegistry/Public/AssetRegistry/IAssetRegistry.h
@@ -439,10 +439,12 @@ public:
 	 *
 	 * @param PackageName name of the package
 	 * @param OutAssetPackageData out FAssetPackageData
+	 * @param bFailIfLockHeld If the AssetRegistry is currently being written to by another thread this 
+	 *	      function will return EExists::Unknown rather than waiting for safe read access.
 	 * @return Enum return code
 	 */
 	virtual UE::AssetRegistry::EExists TryGetAssetPackageData(FName PackageName,
-		FAssetPackageData& OutAssetPackageData) const = 0;
+		FAssetPackageData& OutAssetPackageData, bool bFailIfLockHeld = false) const = 0;
 
 	/**
 	 * Tries to get the package data for the specified package name. If found, OutCorrectCasePackageName
@@ -451,10 +453,12 @@ public:
 	 * @param PackageName name of the package
 	 * @param OutAssetPackageData out FAssetPackageData
 	 * @param OutCorrectCasePackageName out FName matching filesystem casing
+	 * @param bFailIfLockHeld If the AssetRegistry is currently being written to by another thread this 
+	 *	      function will return EExists::Unknown rather than waiting for safe read access.
 	 * @return Return code enum
 	 */
 	virtual UE::AssetRegistry::EExists TryGetAssetPackageData(FName PackageName,
-		FAssetPackageData& OutAssetPackageData, FName& OutCorrectCasePackageName) const = 0;
+		FAssetPackageData& OutAssetPackageData, FName& OutCorrectCasePackageName, bool bFailIfLockHeld = false) const = 0;
 	
 	/**
 	 * Gets asset data for all assets in the registry.
diff --git a/Engine/Source/Runtime/CoreUObject/Private/Misc/PackageName.cpp b/Engine/Source/Runtime/CoreUObject/Private/Misc/PackageName.cpp
--- a/Engine/Source/Runtime/CoreUObject/Private/Misc/PackageName.cpp
+++ b/Engine/Source/Runtime/CoreUObject/Private/Misc/PackageName.cpp
@@ -2461,7 +2461,7 @@ FPackageName::EPackageLocationFilter FPackageName::DoesPackageExistEx(const FPac
 		FName CorrectCasePackageName;
 		FAssetPackageData AssetPackageData;
 		FName PackageName = PackagePath.GetPackageFName();
-		UE::AssetRegistry::EExists ExistsStatus = AssetRegistry->TryGetAssetPackageData(PackageName, AssetPackageData, CorrectCasePackageName);
+		UE::AssetRegistry::EExists ExistsStatus = AssetRegistry->TryGetAssetPackageData(PackageName, AssetPackageData, CorrectCasePackageName, true /*bFailIfLockHeld*/);
 		if (ExistsStatus == UE::AssetRegistry::EExists::Exists)
 		{
 			EPackageLocationFilter PackageLocation = AssetPackageData.GetPackageLocation();
diff --git a/Engine/Source/Runtime/CoreUObject/Public/Misc/AssetRegistryInterface.h b/Engine/Source/Runtime/CoreUObject/Public/Misc/AssetRegistryInterface.h
--- a/Engine/Source/Runtime/CoreUObject/Public/Misc/AssetRegistryInterface.h
+++ b/Engine/Source/Runtime/CoreUObject/Public/Misc/AssetRegistryInterface.h
@@ -443,9 +443,11 @@ public:
 	 *
 	 * @param PackageName name of the package
 	 * @param OutAssetPackageData out FAssetPackageData
+	 * @param bFailIfLockHeld If the AssetRegistry is currently being written to by another thread this
+	 *	      function will return EExists::Unknown rather than waiting for safe read access.
 	 * @return Return code enum
 	 */
-	virtual UE::AssetRegistry::EExists TryGetAssetPackageData(FName PackageName, class FAssetPackageData& OutPackageData) const = 0;
+	virtual UE::AssetRegistry::EExists TryGetAssetPackageData(FName PackageName, class FAssetPackageData& OutPackageData, bool bFailIfLockHeld = false) const = 0;
 
 	/**
 	 * Tries to get the package data for the specified package name. If found, OutCorrectCasePackageName
@@ -454,9 +456,11 @@ public:
 	 * @param PackageName name of the package
 	 * @param OutAssetPackageData out FAssetPackageData
 	 * @param OutCorrectCasePackageName out FName matching filesystem casing
+	 * @param bFailIfLockHeld If the AssetRegistry is currently being written to by another thread this 
+	 *	      function will return EExists::Unknown rather than waiting for safe read access.
 	 * @return Return code enum
 	 */
-	virtual UE::AssetRegistry::EExists TryGetAssetPackageData(FName PackageName, class FAssetPackageData& OutPackageData, FName& OutCorrectCasePackageName) const = 0;
+	virtual UE::AssetRegistry::EExists TryGetAssetPackageData(FName PackageName, class FAssetPackageData& OutPackageData, FName& OutCorrectCasePackageName, bool bFailIfLockHeld = false) const = 0;
 
 	/**
 	 * Enumerate asset data for all assets that match the filter.
