Enable AR cache memory mapping on buildmachines if cache files aren't being overwritten.

#Synced-CL 50902315

#Synced-CL 50902315

diff --git a/Engine/Source/Runtime/AssetRegistry/Private/AssetDataGatherer.cpp b/Engine/Source/Runtime/AssetRegistry/Private/AssetDataGatherer.cpp
--- a/Engine/Source/Runtime/AssetRegistry/Private/AssetDataGatherer.cpp
+++ b/Engine/Source/Runtime/AssetRegistry/Private/AssetDataGatherer.cpp
@@ -377,7 +377,7 @@ struct FCachePayload
 
 void SerializeCacheSave(FAssetRegistryWriter& Ar, const TArray<TPair<FName, FDiskCachedAssetData*>>& AssetsToSave);
 FCachePayload SerializeCacheLoad(FAssetRegistryReader& Ar);
-TArray<FCachePayload> LoadCacheFiles(TConstArrayView<FString> CacheFilenames);
+TArray<FCachePayload> LoadCacheFiles(TConstArrayView<FString> CacheFilenames, bool bCacheWriteEnabled);
 
 
 /** InOutResult = Value, but without shrinking the string to fit. */
@@ -3765,6 +3765,8 @@ private:
 	{
 		GGatherSettings.Initialize();
 
+		bCacheWriteEnabled = GGatherSettings.IsGatherCacheWriteEnabled();
+
 		if (GGatherSettings.IsPreloadDiscoveryCache())
 		{
 			Discovery = MakeUnique<FAssetDataDiscovery>();
@@ -3784,7 +3786,7 @@ private:
 	void LoadAsync(const TArray<FString>& Paths)
 	{
 		LLM_SCOPE(ELLMTag::AssetRegistry);
-		GatherCachePayloads = LoadCacheFiles(Paths);
+		GatherCachePayloads = LoadCacheFiles(Paths, bCacheWriteEnabled);
 	}
 
 	TFuture<void> PreloadReady;
@@ -3793,6 +3795,7 @@ private:
 	bool bInitialized = false;
 	bool bConsumedGatherCache = false;
 	bool bConsumedDiscoveryCache = false;
+	bool bCacheWriteEnabled = false;
 };
 FPreloader GPreloader;
 #endif
@@ -4392,7 +4395,9 @@ FAssetDataGatherer::ETickResult FAssetDataGatherer::TickInternal(double& TickSta
 	{
 		double CacheLoadStartTime = FPlatformTime::Seconds();
 		TArray<FCachePayload> Payloads;
+		bool bGatherCacheWriteEnabled = false;
 #if UE_EDITOR
+		bGatherCacheWriteEnabled = GGatherSettings.IsGatherCacheWriteEnabled();
 		if (GGatherSettings.IsPreloadGatherCache())
 		{
 			Payloads = GPreloader.Consume();
@@ -4401,7 +4406,7 @@ FAssetDataGatherer::ETickResult FAssetDataGatherer::TickInternal(double& TickSta
 #endif
 		{
 			TArray<FString> CachePaths = GGatherSettings.FindShardedCacheFiles();
-			Payloads = UE::AssetDataGather::Private::LoadCacheFiles(CachePaths);
+			Payloads = UE::AssetDataGather::Private::LoadCacheFiles(CachePaths, bGatherCacheWriteEnabled);
 		}
 		ConsumeCacheFiles(MoveTemp(Payloads));
 		UE_LOG(LogAssetRegistry, Display, TEXT("AssetDataGatherer spent %.3fs loading caches %s_*.bin."),
@@ -5414,14 +5419,14 @@ FCachePayload SerializeCacheLoad(FAssetRegistryReader& Ar)
 	return Result;
 }
 
-TArray<FCachePayload> LoadCacheFiles(TConstArrayView<FString> InCacheFilenames)
+TArray<FCachePayload> LoadCacheFiles(TConstArrayView<FString> InCacheFilenames, bool bCacheWriteEnabled)
 {
 	TRACE_CPUPROFILER_EVENT_SCOPE(LoadCacheFiles);
 	TArray<FCachePayload> Results;
 	Results.AddDefaulted(InCacheFilenames.Num());
 
-	// disable memory mapping on the builder machines until the situation with deletion clarifies
-	if (GIsBuildMachine)
+	// when writing cache files, disable memory mapping on the builder machines until the situation with deletion clarifies
+	if (GIsBuildMachine && bCacheWriteEnabled)
 	{
 		UE::AssetDataGather::Private::bMemoryMappedCacheFiles = false;
 	}
