AssetRegistry: Cooking: Add some functions to modify dependencies and packagedatas for use by cooker's AssetRegistry generator in an upcoming change to incrementally update the AssetRegistry.bin cook artifacts.
#rnx
#Synced-CL 51045425

diff --git a/Engine/Source/Runtime/AssetRegistry/Private/AssetRegistry.cpp b/Engine/Source/Runtime/AssetRegistry/Private/AssetRegistry.cpp
--- a/Engine/Source/Runtime/AssetRegistry/Private/AssetRegistry.cpp
+++ b/Engine/Source/Runtime/AssetRegistry/Private/AssetRegistry.cpp
@@ -6508,6 +6508,40 @@ void UAssetRegistryImpl::InitializeTemporaryAssetRegistryState(FAssetRegistrySta
 	}
 }
 
+void UAssetRegistryImpl::CopyManageDependenciesToState(FAssetRegistryState& InOutState) const
+{
+	using namespace UE::AssetRegistry;
+
+	UE::AssetRegistry::FInterfaceReadScopeLock InterfaceScopeLock(InterfaceLock);
+	const FAssetRegistryState& State = GuardedData.GetState();
+	TArray<FAssetDependency> Dependencies;
+	for (const TPair<FAssetIdentifier, FDependsNode*>& Pair : State.CachedDependsNodes)
+	{
+		FDependsNode* ManagerNode = Pair.Value;
+		Dependencies.Reset();
+		ManagerNode->GetDependencies(Dependencies, EDependencyCategory::Manage);
+		// Remove dependencies to packages that do not exist in InOutState, unless they are PrimaryAssets. The
+		// AssetRegistryGenerator's state does not pre-add the PrimaryAsset nodes, this function needs to add them.
+		// AddDependencies will create the node if it does not already exist.
+		Dependencies.RemoveAllSwap([&InOutState](const FAssetDependency& AssetDependency)
+			{
+				if (AssetDependency.AssetId.GetPrimaryAssetId().IsValid())
+				{
+					return false;
+				}
+				if (InOutState.CachedPackageData.Contains(AssetDependency.AssetId.PackageName))
+				{
+					return false;
+				}
+				return true;
+			});
+		if (!Dependencies.IsEmpty())
+		{
+			InOutState.AddDependencies(Pair.Key, Dependencies);
+		}
+	}
+}
+
 #if ASSET_REGISTRY_STATE_DUMPING_ENABLED
 void UAssetRegistryImpl::DumpState(const TArray<FString>& Arguments, TArray<FString>& OutPages, int32 LinesPerPage) const
 {
diff --git a/Engine/Source/Runtime/AssetRegistry/Private/AssetRegistry.h b/Engine/Source/Runtime/AssetRegistry/Private/AssetRegistry.h
--- a/Engine/Source/Runtime/AssetRegistry/Private/AssetRegistry.h
+++ b/Engine/Source/Runtime/AssetRegistry/Private/AssetRegistry.h
@@ -245,6 +245,7 @@ protected:
 		TSet<FDependsNode*>& ExistingManagedNode, ShouldSetManagerPredicate ShouldSetManager = nullptr) override;
 	virtual void SetManageReferences(UE::AssetRegistry::FSetManageReferencesContext& Context) override;
 	virtual bool SetPrimaryAssetIdForObjectPath(const FSoftObjectPath& ObjectPath, FPrimaryAssetId PrimaryAssetId) override;
+	virtual void CopyManageDependenciesToState(FAssetRegistryState& InOutState) const override;
 
 private:
 	void OnPreExit();
diff --git a/Engine/Source/Runtime/AssetRegistry/Private/AssetRegistryState.cpp b/Engine/Source/Runtime/AssetRegistry/Private/AssetRegistryState.cpp
--- a/Engine/Source/Runtime/AssetRegistry/Private/AssetRegistryState.cpp
+++ b/Engine/Source/Runtime/AssetRegistry/Private/AssetRegistryState.cpp
@@ -17,6 +17,7 @@
 #include "Containers/MpscQueue.h"
 #include "DependsNode.h"
 #include "GenericPlatform/GenericPlatformChunkInstall.h"
+#include "Logging/StructuredLog.h"
 #include "Misc/CommandLine.h"
 #include "Misc/ConfigCacheIni.h"
 #include "Misc/FileHelper.h"
@@ -4098,10 +4099,12 @@ bool FAssetRegistryState::RemoveDependsNode(const FAssetIdentifier& Identifier)
 	return false;
 }
 
-bool FAssetRegistryState::RemoveDependsNodes(const TSet<FAssetIdentifier>& Identifiers)
+bool FAssetRegistryState::RemoveDependsNodes(const TSet<FAssetIdentifier>& Identifiers,
+	ERemoveDependsFlags Flags)
 {
 	using namespace UE::Tasks;
 
+	const bool bRemoveOrphans = EnumHasAnyFlags(Flags, ERemoveDependsFlags::RemoveOrphans);
 	TSet<FDependsNode*> NodesToRemove;
 	NodesToRemove.Reserve(Identifiers.Num());
 
@@ -4121,8 +4124,10 @@ bool FAssetRegistryState::RemoveDependsNodes(const TSet<FAssetIdentifier>& Ident
 		}
 	}
 
+	TMpscQueue<FDependsNode*> OrphansToRemoveQueue;
 	TMap<FDependsNode*, TSet<FDependsNode*>> RemovedReferences;
-	FTask PopulateTask = Launch(TEXT("RemoveDependsNode_PopulateSchedule"), [&NodesToRemove, &RemovedReferences]()
+	FTask PopulateTask = Launch(TEXT("RemoveDependsNode_PopulateSchedule"),
+		[&NodesToRemove, &RemovedReferences, bRemoveOrphans, &OrphansToRemoveQueue]()
 		{
 			for (FDependsNode* Node : NodesToRemove)
 			{
@@ -4136,6 +4141,10 @@ bool FAssetRegistryState::RemoveDependsNodes(const TSet<FAssetIdentifier>& Ident
 					if (!NodesToRemove.Contains(ReferencerNode))
 					{
 						ReferencerNode->RemoveDependency(Node);
+						if (bRemoveOrphans && ReferencerNode->GetConnectionCount() == 0)
+						{
+							OrphansToRemoveQueue.Enqueue(ReferencerNode);
+						}
 					}
 				}
 			}
@@ -4158,23 +4167,29 @@ bool FAssetRegistryState::RemoveDependsNodes(const TSet<FAssetIdentifier>& Ident
 			}
 		}, ETaskPriority::Default, EExtendedTaskPriority::None);
 
-	FTask RemoveTask = Launch(TEXT("RemoveDependsNode_RemoveReferencerSchedule"), [&RemovedReferences]()
+	FTask RemoveTask = Launch(TEXT("RemoveDependsNode_RemoveReferencerSchedule"),
+		[&RemovedReferences, bRemoveOrphans, &OrphansToRemoveQueue]()
 		{
 			TArray<FDependsNode*> DependenciesToUpdate;
 			RemovedReferences.GetKeys(DependenciesToUpdate);
 			ParallelFor(TEXT("RemoveDependsNode_RemoveReferencerParallel"), DependenciesToUpdate.Num(), 1, 
-				[&DependenciesToUpdate, &RemovedReferences](int32 Index)
+				[&DependenciesToUpdate, &RemovedReferences, bRemoveOrphans, &OrphansToRemoveQueue](int32 Index)
 				{
 					FDependsNode* DependencyNode = DependenciesToUpdate[Index];
 					TSet<FDependsNode*>* Nodes = RemovedReferences.Find(DependencyNode);
 					DependencyNode->RemoveReferencers(*Nodes);
+					if (bRemoveOrphans && DependencyNode->GetConnectionCount() == 0)
+					{
+						OrphansToRemoveQueue.Enqueue(DependencyNode);
+					}
 				}, EParallelForFlags::Unbalanced);
 		}, PopulateTask);
 
 	// If the depends nodes were preallocated in a block, we can't delete 
 	// them one at a time, only the whole chunk in the destructor
+	const bool bDeleteRemovedNodes = PreallocatedDependsNodeDataBuffers.Num() == 0;
 	FTask DeleteTask;
-	if (PreallocatedDependsNodeDataBuffers.Num() == 0)
+	if (bDeleteRemovedNodes)
 	{
 		DeleteTask = Launch(TEXT("RemoveDependsNode_DeleteNodes"), [&NodesToRemove]()
 			{
@@ -4189,6 +4204,37 @@ bool FAssetRegistryState::RemoveDependsNodes(const TSet<FAssetIdentifier>& Ident
 	RemoveTask.Wait();
 	DeleteTask.Wait();
 
+	if (!OrphansToRemoveQueue.IsEmpty())
+	{
+		// It should be impossible to have duplicates in OrphansToRemove, but guarantee that with Algo::Unique.
+		TArray<FDependsNode*> OrphansToRemove;
+		FDependsNode* Orphan;
+		while (OrphansToRemoveQueue.Dequeue(Orphan))
+		{
+			OrphansToRemove.Add(Orphan);
+		}
+		Algo::Sort(OrphansToRemove);
+		int32 PreUniqueNum = OrphansToRemove.Num();
+		OrphansToRemove.SetNum(Algo::Unique(OrphansToRemove));
+		if (OrphansToRemove.Num() != PreUniqueNum)
+		{
+			UE_LOGFMT(LogAssetRegistry, Error,
+				"Logic error in FAssetRegistryState::RemoveDependsNodes, OrphansToRemove had duplicates. PreUniqueNum={PreUniqueNum}, UniqueNum={UniqueNum}.",
+				PreUniqueNum, OrphansToRemove.Num());
+		}
+
+		// We had to wait for RemoveTask to complete, which was running in parallel with DeleteTask, so we
+		// can't lump these in to DeleteTask. Do a pass over them now to remove and delete them.
+		for (FDependsNode* Node : OrphansToRemove)
+		{
+			int32 NumRemoved = CachedDependsNodes.Remove(Node->GetIdentifier());
+			NumDependsNodes -= NumRemoved;
+			if (bDeleteRemovedNodes)
+			{
+				delete Node;
+			}
+		}
+	}
 	return !NodesToRemove.IsEmpty();
 }
 
@@ -4252,6 +4298,12 @@ bool FAssetRegistryState::RemovePackageData(FName PackageName)
 	return RemovePackageData(PackageNames);
 }
 
+bool FAssetRegistryState::RemovePackageDatas(const TSet<FAssetIdentifier>& PackageNames)
+{
+	// Call the templated function (templated based on TSet allocator) for the default TSet allocator.
+	return RemovePackageData(PackageNames);
+}
+
 template<typename SetKeyFuncs, typename SetAllocator>
 bool FAssetRegistryState::RemovePackageData(const TSet<FAssetIdentifier, SetKeyFuncs, SetAllocator>& AssetIdentifiers)
 {
diff --git a/Engine/Source/Runtime/AssetRegistry/Public/AssetRegistry/AssetRegistryState.h b/Engine/Source/Runtime/AssetRegistry/Public/AssetRegistry/AssetRegistryState.h
--- a/Engine/Source/Runtime/AssetRegistry/Public/AssetRegistry/AssetRegistryState.h
+++ b/Engine/Source/Runtime/AssetRegistry/Public/AssetRegistry/AssetRegistryState.h
@@ -15,6 +15,7 @@
 #include "CoreTypes.h"
 #include "HAL/PlatformCrt.h"
 #include "Misc/AssetRegistryInterface.h"
+#include "Misc/EnumClassFlags.h"
 #include "Templates/Function.h"
 #include "Templates/UnrealTemplate.h"
 #include "UObject/NameTypes.h"
@@ -446,6 +447,9 @@ public:
 	/** Removes existing package data */
 	ASSETREGISTRY_API bool RemovePackageData(FName PackageName);
 
+	/** Removes existing package datas, passed in as FAssetIdentifier(PackageName) for internal performance. */
+	ASSETREGISTRY_API bool RemovePackageDatas(const TSet<FAssetIdentifier>& PackageNames);
+
 	/** Adds the asset data to the lookup maps */
 	ASSETREGISTRY_API void AddAssetData(FAssetData* AssetData);
 
@@ -666,11 +670,23 @@ private:
 	/** Creates a node in the CachedDependsNodes map or finds the existing node and returns it */
 	FDependsNode* CreateOrFindDependsNode(const FAssetIdentifier& Identifier, bool* bAlreadyExists = nullptr);
 
-	/** Removes the depends node and updates the dependencies to no longer contain it as as a referencer. */
-	bool RemoveDependsNode(const FAssetIdentifier& Identifier);
+	/** Removes the depends node and updates the dependencies to no longer contain it as a referencer. */
+	ASSETREGISTRY_API bool RemoveDependsNode(const FAssetIdentifier& Identifier);
 
-	/** Removes a set of depends nodes and updates their dependencies to no longer contain them as as a referencer. */
-	bool RemoveDependsNodes(const TSet<FAssetIdentifier>& Identifiers);
+	enum class ERemoveDependsFlags
+	{
+		None,
+		/**
+		 * If passed in, RemoveDependsNode functions will remove nodes that were not in the remove request
+		 * but that had links to the removed nodes and no longer have any links after the removals.
+		 */
+		RemoveOrphans = 1 << 0,
+	};
+	FRIEND_ENUM_CLASS_FLAGS(ERemoveDependsFlags);
+
+	/** Removes a set of depends nodes and updates their dependencies to no longer contain them as a referencer. */
+	ASSETREGISTRY_API bool RemoveDependsNodes(const TSet<FAssetIdentifier>& Identifiers,
+		ERemoveDependsFlags Flags = ERemoveDependsFlags::None);
 
 	/** Removes a set of existing package data. Returns true if any were removed, false if all already did not exist. */
 	template<typename SetKeyFuncs, typename SetAllocator>
@@ -798,6 +814,8 @@ private:
 ///////////////////////////////////////////////////////
 
 
+ENUM_CLASS_FLAGS(FAssetRegistryState::ERemoveDependsFlags);
+
 inline FAssetRegistrySerializationOptions::FAssetRegistrySerializationOptions(
 	UE::AssetRegistry::ESerializationTarget Target)
 {
diff --git a/Engine/Source/Runtime/AssetRegistry/Public/AssetRegistry/IAssetRegistry.h b/Engine/Source/Runtime/AssetRegistry/Public/AssetRegistry/IAssetRegistry.h
--- a/Engine/Source/Runtime/AssetRegistry/Public/AssetRegistry/IAssetRegistry.h
+++ b/Engine/Source/Runtime/AssetRegistry/Public/AssetRegistry/IAssetRegistry.h
@@ -20,6 +20,7 @@
 #endif
 
 class FArchive;
+class FAssetRegistryGenerator;
 class FAssetRegistryState;
 class FCbFieldView;
 class FCbWriter;
@@ -1315,6 +1316,16 @@ protected:
 	 * when the AssetManager is more up to date than the on disk Registry.
 	 */
 	virtual bool SetPrimaryAssetIdForObjectPath(const FSoftObjectPath& ObjectPath, FPrimaryAssetId PrimaryAssetId) = 0;
+
+	// Functions specifically for calling from cooker's FAssetRegistryGenerator
+	friend class FAssetRegistryGenerator;
+
+	/**
+	 * Copy all manage dependencies in the global AssetRegistry that target any assets that exist in InOutState into
+	 * the manage dependencies of InOutState.
+	 */
+	virtual void CopyManageDependenciesToState(FAssetRegistryState& InOutState) const = 0;
+
 };
 
 namespace UE::AssetRegistry
