AssetRegistry: CookDeterminism: Fix incorrect capitalization of DependsNodes in the AssetRegistry (which shows up as capitalization indeterminism in the cooked AssetRegistry) when a sourcepackage referencing a targetpackage with a captialization that does not match the targetpackage's capitalization on disk is added to the assetregistry before the targetpackage is.
#rnx

#Synced-CL 51045425

diff --git a/Engine/Source/Runtime/AssetRegistry/Private/AssetRegistry.cpp b/Engine/Source/Runtime/AssetRegistry/Private/AssetRegistry.cpp
--- a/Engine/Source/Runtime/AssetRegistry/Private/AssetRegistry.cpp
+++ b/Engine/Source/Runtime/AssetRegistry/Private/AssetRegistry.cpp
@@ -7644,7 +7644,8 @@ void FAssetRegistryImpl::DependencyDataGathered(TMultiMap<FName, FPackageDepende
 					}
 				}
 
-				FDependsNode* SourceNode = State.CreateOrFindDependsNode(SourceNodePackageName, &bSourceNodeAlreadyExisted);
+				FAssetIdentifier SourceNodeIdentifier(SourceNodePackageName);
+				FDependsNode* SourceNode = State.CreateOrFindDependsNode(SourceNodeIdentifier, &bSourceNodeAlreadyExisted);
 				if (!bSourceNodeAlreadyExisted)
 				{
 					if (FPackageDependencyData* DependencyDependencyData = DependsResults.Find(SourceNodePackageName))
@@ -7652,6 +7653,19 @@ void FAssetRegistryImpl::DependencyDataGathered(TMultiMap<FName, FPackageDepende
 						SourceNode->Reserve(DependencyDependencyData->DependsNodeReservations);
 					}
 				}
+				else
+				{
+					// The node was created by an earlier package registering SourceNode as a dependency, and it may
+					// have had a mismatching capitalization that was saved into the package with the dependency.
+					// Force the capitalization of the DependsNode for SourceNodePackage to match SourceNodePackage's
+					// capitalization (which comes from its filepath on disk) if they are different.
+					if (SourceNode->GetPackageName().GetDisplayIndex() != SourceNodePackageName.GetDisplayIndex())
+					{
+						State.CachedDependsNodes.Remove(SourceNode->GetIdentifier());
+						SourceNode->SetIdentifier(SourceNodeIdentifier);
+						State.CachedDependsNodes.Add(SourceNode->GetIdentifier(), SourceNode);
+					}
+				}
 
 				if (Result.bHasDependencyData)
 				{
