Reduce FPackageDependencyData size by 8 bytes (saving ~5%, 184->176) by lesser padding.

#Synced-CL 50902315

diff --git a/Engine/Source/Runtime/AssetRegistry/Private/PackageDependencyData.h b/Engine/Source/Runtime/AssetRegistry/Private/PackageDependencyData.h
--- a/Engine/Source/Runtime/AssetRegistry/Private/PackageDependencyData.h
+++ b/Engine/Source/Runtime/AssetRegistry/Private/PackageDependencyData.h
@@ -41,18 +41,18 @@ public:
 		}
 	};
 
-	/** The name of the package that dependency data is gathered from */
-	FName PackageName;
-
-	/** Asset Package data, gathered at the same time as dependency data */
-	FAssetPackageData PackageData;
-
 	// Dependency Data
 	TArray<FPackageDependency> PackageDependencies;
 	TArray<FSearchableNamesDependency> SearchableNameDependencies;
 
 	FDependsNodeReservations DependsNodeReservations;
 
+	/** Asset Package data, gathered at the same time as dependency data */
+	FAssetPackageData PackageData;
+
+	/** The name of the package that dependency data is gathered from */
+	FName PackageName;
+
 	// Transient Flags indicating which types of data have been gathered
 	bool bHasPackageData = false;
 	bool bHasDependencyData = false;
