AssetRegistry ActiveMountsLoading:
Add ActiveMountsReloadAllAsync and ActiveMountsUnloadAllUnmountedAsync. This is useful for projects that need all AssetRegistry data loaded in an initial menu but that want to remove all unneeded data once gameplay starts.
Update GameFeaturePluginStateMachine AssetRegistry api calls to support working with ActiveMountsLoader when it is enabled.
Fixes from testing: fix some missing broadcast calls, and downgrade AssetRegistry unrecognized mountpoints error to warning since it can occur due to a content bug during cook that is only a warning.
#rnx

#Synced-CL 51179091

diff --git a/Engine/Plugins/Runtime/GameFeatures/Source/GameFeatures/Private/GameFeaturePluginStateMachine.cpp b/Engine/Plugins/Runtime/GameFeatures/Source/GameFeatures/Private/GameFeaturePluginStateMachine.cpp
--- a/Engine/Plugins/Runtime/GameFeatures/Source/GameFeatures/Private/GameFeaturePluginStateMachine.cpp
+++ b/Engine/Plugins/Runtime/GameFeatures/Source/GameFeatures/Private/GameFeaturePluginStateMachine.cpp
@@ -2792,11 +2792,9 @@ struct FGameFeaturePluginState_Mounting : public FGameFeaturePluginState
 
 		if (!UseAsyncLoading())
 		{
-			FAssetRegistryState PluginAssetRegistryState;
-			if (FAssetRegistryState::LoadFromDisk(*PluginAssetRegistry, FAssetRegistryLoadOptions(), PluginAssetRegistryState))
+			IAssetRegistry& AssetRegistry = UAssetManager::Get().GetAssetRegistry();
+			if (AssetRegistry.ActiveMountsRegisterAndLoadFilePathSynchronous(PluginAssetRegistry))
 			{
-				IAssetRegistry& AssetRegistry = UAssetManager::Get().GetAssetRegistry();
-				AssetRegistry.AppendState(PluginAssetRegistryState);
 				RefreshPackageLocalizationCacheForPlugin();
 			}
 			else
@@ -2809,41 +2807,59 @@ struct FGameFeaturePluginState_Mounting : public FGameFeaturePluginState
 		}
 
 		const bool bForceSyncAssetRegistryAppend = UE::GameFeatures::CVarForceSyncAssetRegistryAppend.GetValueOnGameThread();
-		UE::Tasks::Launch(UE_SOURCE_LOCATION, [this, PluginAssetRegistry=MoveTemp(PluginAssetRegistry), bForceSyncAssetRegistryAppend, RefreshPackageLocalizationCacheForPlugin]
+		if (bForceSyncAssetRegistryAppend)
 		{
-			bool bSuccess = false;
-			TSharedPtr<FAssetRegistryState> PluginAssetRegistryState = MakeShared<FAssetRegistryState>();
-			if (FAssetRegistryState::LoadFromDisk(*PluginAssetRegistry, FAssetRegistryLoadOptions(), *PluginAssetRegistryState))
+			// ActiveMountsRegisterAndLoadFilePath does not support bForceSyncAssetRegistryAppend; we have to load the AssetRegistry file
+			// ourselves into an FAssetRegistryState to be able to append it on the game thread.
+			UE::Tasks::Launch(UE_SOURCE_LOCATION, [this, PluginAssetRegistry=MoveTemp(PluginAssetRegistry), RefreshPackageLocalizationCacheForPlugin]
 			{
-				IAssetRegistry& AssetRegistry = UAssetManager::Get().GetAssetRegistry();
-				if (!bForceSyncAssetRegistryAppend)
-				{
-					AssetRegistry.AppendState(*PluginAssetRegistryState);
-					RefreshPackageLocalizationCacheForPlugin();
-				}
-				bSuccess = true;
-			}
+				TSharedPtr<FAssetRegistryState> PluginAssetRegistryState = MakeShared<FAssetRegistryState>();
+				bool bSuccess = FAssetRegistryState::LoadFromDisk(*PluginAssetRegistry, FAssetRegistryLoadOptions(),
+					*PluginAssetRegistryState);
 
-			ExecuteOnGameThread(UE_SOURCE_LOCATION, [this, PluginAssetRegistryState, bSuccess, bForceSyncAssetRegistryAppend, RefreshPackageLocalizationCacheForPlugin]
-			{
-				TRACE_CPUPROFILER_EVENT_SCOPE(GFP_Mounting_ARComplete);
-
-				if (!bSuccess)
-				{
-					Result = GetErrorResult(TEXT("Failed_To_Load_Plugin_AssetRegistry"));
-				}
-				else if (bForceSyncAssetRegistryAppend)
+				ExecuteOnGameThread(UE_SOURCE_LOCATION, [this, PluginAssetRegistryState, bSuccess, RefreshPackageLocalizationCacheForPlugin]
 				{
-					IAssetRegistry& AssetRegistry = UAssetManager::Get().GetAssetRegistry();
-					AssetRegistry.AppendState(*PluginAssetRegistryState);
-					RefreshPackageLocalizationCacheForPlugin();
-				}
+					TRACE_CPUPROFILER_EVENT_SCOPE(GFP_Mounting_ARComplete);
 
-				CompletedSubStates |= ESubState::LoadAssetRegistry;
-				UpdateStateMachineImmediate();
-			});
+					if (!bSuccess)
+					{
+						Result = GetErrorResult(TEXT("Failed_To_Load_Plugin_AssetRegistry"));
+					}
+					else
+					{
+						IAssetRegistry& AssetRegistry = UAssetManager::Get().GetAssetRegistry();
+						AssetRegistry.AppendState(*PluginAssetRegistryState);
+						RefreshPackageLocalizationCacheForPlugin();
+					}
+
+					CompletedSubStates |= ESubState::LoadAssetRegistry;
+					UpdateStateMachineImmediate();
+				});
 
-		}, UE::Tasks::ETaskPriority::BackgroundHigh);
+			}, UE::Tasks::ETaskPriority::BackgroundHigh);
+		}
+		else
+		{
+			IAssetRegistry& AssetRegistry = UAssetManager::Get().GetAssetRegistry();
+			AssetRegistry.ActiveMountsRegisterAndLoadFilePathAsync(PluginAssetRegistry,
+				[this, RefreshPackageLocalizationCacheForPlugin](bool bSuccess)
+				{
+					if (bSuccess)
+					{
+						RefreshPackageLocalizationCacheForPlugin();
+					}
+					ExecuteOnGameThread(UE_SOURCE_LOCATION, [this, bSuccess]
+						{
+							TRACE_CPUPROFILER_EVENT_SCOPE(GFP_Mounting_ARComplete);
+							if (!bSuccess)
+							{
+								Result = GetErrorResult(TEXT("Failed_To_Load_Plugin_AssetRegistry"));
+							}
+							CompletedSubStates |= ESubState::LoadAssetRegistry;
+							UpdateStateMachineImmediate();
+						});
+				}, UE::Tasks::ETaskPriority::BackgroundHigh);
+		}
 	}
 
 	virtual void UpdateState(FGameFeaturePluginStateStatus& StateStatus) override
diff --git a/Engine/Source/Runtime/AssetRegistry/Private/ARLoader.cpp b/Engine/Source/Runtime/AssetRegistry/Private/ARLoader.cpp
--- a/Engine/Source/Runtime/AssetRegistry/Private/ARLoader.cpp
+++ b/Engine/Source/Runtime/AssetRegistry/Private/ARLoader.cpp
@@ -13,7 +13,9 @@
 #include "GenericPlatform/GenericPlatformFile.h"
 #include "HAL/Event.h"
 #include "HAL/FileManager.h"
+#include "HAL/FileManagerGeneric.h"
 #include "HAL/LowLevelMemTracker.h"
+#include "HAL/PlatformFileManager.h"
 #include "HAL/PlatformMisc.h"
 #include "HAL/PlatformProcess.h"
 #include "HAL/PlatformProperties.h"
@@ -440,38 +442,46 @@ void FAsyncConsumer::Consume(UAssetRegistryImpl& UARI, UE::AssetRegistry::Impl::
 namespace UE::AssetRegistry::Impl
 {
 
-TUniquePtr<FArchive> OpenReadStateFileArchive(FStringView FilePath)
+TUniquePtr<FArchive> OpenReadStateFileArchive(FStringView FilePath, uint32 BufferSize)
 {
-	// Read the entire file into memory, this is faster than serializing repeatedly from a FileReader.
-	TUniquePtr<FArchive> FileReader(IFileManager::Get().CreateFileReader(*WriteToString<512>(FilePath)));
-	if (!FileReader)
+	TStringBuilder<512> FilePathSZ(InPlace, FilePath);
+	TUniquePtr<IFileHandle> FileHandle(FPlatformFileManager::Get().GetPlatformFile().OpenRead(*FilePathSZ));
+	if (!FileHandle)
 	{
 		return TUniquePtr<FArchive>();
 	}
-	// FLargeMemoryReader can take ownership of the Data buffer; it requires that we allocate with malloc.
-	int64 Size = FileReader->TotalSize();
-	if (Size <= 0)
+	int64 FileSize = FileHandle->Size();
+	if (FileSize <= 0)
 	{
 		return TUniquePtr<FArchive>();
 	}
-	uint8* Data = static_cast<uint8*>(FMemory::Malloc(Size));
-	ON_SCOPE_EXIT
+	if (BufferSize == 0 || FileSize <= static_cast<int64>(BufferSize))
 	{
-		if (Data)
+		// Read the entire file into memory. We set BufferSize=0 and take this route by default because
+		// AssetRegistry serialization is on the critical path for starting the game in many cases, and its
+		// faster to read all at once rather than serializing repeatedly from a FileReader.
+
+		// FLargeMemoryReader can take ownership of the Data buffer; it requires that we allocate with malloc.
+		uint8* Data = static_cast<uint8*>(FMemory::Malloc(FileSize));
+		ON_SCOPE_EXIT
 		{
-			FMemory::Free(Data);
+			if (Data)
+			{
+				FMemory::Free(Data);
+			}
+		};
+		if (!FileHandle->Read(Data, FileSize))
+		{
+			return TUniquePtr<FArchive>();
 		}
-	};
-
-	FileReader->Serialize(Data, Size);
-	if (FileReader->IsError())
+		TUniquePtr<FArchive> Result(new FLargeMemoryReader(Data, FileSize, ELargeMemoryReaderFlags::TakeOwnership));
+		Data = nullptr; // detach Data from our call to FMemory::Free
+		return Result;
+	}
+	else
 	{
-		return TUniquePtr<FArchive>();
+		return MakeUnique<FArchiveFileReaderGeneric>(FileHandle.Release(), *FilePathSZ, FileSize, BufferSize);
 	}
-
-	TUniquePtr<FArchive> Result(new FLargeMemoryReader(Data, Size, ELargeMemoryReaderFlags::TakeOwnership));
-	Data = nullptr; // detach Data from our call to FMemory::Free
-	return Result;
 }
 
 FAssetProviderInitData::FAssetProviderInitData(FAssetProviderInitData&& Other)
diff --git a/Engine/Source/Runtime/AssetRegistry/Private/ARLoader.h b/Engine/Source/Runtime/AssetRegistry/Private/ARLoader.h
--- a/Engine/Source/Runtime/AssetRegistry/Private/ARLoader.h
+++ b/Engine/Source/Runtime/AssetRegistry/Private/ARLoader.h
@@ -207,8 +207,11 @@ private:
 namespace UE::AssetRegistry::Impl
 {
 
-/** Open the given AssetRegistry.bin file using the performant archive we use for loading it. */
-TUniquePtr<FArchive> OpenReadStateFileArchive(FStringView FilePath);
+/**
+ * Open the given AssetRegistry.bin file using the performant archive we use for loading it.
+ * If BufferSize == 0 or is greater than filesize, reads the entire file into an in-memory archive.
+ */
+TUniquePtr<FArchive> OpenReadStateFileArchive(FStringView FilePath, uint32 BufferSize=0);
 
 /** Which event triggered our need to register the premade AssetRegistry with the ActiveMountsLoading. */
 enum class EInitialProviderRegisterReason
@@ -283,14 +286,14 @@ public:
 	void LoadAsync(UAssetRegistryImpl& UARI, FActiveMountEvents& EventContext,
 		TConstArrayView<TRefCountPtr<UE::PackageName::IMountPoint>> InMountPoints,
 		TFunction<void(bool bCanceled)>&& OnComplete);
+	/** @see IAssetRegistry::ActiveMountsReloadAllAsync. Calls OnComplete(bCanceled=true) immediately if disabled. */
+	void ReloadAllAsync(UAssetRegistryImpl& UARI, FActiveMountEvents& EventContext,
+		TFunction<void(bool bCanceled)>&& OnComplete);
 	/** @see IAssetRegistry::ActiveMountsUnloadAsync. Calls OnComplete(bCanceled=true) immediately if disabled. */
 	void UnloadAsync(UAssetRegistryImpl& UARI, FActiveMountEvents& EventContext,
 		TConstArrayView<TRefCountPtr<UE::PackageName::IMountPoint>> InMountPoints,
 		TFunction<void(bool bCanceled)>&& OnComplete);
-	/**
-	 * Checks the PackageName mounted status of all MountPoints with a registered AssetProvider, and unloads the
-	 * AssetDatas of all of the currently unmounted ones from the AssetRegistry.
-	 */
+	/** @see IAssetRegistry::ActiveMountsUnloadAllUnmountedAsync. Calls OnComplete(bCanceled=true) immediately if disabled. */
 	void UnloadAllUnmountedAsync(UAssetRegistryImpl& UARI, FActiveMountEvents& EventContext,
 		TFunction<void(bool bCanceled)>&& OnComplete);
 	/** @see IAssetRegistry::ActiveMountsIsLoaded. */
diff --git a/Engine/Source/Runtime/AssetRegistry/Private/ActiveMountsLoader.cpp b/Engine/Source/Runtime/AssetRegistry/Private/ActiveMountsLoader.cpp
--- a/Engine/Source/Runtime/AssetRegistry/Private/ActiveMountsLoader.cpp
+++ b/Engine/Source/Runtime/AssetRegistry/Private/ActiveMountsLoader.cpp
@@ -225,6 +225,24 @@ void FActiveMountsLoaderWrapper::LoadAsync(UAssetRegistryImpl& UARI, FActiveMoun
 	EventContext.AddCompletion(MoveTemp(OnComplete), true /* bCanceled */);
 }
 
+void FActiveMountsLoaderWrapper::ReloadAllAsync(UAssetRegistryImpl& UARI, FActiveMountEvents& EventContext,
+	TFunction<void(bool bCanceled)>&& OnComplete)
+{
+	if (FActiveMountsLoader* LocalInner = LockIfEnabled(); LocalInner)
+	{
+		ON_SCOPE_EXIT
+		{
+			Unlock();
+		};
+		if (!LocalInner->bShutdown)
+		{
+			LocalInner->ReloadAllAsync(UARI, EventContext, MoveTemp(OnComplete));
+			return;
+		}
+	}
+	EventContext.AddCompletion(MoveTemp(OnComplete), true /* bCanceled */);
+}
+
 void FActiveMountsLoaderWrapper::UnloadAsync(UAssetRegistryImpl& UARI, FActiveMountEvents& EventContext,
 	TConstArrayView<TRefCountPtr<UE::PackageName::IMountPoint>> InMountPoints,
 	TFunction<void(bool bCanceled)>&& OnComplete)
@@ -911,6 +929,40 @@ void FActiveMountsLoader::LoadAsync(UAssetRegistryImpl& UARI, FActiveMountEvents
 	}
 }
 
+void FActiveMountsLoader::ReloadAllAsync(UAssetRegistryImpl& UARI, FActiveMountEvents& EventContext,
+	TFunction<void(bool bCanceled)>&& OnComplete)
+{
+	if (!bLoaderIsActive)
+	{
+		// See the comment in LoadAsync. We set bCanceled=false because we know all mountpoints are loaded.
+		EventContext.AddCompletion(MoveTemp(OnComplete), false /* bCanceled */);
+		return;
+	}
+
+	TArray<TRefCountPtr<UE::PackageName::IMountPoint>> MountPointsToLoad;
+	for (FMountPointData& MountData : MountPoints)
+	{
+		FAssetProvider* AssetProvider = MountData.GetAssetProvider();
+		if (!AssetProvider || AssetProvider->IsShutdown())
+		{
+			continue;
+		}
+
+		// Reload from AssetRegistry the MountPoints that we have previously unmounted.
+		// Include MountPoints that are in the process of loading but have not finished loading yet, because we need
+		// to wait on those before reporting to our caller.
+		// Include MountPoints that are loaded but that have been requested to unload, because we need to cancel those
+		// unloads.
+		if (MountData.GetCurrentState() != EMountState::Loaded ||
+			MountData.GetTargetState() != EMountState::Loaded)
+		{
+			const TRefCountPtr<UE::PackageName::IMountPoint>& MountPoint = MountData.GetMountPoint();
+			MountPointsToLoad.Add(MountPoint);
+		}
+	}
+	LoadAsync(UARI, EventContext, MountPointsToLoad, MoveTemp(OnComplete));
+}
+
 void FActiveMountsLoader::UnloadAsync(UAssetRegistryImpl& UARI, FActiveMountEvents& EventContext,
 	TConstArrayView<TRefCountPtr<UE::PackageName::IMountPoint>> InMountPoints,
 	TFunction<void(bool bCanceled)>&& OnComplete)
@@ -1518,6 +1570,13 @@ void FActiveMountsLoader::ParseInitialState(FAssetProvider& AssetProvider, const
 				{
 					if (!UnmountedRoots.FindByHash(GetTypeHash(MountPointRoot), MountPointRoot))
 					{
+						if (!AssetProvider.HasMissingMountPoints())
+						{
+							UE_LOGFMT(LogAssetRegistry, Warning,
+								"ActiveMountsLoader found paths in an unrecognized mountpoint in a cooked AssetRegistry. AssetRegistry data for assets in this path will always remain loaded. "
+								"AssetRegistryFile: '{AssetRegistryFile}', Path: '{Path}'.",
+								*FString(AssetProvider.GetFilePath()), *PathStr);
+						}
 						UnmountedRoots.Add(FString(MountPointRoot));
 					}
 				}
@@ -1526,18 +1585,11 @@ void FActiveMountsLoader::ParseInitialState(FAssetProvider& AssetProvider, const
 
 	if (!UnmountedRoots.IsEmpty())
 	{
-		if (!AssetProvider.HasMissingMountPoints())
-		{
-			AssetProvider.SetHasMissingMountPoints(true);
-			AssetProvider.SetIsErrorNoUnload(true);
-			UE_LOGFMT(LogAssetRegistry, Error,
-				"ActiveMountsLoader found a path in an unrecognized mountpoint in a cooked AssetRegistry. "
-				"We cannot validate the correctness of unloading without knowing the mountpoint, so we are marking this AssetRegistry as not unloadable. "
-				"AssetRegistry data for assets in this cooked AssetRegistry will always remain loaded. "
-				"To fix this, declare all MountPoints before calling ActiveMountsRegisterAndLoadFilePathAsync or ActiveMountsInitialize. "
-				"AssetRegistryFile: '{AssetRegistryFile}', Unrecognized mountpoint : '{MountPoint}'.",
-				*FString(AssetProvider.GetFilePath()), **UnmountedRoots.FindArbitraryElement());
-		}
+		// We do not treat this as an error for the entire AssetProvider; we still allow unloading AssetDatas in
+		// recognized mountpoints in the AssetProvider.
+		// TODO: We could eliminate this warning and allow unloading Assets whenever they are registered if we
+		// added support for placeholder MountPoints with known packagename but unknown contentpath.
+		AssetProvider.SetHasMissingMountPoints(true);
 	}
 
 	// Record the MountPoints
@@ -1829,7 +1881,8 @@ void FActiveMountsLoader::RunAssetProviderLoadTask(UAssetRegistryImpl& UARI, FAs
 
 		// Open the file and preload it, but do not deserialize it yet; we need to know the list of mountpoints first,
 		// and we want to capture those as late as possible in case of stragglers while we are preloading the file.
-		TUniquePtr<FArchive> FileArchive = OpenReadStateFileArchive(StateFilePath);
+		constexpr uint32 BufferSize = 10 * 1024 * 1024;
+		TUniquePtr<FArchive> FileArchive = OpenReadStateFileArchive(StateFilePath, BufferSize);
 		FixedTagPrivate::FWeakStorePtr Store;
 
 		// Reenter the lock to get the list of mountpoints to load from the statefile.
@@ -2032,6 +2085,7 @@ void FActiveMountsLoader::RunAssetProviderLoadTask(UAssetRegistryImpl& UARI, FAs
 				}
 			}
 		}
+		EventContext.Broadcast();
 	}
 }
 
@@ -2183,6 +2237,7 @@ void FActiveMountsLoader::RunMountPointUnloadTask(UAssetRegistryImpl& UARI)
 				}
 			}
 		}
+		EventContext.Broadcast();
 	}
 }
 
diff --git a/Engine/Source/Runtime/AssetRegistry/Private/ActiveMountsLoader.h b/Engine/Source/Runtime/AssetRegistry/Private/ActiveMountsLoader.h
--- a/Engine/Source/Runtime/AssetRegistry/Private/ActiveMountsLoader.h
+++ b/Engine/Source/Runtime/AssetRegistry/Private/ActiveMountsLoader.h
@@ -240,6 +240,8 @@ public:
 	void LoadAsync(UAssetRegistryImpl& UARI, FActiveMountEvents& EventContext,
 		TConstArrayView<TRefCountPtr<UE::PackageName::IMountPoint>> InMountPoints,
 		TFunction<void(bool bCanceled)>&& OnComplete);
+	void ReloadAllAsync(UAssetRegistryImpl& UARI, FActiveMountEvents& EventContext,
+		TFunction<void(bool bCanceled)>&& OnComplete);
 	void UnloadAsync(UAssetRegistryImpl& UARI, FActiveMountEvents& EventContext,
 		TConstArrayView<TRefCountPtr<UE::PackageName::IMountPoint>> InMountPoints,
 		TFunction<void(bool bCanceled)>&& OnComplete);
diff --git a/Engine/Source/Runtime/AssetRegistry/Private/AssetRegistry.cpp b/Engine/Source/Runtime/AssetRegistry/Private/AssetRegistry.cpp
--- a/Engine/Source/Runtime/AssetRegistry/Private/AssetRegistry.cpp
+++ b/Engine/Source/Runtime/AssetRegistry/Private/AssetRegistry.cpp
@@ -4672,10 +4672,9 @@ UE::AssetRegistry::Impl::EGatherStatus UAssetRegistryImpl::TickOnBackgroundThrea
 			ON_SCOPE_EXIT { GatheredDataProcessingLock.Unlock(); };
 			UE::AssetRegistry::Impl::FEventContext EventContext;
 			UE::AssetRegistry::Impl::FClassInheritanceContext InheritanceContext;
-			UE::AssetRegistry::Impl::FInitializeContext InitializeContext{ *this };
 			UE::AssetRegistry::Impl::FClassInheritanceBuffer InheritanceBuffer;
 			UE::AssetRegistry::FInterfaceWriteScopeLock InterfaceScopeLock(InterfaceLock, UE::AssetRegistry::Private::PriorityLow);
-			GetInheritanceContextWithRequiredLock(InterfaceScopeLock, InitializeContext.InheritanceContext, InitializeContext.InheritanceBuffer);
+			GetInheritanceContextWithRequiredLock(InterfaceScopeLock, InheritanceContext, InheritanceBuffer);
 
 			UE::AssetRegistry::Impl::FInterruptionContext::ShouldExitEarlyCallbackType EarlyExitHelper = 
 				[this]()->bool
@@ -6154,6 +6153,13 @@ void UAssetRegistryImpl::ActiveMountsLoadAsync(TConstArrayView<FStringView> Moun
 	ActiveMountEvents.Broadcast();
 }
 
+void UAssetRegistryImpl::ActiveMountsReloadAllAsync(TFunction<void(bool bCanceled)> OnComplete)
+{
+	UE::AssetRegistry::Impl::FActiveMountEvents ActiveMountEvents;
+	ActiveMountsLoader.ReloadAllAsync(*this, ActiveMountEvents, MoveTemp(OnComplete));
+	ActiveMountEvents.Broadcast();
+}
+
 void UAssetRegistryImpl::ActiveMountsUnloadAsync(TConstArrayView<FStringView> MountPointLongPackageNames,
 	TFunction<void(bool bCanceled)> OnComplete)
 {
@@ -6177,6 +6183,13 @@ void UAssetRegistryImpl::ActiveMountsUnloadAsync(TConstArrayView<FStringView> Mo
 	ActiveMountEvents.Broadcast();
 }
 
+void UAssetRegistryImpl::ActiveMountsUnloadAllUnmountedAsync(TFunction<void(bool bCanceled)> OnComplete)
+{
+	UE::AssetRegistry::Impl::FActiveMountEvents ActiveMountEvents;
+	ActiveMountsLoader.UnloadAllUnmountedAsync(*this, ActiveMountEvents, MoveTemp(OnComplete));
+	ActiveMountEvents.Broadcast();
+}
+
 bool UAssetRegistryImpl::ActiveMountsIsLoaded(FStringView MountPointLongPackageName)
 {
 	using namespace UE::AssetRegistry;
diff --git a/Engine/Source/Runtime/AssetRegistry/Private/AssetRegistry.h b/Engine/Source/Runtime/AssetRegistry/Private/AssetRegistry.h
--- a/Engine/Source/Runtime/AssetRegistry/Private/AssetRegistry.h
+++ b/Engine/Source/Runtime/AssetRegistry/Private/AssetRegistry.h
@@ -124,8 +124,12 @@ public:
 		UE::Tasks::ETaskPriority TaskPriority = UE::Tasks::ETaskPriority::Default) override;
 	virtual void ActiveMountsLoadAsync(TConstArrayView<FStringView> MountPointLongPackageNames,
 		TFunction<void(bool bCanceled)> OnComplete = TFunction<void(bool bCanceled)>()) override;
+	virtual void ActiveMountsReloadAllAsync(
+		TFunction<void(bool bCanceled)> OnComplete = TFunction<void(bool bCanceled)>()) override;
 	virtual void ActiveMountsUnloadAsync(TConstArrayView<FStringView> MountPointLongPackageNames,
 		TFunction<void(bool bCanceled)> OnComplete = TFunction<void(bool bCanceled)>()) override;
+	virtual void ActiveMountsUnloadAllUnmountedAsync(
+		TFunction<void(bool bCanceled)> OnComplete = TFunction<void(bool bCanceled)>()) override;
 	virtual bool ActiveMountsIsLoaded(FStringView MountPointLongPackageName) override;
 	virtual bool ActiveMountsIsInProgress() override;
 	virtual void ActiveMountsReportWhenIdle(TFunction<void()> OnIdle) override;
diff --git a/Engine/Source/Runtime/AssetRegistry/Public/AssetRegistry/IAssetRegistry.h b/Engine/Source/Runtime/AssetRegistry/Public/AssetRegistry/IAssetRegistry.h
--- a/Engine/Source/Runtime/AssetRegistry/Public/AssetRegistry/IAssetRegistry.h
+++ b/Engine/Source/Runtime/AssetRegistry/Public/AssetRegistry/IAssetRegistry.h
@@ -1110,6 +1110,17 @@ public:
 	 */
 	virtual void ActiveMountsLoadAsync(TConstArrayView<FStringView> MountPointLongPackageNames,
 		TFunction<void(bool bCanceled)> OnComplete = TFunction<void(bool bCanceled)>()) = 0;
+	/**
+	 * Requests that all AssetDatas that have been unloaded by ActiveMounts be reloaded. Calls OnComplete once the
+	 * load completes.
+	 *
+	 * @param OnComplete Optional delegate: void OnComplete(bool bCanceled). Called when the load is complete.
+	 *        bCanceled is set to true iff ActiveMountsLoading is not enabled or for any MountPoint the load failed
+	 *        or was overridden by an unload request. 
+	 *        OnComplete is called from an async thread, caller should handle thread safety in their callback.
+	 */
+	virtual void ActiveMountsReloadAllAsync(
+		TFunction<void(bool bCanceled)> OnComplete = TFunction<void(bool bCanceled)>()) = 0;
 	/**
 	 * Requests that AssetDatas under the given MountPoints (looked up via FPackageName::FindOrAddMountPoint) be
 	 * unloaded from the AssetRegistry. This is only possible for AssetDatas in the premade AssetRegistry or in an
@@ -1123,6 +1134,18 @@ public:
 	 */
 	virtual void ActiveMountsUnloadAsync(TConstArrayView<FStringView> MountPointLongPackageNames,
 		TFunction<void(bool bCanceled)> OnComplete = TFunction<void(bool bCanceled)>()) = 0;
+	/**
+	 * Requests that AssetDatas under MountPoints that are available for unloading via the premade
+	 * AssetRegistry or in an AssetRegistry registered via an ActiveMountsRegister function be unloaded if they
+	 * are not currently mounted with FPackageName. Calls OnComplete once the unload completes.
+	 *
+	 * @param OnComplete Optional delegate: void OnComplete(bool bCanceled). Called when the unload is complete.
+	 *        bCanceled is set to true iff ActiveMountsLoading is not enabled or for any MountPoint the unload failed
+	 *        or was overridden by a load request. 
+	 *        OnComplete is called from an async thread, caller should handle thread safety in their callback.
+	 */
+	virtual void ActiveMountsUnloadAllUnmountedAsync(
+		TFunction<void(bool bCanceled)> OnComplete = TFunction<void(bool bCanceled)>()) = 0;
 	/**
 	 * Returns whether the given MountPoint (looked up via FPackageName::FindOrAddMountPoint) is loaded in the
 	 * ActiveMounts system. Returns false if currently loading or unloading or if unloaded. Only valid if
