(Cherrypick) Map tag store data in Asset Registry cache, saving GBs of editor memory usage on large projects.

- Added support for older versions of the AR cache.

Also includes cherrypick of CL 50995321 (Enable AR cache memory mapping on buildmachines if cache files aren't being overwritten).

Also includes cherrypick of CL 50448955 (Add HashingArchiveProxy to MemoryHasher, replacing a private one.)

#was-backout 50479750

#Synced-CL 50510093

#Synced-CL 51216913

diff --git a/Engine/Source/Runtime/AssetRegistry/Private/AssetDataGatherer.cpp b/Engine/Source/Runtime/AssetRegistry/Private/AssetDataGatherer.cpp
--- a/Engine/Source/Runtime/AssetRegistry/Private/AssetDataGatherer.cpp
+++ b/Engine/Source/Runtime/AssetRegistry/Private/AssetDataGatherer.cpp
@@ -5423,6 +5423,10 @@ TArray<FCachePayload> LoadCacheFiles(TConstArrayView<FString> InCacheFilenames,
 {
 	TRACE_CPUPROFILER_EVENT_SCOPE(LoadCacheFiles);
 	TArray<FCachePayload> Results;
+	if (UNLIKELY(InCacheFilenames.IsEmpty()))
+	{
+		return MoveTemp(Results);
+	}
 	Results.AddDefaulted(InCacheFilenames.Num());
 
 	// when writing cache files, disable memory mapping on the builder machines until the situation with deletion clarifies
