AssetRegistry: ActiveMountsLoading: Move the call to ActiveMountsInitialize so that it always occurs, and add more logging for it.
#rnx

#Synced-CL 51202900

#Synced-CL 51202900

diff --git a/Engine/Source/Runtime/AssetRegistry/Private/ActiveMountsLoader.cpp b/Engine/Source/Runtime/AssetRegistry/Private/ActiveMountsLoader.cpp
--- a/Engine/Source/Runtime/AssetRegistry/Private/ActiveMountsLoader.cpp
+++ b/Engine/Source/Runtime/AssetRegistry/Private/ActiveMountsLoader.cpp
@@ -1084,6 +1084,7 @@ void FActiveMountsLoader::UnloadAllUnmountedAsync(UAssetRegistryImpl& UARI, FAct
 {
 	if (!bLoaderIsActive)
 	{
+		ensureMsgf(false, TEXT("UnloadAllUnmountedAsync called before ActiveMountsInitialized. The unload call will be ignored."));
 		EventContext.AddCompletion(MoveTemp(OnComplete), true /* bCanceled */);
 		return;
 	}
@@ -1163,6 +1164,7 @@ void FActiveMountsLoader::OnActiveMountsInitialize(UAssetRegistryImpl& UARI, FAc
 	// Note that we require the ProofOfWriteLock around the read of UARI.GuardedData.State.
 	InitialProviderConditionalRegister(UARI, EventContext, UARI.GuardedData.State);
 	bLoaderIsActive = true;
+	UE_LOGFMT(LogAssetRegistry, Display, "ActiveMountsInitialize called. AssetRegistry data unloading is now available.");
 }
 
 bool FActiveMountsLoader::TryRegisterAssetProviderAndAppendStateOutsideARLock(UAssetRegistryImpl& UARI,
