Convert some asset registry log calls from Display verbosity to Log verbosity so they don't show up on stdout for programs like VerseTestScriptCmd/VerseCmd.

#Synced-CL 51293499

diff --git a/Engine/Source/Runtime/AssetRegistry/Private/ActiveMountsLoader.cpp b/Engine/Source/Runtime/AssetRegistry/Private/ActiveMountsLoader.cpp
--- a/Engine/Source/Runtime/AssetRegistry/Private/ActiveMountsLoader.cpp
+++ b/Engine/Source/Runtime/AssetRegistry/Private/ActiveMountsLoader.cpp
@@ -102,7 +102,7 @@ void FActiveMountsLoaderWrapper::InitializeConfig()
 {
 	if (!IsEnabledByConfig())
 	{
-		UE_LOGFMT(LogAssetRegistry, Display, "ActiveMountsLoading is disabled.");
+		UE_LOGFMT(LogAssetRegistry, Log, "ActiveMountsLoading is disabled.");
 		return;
 	}
 
@@ -122,7 +122,7 @@ void FActiveMountsLoaderWrapper::InitializeConfig()
 	// Races during initialization are unexpected but are also not a problem.
 	if (!Inner.load(std::memory_order_relaxed))
 	{
-		UE_LOGFMT(LogAssetRegistry, Display, "ActiveMountsLoading is enabled.");
+		UE_LOGFMT(LogAssetRegistry, Log, "ActiveMountsLoading is enabled.");
 		FActiveMountsLoader* LocalInner = new FActiveMountsLoader();
 		LocalInner->RegistrationLock = MakeShared<UE::FRecursiveMutex>();
 		LocalInner->bAutoUnload = bAutoUnload;
@@ -1164,7 +1164,7 @@ void FActiveMountsLoader::OnActiveMountsInitialize(UAssetRegistryImpl& UARI, FAc
 	// Note that we require the ProofOfWriteLock around the read of UARI.GuardedData.State.
 	InitialProviderConditionalRegister(UARI, EventContext, UARI.GuardedData.State);
 	bLoaderIsActive = true;
-	UE_LOGFMT(LogAssetRegistry, Display, "ActiveMountsInitialize called. AssetRegistry data unloading is now available.");
+	UE_LOGFMT(LogAssetRegistry, Log, "ActiveMountsInitialize called. AssetRegistry data unloading is now available.");
 }
 
 bool FActiveMountsLoader::TryRegisterAssetProviderAndAppendStateOutsideARLock(UAssetRegistryImpl& UARI,
@@ -1448,7 +1448,7 @@ void FActiveMountsLoader::InitialProviderOnReadyForRegister(UAssetRegistryImpl&
 			"AssetRegistry data for assets without MountPoints currently known in FPackageName will remain loaded for the entire process.");
 		break;
 	case EInitialProviderRegisterReason::AppendState:
-		UE_LOGFMT(LogAssetRegistry, Display,
+		UE_LOGFMT(LogAssetRegistry, Log,
 			"AppendState was called with another AssetRegistry before ActiveMountsInitialize was called. Triggering ActiveMountsInitialize.");
 		break;
 	default:
@@ -2043,7 +2043,7 @@ void FActiveMountsLoader::RunAssetProviderLoadTask(UAssetRegistryImpl& UARI, FAs
 					*StateFilePath, *FString(MountPointsToLoad[0]->GetLongPackageName()));
 			}
 
-			UE_LOGFMT(LogAssetRegistry, Display, "ActiveMountsLoading reloaded AssetRegistry data for {NumMountPoints} MountPoints.",
+			UE_LOGFMT(LogAssetRegistry, Log, "ActiveMountsLoading reloaded AssetRegistry data for {NumMountPoints} MountPoints.",
 				MountPointsToLoad.Num());
 
 			for (TRefCountPtr<UE::PackageName::IMountPoint>& MountPoint : MountPointsToLoad)
@@ -2198,7 +2198,7 @@ void FActiveMountsLoader::RunMountPointUnloadTask(UAssetRegistryImpl& UARI)
 				return;
 			}
 
-			UE_LOGFMT(LogAssetRegistry, Display, "ActiveMountsLoading unloaded AssetRegistry data for {NumMountPoints} MountPoints.",
+			UE_LOGFMT(LogAssetRegistry, Log, "ActiveMountsLoading unloaded AssetRegistry data for {NumMountPoints} MountPoints.",
 				MountPointsToUnload.Num());
 
 			for (TRefCountPtr<UE::PackageName::IMountPoint>& MountPoint : MountPointsToUnload)
