Provide an explicit argument to disable memory mapping in the asset registry.

- Setting the CVar value via -dpcvars is too late as AR reads will be initiated before that.

#Synced-CL 51373877

diff --git a/Engine/Source/Runtime/AssetRegistry/Private/AssetDataGatherer.cpp b/Engine/Source/Runtime/AssetRegistry/Private/AssetDataGatherer.cpp
--- a/Engine/Source/Runtime/AssetRegistry/Private/AssetDataGatherer.cpp
+++ b/Engine/Source/Runtime/AssetRegistry/Private/AssetDataGatherer.cpp
@@ -5430,7 +5430,8 @@ TArray<FCachePayload> LoadCacheFiles(TConstArrayView<FString> InCacheFilenames,
 	Results.AddDefaulted(InCacheFilenames.Num());
 
 	// when writing cache files, disable memory mapping on the builder machines until the situation with deletion clarifies
-	if (GIsBuildMachine && bCacheWriteEnabled)
+	static bool bNoAssetRegistryMemoryMapping = FParse::Param(FCommandLine::Get(), TEXT("NoAssetRegistryMemoryMapping"));
+	if ((GIsBuildMachine && bCacheWriteEnabled) || bNoAssetRegistryMemoryMapping)
 	{
 		UE::AssetDataGather::Private::bMemoryMappedCacheFiles = false;
 	}
