[AssetRegistry]

- Allow the asset registry to tick until interrupted to reduce overhead from exiting and re-gathering every 0.1 seconds even when the game thread hasn't requested AR access.
- Move some large GCScopeGuardss to be held only around the code that needs the scoping to allow the GameThread to progress more often when gathering
- Change iterative sync loading code to issue a batch of async loads and then wait on all of them instead.
- Adds a wrapper around the priority scope locks used by the AR so we may measure InterfaceLock contention

#rnx
#Synced-CL 51471209

diff --git a/Engine/Source/Runtime/AssetRegistry/Private/ARTypes.h b/Engine/Source/Runtime/AssetRegistry/Private/ARTypes.h
--- a/Engine/Source/Runtime/AssetRegistry/Private/ARTypes.h
+++ b/Engine/Source/Runtime/AssetRegistry/Private/ARTypes.h
@@ -12,9 +12,11 @@
 
 namespace UE::AssetRegistry::Private { class FRWLockWithPriority; }
 namespace UE::AssetRegistry::Private { using FInterfaceRWLock = FRWLockWithPriority; }
+namespace UE::AssetRegistry::Private { template <typename TMutexType> class TTracingWriteScopeLock; }
+namespace UE::AssetRegistry::Private { template <typename TMutexType> class TTracingReadScopeLock; }
 namespace UE::AssetRegistry { template <typename TScopeLockType> class TRWScopeLockWithPriority; }
-namespace UE::AssetRegistry { using FInterfaceReadScopeLock = TRWScopeLockWithPriority<UE::TReadScopeLock<Private::FRWLockWithPriority>>; }
-namespace UE::AssetRegistry { using FInterfaceWriteScopeLock = TRWScopeLockWithPriority<UE::TWriteScopeLock<Private::FRWLockWithPriority>>; }
+namespace UE::AssetRegistry { using FInterfaceReadScopeLock = TRWScopeLockWithPriority<Private::TTracingReadScopeLock<Private::FRWLockWithPriority>>; }
+namespace UE::AssetRegistry { using FInterfaceWriteScopeLock = TRWScopeLockWithPriority<Private::TTracingWriteScopeLock<Private::FRWLockWithPriority>>; }
 namespace UE::AssetRegistry { using FInterfaceRWScopeLock = class FRWScopeLockWithPriority; }
 
 namespace UE::AssetRegistry::Private
@@ -77,6 +79,77 @@ private:
 	std::atomic<int32> HighPriorityWaitersCount { 0 };
 };
 
+// Custom lock scopes to allow measuring lock contention
+template<typename MutexType>
+class TTracingReadScopeLock
+{
+public:
+	UE_NONCOPYABLE(TTracingReadScopeLock);
+
+	UE_NODISCARD_CTOR TTracingReadScopeLock(MutexType& InMutex)
+		: Mutex(&InMutex)
+	{
+		check(Mutex);
+		if (!Mutex->TryReadLock())
+		{
+			TRACE_CPUPROFILER_EVENT_SCOPE(InterfaceReadLockContention);
+			Mutex->ReadLock();
+		}
+	}
+
+	~TTracingReadScopeLock()
+	{
+		ReadUnlock();
+	}
+
+	void ReadUnlock()
+	{
+		if (Mutex)
+		{
+			Mutex->ReadUnlock();
+			Mutex = nullptr;
+		}
+	}
+
+private:
+	MutexType* Mutex;
+};
+
+template<typename MutexType>
+class TTracingWriteScopeLock
+{
+public:
+	UE_NONCOPYABLE(TTracingWriteScopeLock);
+
+	UE_NODISCARD_CTOR TTracingWriteScopeLock(MutexType& InMutex)
+		: Mutex(&InMutex)
+	{
+		check(Mutex);
+		if (!Mutex->TryWriteLock())
+		{
+			TRACE_CPUPROFILER_EVENT_SCOPE(InterfaceWriteLockContention);
+			Mutex->WriteLock();
+		}
+	}
+
+	~TTracingWriteScopeLock()
+	{
+		WriteUnlock();
+	}
+
+	void WriteUnlock()
+	{
+		if (Mutex)
+		{
+			Mutex->WriteUnlock();
+			Mutex = nullptr;
+		}
+	}
+
+private:
+	MutexType* Mutex;
+};
+
 enum ELockPriority : uint8
 {
 	PriorityLow,
diff --git a/Engine/Source/Runtime/AssetRegistry/Private/AssetRegistry.cpp b/Engine/Source/Runtime/AssetRegistry/Private/AssetRegistry.cpp
--- a/Engine/Source/Runtime/AssetRegistry/Private/AssetRegistry.cpp
+++ b/Engine/Source/Runtime/AssetRegistry/Private/AssetRegistry.cpp
@@ -193,11 +193,11 @@ namespace UE::AssetRegistry::Impl
 		TEXT("AssetRegistry.MaxSecondsPerFrame"),
 		UE::AssetRegistry::Impl::MaxSecondsPerFrame,
 		TEXT("Maximum amount of time allowed for Asset Registry processing, in seconds"));
-	float MaxSecondsPerTickBackgroundThread = 0.1f;
+	float MaxSecondsPerTickBackgroundThread = -1.f;
 	static FAutoConsoleVariableRef CVarAssetRegistryMaxSecondsPerTickBackgroundThread(
 		TEXT("AssetRegistry.MaxSecondsPerTickBackgroundThread"),
 		UE::AssetRegistry::Impl::MaxSecondsPerTickBackgroundThread,
-		TEXT("Maximum amount of time allowed for Asset Registry processing, in seconds, per iteration on the background thread. Very large values could result in main thread delays due to the background thread holding locks."));
+		TEXT("Maximum amount of time allowed for Asset Registry processing, in seconds, per iteration on the background thread. Negative are treated as no maximum."));
 
 	/** If true, defer sorting of dependencies until loading is complete */
 	bool bDeferDependencySort = false;
@@ -4690,8 +4690,11 @@ UE::AssetRegistry::Impl::EGatherStatus UAssetRegistryImpl::TickOnBackgroundThrea
 				};
 
 			UE::AssetRegistry::Impl::FTickContext TickContext(GuardedData, EventContext, InheritanceContext);
-			TickContext.InterruptionContext.SetLimitedTickTime(FPlatformTime::Seconds(),
-				UE::AssetRegistry::Impl::MaxSecondsPerTickBackgroundThread);
+			if (UE::AssetRegistry::Impl::MaxSecondsPerTickBackgroundThread > 0.f)
+			{
+				TickContext.InterruptionContext.SetLimitedTickTime(FPlatformTime::Seconds(),
+					UE::AssetRegistry::Impl::MaxSecondsPerTickBackgroundThread);
+			}
 			TickContext.InterruptionContext.SetEarlyExitCallback(EarlyExitHelper);
 			TickContext.bHandleDeferred = true;
 			Status = GuardedData.TickGatherer(TickContext);
@@ -5512,7 +5515,8 @@ void FTickContext::RunAssetSearchDataGathered(TMultiMap<FName, TUniquePtr<FAsset
 		GuardedData.AssetSearchDataGathered(EventContext, InAssetResults, OutDeferredAssetResults,
 			InterruptionContext, MountPointsForVerifyAfterGather, DeferFlag);
 	}
-};
+}
+
 
 void FTickContext::RunDependencyDataGathered(TMultiMap<FName, FPackageDependencyData>& DependenciesToProcess,
 	TMultiMap<FName, FPackageDependencyData>& OutDeferredDependencies,
@@ -5526,7 +5530,7 @@ void FTickContext::RunDependencyDataGathered(TMultiMap<FName, FPackageDependency
 		GuardedData.DependencyDataGathered(DependenciesToProcess, OutDeferredDependencies,
 			OutPackagesNeedingDependencyCalculation, InterruptionContext, MountPointsForVerifyAfterGather);
 	}
-};
+}
 
 } // namespace Impl within UE::AssetRegistry
 
@@ -6722,7 +6726,7 @@ void FAssetRegistryImpl::ScanPathsSynchronous(FInterfaceWriteScopeLock* ScopeLoc
 
 	auto IsInRequestedDir = [&Context](const FAssetData& AssetData)
 		{
-			TStringBuilder<256> PackageNameStr;
+			TStringBuilder<NAME_SIZE> PackageNameStr;
 			AssetData.PackageName.ToString(PackageNameStr);
 			FStringView PackageName(PackageNameStr.ToString(), PackageNameStr.Len());
 
@@ -7378,11 +7382,6 @@ void FAssetRegistryImpl::AssetSearchDataGathered(Impl::FEventContext& EventConte
 		Utils::InitializeMountPoints(MountPointsForVerifyAfterGather);
 	}
 
-#if WITH_EDITOR
-	// This ensures we can search for classes inside PostLoadAssetRegistryTags. We take the lock once out here to reduce overhead
-	FGCScopeGuard Guard;
-#endif
-
 	TSet<FTopLevelAssetPath> MissingClasses;
 	bool bInterrupted = false;
 	int64 IterationCounter = 0;
@@ -7429,9 +7428,15 @@ void FAssetRegistryImpl::AssetSearchDataGathered(Impl::FEventContext& EventConte
 		}
 
 #if WITH_EDITOR
-		// Postload assets based on their declared class. Queue them for for later retry if their class has not yet loaded.
+		bool bCouldPostLoad = false;
 		FTopLevelAssetPath RequestedGameThreadClass;
-		bool bCouldPostLoad = TryPostLoadAssetRegistryTags(BackgroundResult.Get(), DeferFlag, RequestedGameThreadClass);
+		{
+			// Take the FGCScopeGuard for as short a duration as possible to avoid blocking 
+			// the game thread for long if it needs to collect garbage. The GCLock doesn't 
+			// have a means to signal to our interruption context like the InterfaceLock can.
+			FGCScopeGuard Guard;
+			bCouldPostLoad = TryPostLoadAssetRegistryTags(BackgroundResult.Get(), DeferFlag, RequestedGameThreadClass);
+		}
 		if (!bCouldPostLoad)
 		{
 			if (RequestedGameThreadClass.IsValid() && DeferFlag == EDeferFlag::QueueToGameThread)
@@ -7570,9 +7575,6 @@ void FAssetRegistryImpl::DependencyDataGathered(TMultiMap<FName, FPackageDepende
 		GetScriptPackageNameUnrealEd(),
 	};
 
-	// This ensures we can call FindPackage below from a background thread
-	FGCScopeGuard Guard;
-
 	// If we need to verify mounts, initialize the mounts once before iterating
 	FString PackageRoot;
 	FString PackageNameString;
@@ -7746,6 +7748,7 @@ void FAssetRegistryImpl::DependencyDataGathered(TMultiMap<FName, FPackageDepende
 
 						// Get the guid off the script package, it is updated when 
 						// script is changed so we need to refresh it every run
+						FGCScopeGuard Guard;
 						UPackage* Package = FindPackage(nullptr, *DependencyPackageNameStr);
 
 						if (Package)
@@ -7874,11 +7877,15 @@ void FAssetRegistryImpl::DependencyDataGathered(TMultiMap<FName, FPackageDepende
 	}
 
 	// Sort all referencers now that we are done inserting
-	FDependsNode::FDependsNodeList ModifiedDependNodesArray = ModifiedDependNodes.Array();
-	ParallelFor(TEXT("DependencyGatheredData_SortNodes"), ModifiedDependNodesArray.Num(), 1, 
-		[&ModifiedDependNodesArray, bShouldSortReferencers = ShouldSortReferencers()] (int Index)
+	ParallelFor(TEXT("DependencyGatheredData_SortNodes"), ModifiedDependNodes.GetMaxIndex(), 1, 
+		[&ModifiedDependNodes, bShouldSortReferencers = ShouldSortReferencers()] (int Index)
 		{
-			FDependsNode* Node = ModifiedDependNodesArray[Index];
+			FSetElementId Id = FSetElementId::FromInteger(Index);
+			if (!ModifiedDependNodes.IsValidId(Id))
+			{
+				return;
+			}
+			FDependsNode* Node = ModifiedDependNodes.Get(Id);
 
 			// Set AllowShrinking first so SetIsReferencersSorted will reduce our memory use if we have over allocated
 			Node->SetAllowShrinking(true);
@@ -10779,9 +10786,14 @@ void UAssetRegistryImpl::Broadcast(UE::AssetRegistry::Impl::FEventContext& Event
 	}
 	if (EventContext.RequiredLoads.Num())
 	{
+		TArray<int32> RequestIds;
 		for (const FString& RequiredLoad : EventContext.RequiredLoads)
 		{
-			LoadPackage(nullptr, *RequiredLoad, 0);
+			RequestIds.Add(LoadPackageAsync(*RequiredLoad));
+		}
+		if (!RequestIds.IsEmpty())
+		{
+			FlushAsyncLoading(RequestIds);
 		}
 		EventContext.RequiredLoads.Empty();
 	}
