Avoid grabbing the lock if we always have the temp cache enabled
If we dont, only grab the read lock to check if our value is the same to avoid always grabbing the write lock

#Synced-CL 51322414

diff --git a/Engine/Source/Runtime/AssetRegistry/Private/AssetRegistry.cpp b/Engine/Source/Runtime/AssetRegistry/Private/AssetRegistry.cpp
--- a/Engine/Source/Runtime/AssetRegistry/Private/AssetRegistry.cpp
+++ b/Engine/Source/Runtime/AssetRegistry/Private/AssetRegistry.cpp
@@ -545,7 +545,8 @@ bool UAssetRegistryImpl::IsPathBeautificationNeeded(const FString& InAssetPath)
 namespace UE::AssetRegistry
 {
 
-FAssetRegistryImpl::FAssetRegistryImpl()
+FAssetRegistryImpl::FAssetRegistryImpl() :
+	bIsTempCachingAlwaysEnabled(ASSETREGISTRY_CACHE_ALWAYS_ENABLED)
 {
 }
 
@@ -707,7 +708,6 @@ void FAssetRegistryImpl::Initialize(UE::AssetRegistry::FInterfaceWriteScopeLock&
 	}
 #endif
 
-	bIsTempCachingAlwaysEnabled = ASSETREGISTRY_CACHE_ALWAYS_ENABLED;
 	bIsTempCachingEnabled = bIsTempCachingAlwaysEnabled;
 	TempCachedInheritanceBuffer.bDirty = true;
 
@@ -9336,6 +9336,18 @@ void UAssetRegistryImpl::SetTemporaryCachingMode(bool bEnable)
 {
 	checkf(IsInGameThread(), TEXT("Changing Caching mode is only available on the game thread because it affects behavior on all threads"));
 	LLM_SCOPE(ELLMTag::AssetRegistry);
+	// IsTempCachingAlwaysEnabled is only set once so is safe to check out side of a lock, and if we are always true, never grab a lock here
+	if (GuardedData.IsTempCachingAlwaysEnabled())
+	{
+		return;
+	}
+	{
+		UE::AssetRegistry::FInterfaceReadScopeLock ReadLock(InterfaceLock);
+		if (bEnable == GuardedData.IsTempCachingEnabled())
+		{
+			return;
+		}
+	}
 	UE::AssetRegistry::FInterfaceWriteScopeLock InterfaceScopeLock(InterfaceLock);
 	GuardedData.SetTemporaryCachingMode(bEnable);
 }
diff --git a/Engine/Source/Runtime/AssetRegistry/Private/AssetRegistryImpl.h b/Engine/Source/Runtime/AssetRegistry/Private/AssetRegistryImpl.h
--- a/Engine/Source/Runtime/AssetRegistry/Private/AssetRegistryImpl.h
+++ b/Engine/Source/Runtime/AssetRegistry/Private/AssetRegistryImpl.h
@@ -597,7 +597,7 @@ private:
 	bool bIsTempCachingEnabled;
 
 	/** If true, search caching is enabled permanently */
-	bool bIsTempCachingAlwaysEnabled;
+	const bool bIsTempCachingAlwaysEnabled;
 
 	/** Persistent InheritanceBuffer used when SetTemporaryCachingMode is on */
 	Impl::FClassInheritanceBuffer TempCachedInheritanceBuffer;
