Removed temporary array/string allocation from UAssetRegistryHelpers::GetAssetsWithOuterForPaths

#rnx

#Synced-CL 51753475

diff --git a/Engine/Source/Runtime/AssetRegistry/Private/AssetRegistryHelpers.cpp b/Engine/Source/Runtime/AssetRegistry/Private/AssetRegistryHelpers.cpp
--- a/Engine/Source/Runtime/AssetRegistry/Private/AssetRegistryHelpers.cpp
+++ b/Engine/Source/Runtime/AssetRegistry/Private/AssetRegistryHelpers.cpp
@@ -413,7 +413,7 @@ void UAssetRegistryHelpers::FixupRedirectedAssetPath(FName& InOutAssetPath)
 }
 
 #if WITH_EDITOR
-TArray<FAssetData> UAssetRegistryHelpers::GetAssetsWithOuterForPaths(const TArray<FName>& InPackagePaths, FName InOuterPath, bool bInRecursivePaths, bool bInIncludeOnlyOnDiskAssets, bool bInExactOuter)
+TArray<FAssetData> UAssetRegistryHelpers::GetAssetsWithOuterForPaths(TConstArrayView<FName> InPackagePaths, FName InOuterPath, bool bInRecursivePaths, bool bInIncludeOnlyOnDiskAssets, bool bInExactOuter)
 {
 	EAssetsWithOuterForPathsFlags Flags = EAssetsWithOuterForPathsFlags::ScanPaths;
 	if (bInRecursivePaths)
@@ -435,35 +435,41 @@ TArray<FAssetData> UAssetRegistryHelpers::GetAssetsWithOuterForPaths(const TArra
 }
 
 
-TArray<FAssetData> UAssetRegistryHelpers::GetAssetsWithOuterForPaths(const TArray<FName>& InPackagePaths, FName InOuterPath, EAssetsWithOuterForPathsFlags InFlags)
+TArray<FAssetData> UAssetRegistryHelpers::GetAssetsWithOuterForPaths(TConstArrayView<FName> InPackagePaths, FName InOuterPath, EAssetsWithOuterForPathsFlags InFlags)
 {
 	FARFilter Filter;
 	Filter.PackagePaths = InPackagePaths;
 	Filter.bRecursivePaths = EnumHasAnyFlags(InFlags, EAssetsWithOuterForPathsFlags::RecursivePaths);
 	Filter.bIncludeOnlyOnDiskAssets = EnumHasAnyFlags(InFlags, EAssetsWithOuterForPathsFlags::IncludeOnlyOnDiskAsset);
 
-	TArray<FString> PathsToScan;
-	PathsToScan.Reserve(Filter.PackagePaths.Num());
-	Algo::Transform(Filter.PackagePaths, PathsToScan, [](FName InPackagePath) { return InPackagePath.ToString(); });
-
-	TArray<FAssetData> Assets;
 	IAssetRegistry& AssetRegistry = IAssetRegistry::GetChecked();
+	
 	if (EnumHasAnyFlags(InFlags, EAssetsWithOuterForPathsFlags::ScanPaths))
 	{
+		TArray<FString> PathsToScan;
+		PathsToScan.Reserve(Filter.PackagePaths.Num());
+		Algo::Transform(Filter.PackagePaths, PathsToScan, [](FName InPackagePath) { return InPackagePath.ToString(); });
+
 		AssetRegistry.ScanSynchronous(PathsToScan, TArray<FString>());
 	}
+
+	TArray<FAssetData> Assets;
 	AssetRegistry.GetAssets(Filter, Assets);
 
-	bool bInExactOuter = EnumHasAnyFlags(InFlags, EAssetsWithOuterForPathsFlags::ExactOuter);
+	FNameBuilder PartialOuterRoot;
+	if (!EnumHasAnyFlags(InFlags, EAssetsWithOuterForPathsFlags::ExactOuter))
+	{
+		InOuterPath.AppendString(PartialOuterRoot);
+	}
 	Assets.SetNum(
-		Algo::RemoveIf(Assets, [InOuterPath, bInExactOuter](const FAssetData& InAssetData)
+		Algo::RemoveIf(Assets, [InOuterPath, &PartialOuterRoot](const FAssetData& InAssetData)
 		{
 			if (InAssetData.GetOptionalOuterPathName() == InOuterPath)
 			{
 				return false;
 			}
 
-			if (!bInExactOuter && InAssetData.GetOptionalOuterPathName().ToString().StartsWith(InOuterPath.ToString()))
+			if (PartialOuterRoot.Len() > 0 && FNameBuilder(InAssetData.GetOptionalOuterPathName()).ToView().StartsWith(PartialOuterRoot.ToView(), ESearchCase::IgnoreCase))
 			{
 				return false;
 			}
diff --git a/Engine/Source/Runtime/AssetRegistry/Public/AssetRegistry/AssetRegistryHelpers.h b/Engine/Source/Runtime/AssetRegistry/Public/AssetRegistry/AssetRegistryHelpers.h
--- a/Engine/Source/Runtime/AssetRegistry/Public/AssetRegistry/AssetRegistryHelpers.h
+++ b/Engine/Source/Runtime/AssetRegistry/Public/AssetRegistry/AssetRegistryHelpers.h
@@ -214,7 +214,7 @@ public:
 	 * @param	bInExactOuter If we should only return assets directly outered to the provided outer.
 	 */
 	UE_DEPRECATED(5.7, "Please use GetAssetsWithOuterForPaths with enum flags")
-	static ASSETREGISTRY_API TArray<FAssetData> GetAssetsWithOuterForPaths(const TArray<FName>& InPackagePaths, FName InOuterPath, bool bInRecursivePaths = false, bool bInIncludeOnlyOnDiskAsset = false, bool bInExactOuter = true);
+	static ASSETREGISTRY_API TArray<FAssetData> GetAssetsWithOuterForPaths(TConstArrayView<FName> InPackagePaths, FName InOuterPath, bool bInRecursivePaths = false, bool bInIncludeOnlyOnDiskAsset = false, bool bInExactOuter = true);
 
 	/**
 	 * Returns all the corresponding assets from the provided paths that are outered to the provided outer name.
@@ -222,6 +222,6 @@ public:
 	 * @param	InOuterPath The outer path to filter asset with.
 	 * @param	InFlags The flags to control the behavior of this function, see EAssetsWithOuterForPathsFlags
 	 */
-	static ASSETREGISTRY_API TArray<FAssetData> GetAssetsWithOuterForPaths(const TArray<FName>& InPackagePaths, FName InOuterPath, EAssetsWithOuterForPathsFlags InFlags = (EAssetsWithOuterForPathsFlags::ExactOuter | EAssetsWithOuterForPathsFlags::ScanPaths));
+	static ASSETREGISTRY_API TArray<FAssetData> GetAssetsWithOuterForPaths(TConstArrayView<FName> InPackagePaths, FName InOuterPath, EAssetsWithOuterForPathsFlags InFlags = (EAssetsWithOuterForPathsFlags::ExactOuter | EAssetsWithOuterForPathsFlags::ScanPaths));
 #endif
 };
