CookDeterminism: AssetRegistry: Fix FAssetRegistryState functions Prune and InitializeFromExistingAndPrune to replace dependencies to pruned redirectors with dependencies to their not-pruned targets, rather than losing the dependency information. This fixes differences between full cooks which were pruning the dependencies, and incremental cooks, which use the redirected dependencies that were created during FAssetRegistryState::Save on the development assetregistry and therefore did not prune them when writing them to the runtime assetregistry.
Remove the redirection of dependencies from FAssetRegistryState::Save; it was incorrectly calculated: it redirected the dependencies but did not modify the corresponding referencers. It is no longer necessary since Prune and InitializeFromExistingAndPrune now handle the redirection, and one or the other of those is called before saving cooked AssetRegistry.bin.
#rnx
#Synced-CL 51945718

diff --git a/Engine/Source/Runtime/AssetRegistry/Private/AssetRegistryState.cpp b/Engine/Source/Runtime/AssetRegistry/Private/AssetRegistryState.cpp
--- a/Engine/Source/Runtime/AssetRegistry/Private/AssetRegistryState.cpp
+++ b/Engine/Source/Runtime/AssetRegistry/Private/AssetRegistryState.cpp
@@ -385,6 +385,7 @@ void FAssetRegistryState::InitializeFromExistingAndPrune(const FAssetRegistrySta
 
 	// Find valid dependency nodes for all script and required packages, and create a new
 	// node corresponding to each old node.
+	TMap<FDependsNode*, FDependsNode*> RedirectCache;
 	TMap<FDependsNode*, FDependsNode*> OldNodeToNewNode;
 	OldNodeToNewNode.Reserve(ExistingState.CachedDependsNodes.Num());
 	for (const TPair<FAssetIdentifier, FDependsNode*>& Pair : ExistingState.CachedDependsNodes)
@@ -413,12 +414,18 @@ void FAssetRegistryState::InitializeFromExistingAndPrune(const FAssetRegistrySta
 				ExistingNewNode = CreateOrFindDependsNode(Node->GetIdentifier());
 			}
 		}
+		else
+		{
+			// If we're removing this Node because it is a redirector, store its target
+			// in RedirectCache so that referencers can redirect to the target instead.
+			ExistingState.ResolveRedirector(Node, ExistingState.CachedAssets, RedirectCache);
+		}
 	}
 
 	TMpscQueue<TPair<FDependsNode*, FDependsNode*>> ReferencerPairs;
 
 	int32 OldNodeToNewNodeMaxId = OldNodeToNewNode.GetMaxIndex();
-	ParallelFor(OldNodeToNewNodeMaxId, [&OldNodeToNewNode, &ReferencerPairs](int32 Index)
+	ParallelFor(OldNodeToNewNodeMaxId, [&OldNodeToNewNode, &ReferencerPairs, &RedirectCache](int32 Index)
 		{
 			FSetElementId Id = FSetElementId::FromInteger(Index);
 			if (OldNodeToNewNode.IsValidId(Id))
@@ -437,6 +444,15 @@ void FAssetRegistryState::InitializeFromExistingAndPrune(const FAssetRegistrySta
 					{
 						// Only add link if it's part of the filtered asset set
 						FDependsNode** NewDependency = OldNodeToNewNode.Find(InDependency);
+						if (!NewDependency)
+						{
+							// Check for a redirected target; redirectors might be pruned but their targets might not be.
+							FDependsNode** RedirectedInDependency = RedirectCache.Find(InDependency);
+							if (RedirectedInDependency && *RedirectedInDependency != InDependency)
+							{
+								NewDependency = OldNodeToNewNode.Find(*RedirectedInDependency);
+							}
+						}
 						if (NewDependency)
 						{
 							NewNode->AddDependency(*NewDependency, InCategory, InFlags);
@@ -974,18 +990,22 @@ void FAssetRegistryState::Prune(const FAssetRegistryPruneOptions& PruneOptions)
 
 	TArray<int8> KeepAssetDataIndex;
 	TArray<FDependsNode*> AllDependsNodes;
-	TMap<FDependsNode*, int8> KeepDependsNodes;
+	// RemapNodes is passed to RemapLinks for all nodes we are keeping. Nodes that we are keeping map to themselves.
+	// Nodes that we are not keeping eventually either map to either null, or to a redirect that we are keeping.
+	// While constructing it, all nodes we are not keeping map to null, and then we do a pass looking for redirects
+	// and replace those nulls with redirects.
+	TMap<FDependsNode*, FDependsNode*> RemapNodes;
 
 	CachedDependsNodes.GenerateValueArray(AllDependsNodes);
 	KeepAssetDataIndex.SetNum(AllAssetData.Num());
-	KeepDependsNodes.Reserve(AllDependsNodes.Num());
+	RemapNodes.Reserve(AllDependsNodes.Num());
 	for (FDependsNode* DependsNode : AllDependsNodes)
 	{
-		KeepDependsNodes.Add(DependsNode, 0);
+		RemapNodes.Add(DependsNode, nullptr);
 	}
 
 	ParallelFor(AllAssetData.Num(),
-		[&AllAssetData, &KeepAssetDataIndex, &KeepDependsNodes, &NumRemovedAssetDatas,
+		[&AllAssetData, &KeepAssetDataIndex, &RemapNodes, &NumRemovedAssetDatas,
 		this, &RequiredPackages, &RemovePackages, &ChunksToKeep, &Options, &PruneOptions,
 		bIsFilteredByChunkId, bIsFilteredByRequiredPackages, bIsFilteredByRemovedPackages
 		]
@@ -1031,56 +1051,48 @@ void FAssetRegistryState::Prune(const FAssetRegistryPruneOptions& PruneOptions)
 
 			if (bKeepDependsNode)
 			{
-				FDependsNode** DependsNode = CachedDependsNodes.Find(AssetData->PackageName);
+				FDependsNode** DependsNodePtr = CachedDependsNodes.Find(AssetData->PackageName);
+				// We are keeping this node, so set its remap node to itself. 
 				// Multiple assets can share the same dependsnode; keep the depends node if any of them say to keep it.
 				// We write atomically to prevent a technical datarace from multiple writes to the same variable; it's not
 				// an actual datarace because we always write the same value if we write anything at all.
-				if (DependsNode)
+				// AtomicStore_Relaxed only supports integer types, so we have to cast our pointer and value to an integertype.
+				if (DependsNodePtr)
 				{
-					int8& KeepDependsNode = KeepDependsNodes.FindChecked(*DependsNode);
-					FPlatformAtomics::AtomicStore_Relaxed(&KeepDependsNode, 1);
+					FDependsNode* DependsNode = *DependsNodePtr;
+					FDependsNode*& TargetDependsNode = RemapNodes.FindChecked(DependsNode);
+					FPlatformAtomics::AtomicStore_Relaxed((PTRINT*)&TargetDependsNode, (PTRINT)DependsNode);
 				}
 
 				FPrimaryAssetId PrimaryAssetId = AssetData->GetPrimaryAssetId();
 				if (PrimaryAssetId.IsValid())
 				{
-					FDependsNode** PrimaryAssetDependsNode = CachedDependsNodes.Find(PrimaryAssetId);
-					if (PrimaryAssetDependsNode)
+					DependsNodePtr = CachedDependsNodes.Find(PrimaryAssetId);
+					if (DependsNodePtr)
 					{
-						int8& KeepDependsNode = KeepDependsNodes.FindChecked(*PrimaryAssetDependsNode);
-						FPlatformAtomics::AtomicStore_Relaxed(&KeepDependsNode, 1);
+						FDependsNode* DependsNode = *DependsNodePtr;
+						FDependsNode*& TargetDependsNode = RemapNodes.FindChecked(DependsNode);
+						FPlatformAtomics::AtomicStore_Relaxed((PTRINT*)&TargetDependsNode, (PTRINT)DependsNode);
 					}
 				}
 			}
 
 		});
 
-	TArray<FAssetData*> RemovedAssetDatas;
-	RemovedAssetDatas.Reserve(NumRemovedAssetDatas);
-	for (int32 AssetDataIndex = 0; AssetDataIndex < NumAssetDatas; ++AssetDataIndex)
-	{
-		if (!KeepAssetDataIndex[AssetDataIndex])
-		{
-			RemovedAssetDatas.Add(AllAssetData[AssetDataIndex]);
-		}
-	}
-	TBitArray<> UnusedRemoveAssetData;
-	TBitArray<> UnusedRemovedDependencyData;
-	RemoveAssetDatas(RemovedAssetDatas, false /* bRemoveDependencyData*/, UnusedRemoveAssetData, UnusedRemovedDependencyData);
-
 	// Mark-for-removal all other dependsnodes that are filtered out by our settings
-	int32 MaxId = KeepDependsNodes.GetMaxIndex();
-	ParallelFor(MaxId, [&KeepDependsNodes, &Options, &PruneOptions](int32 Index)
+	int32 MaxId = RemapNodes.GetMaxIndex();
+	ParallelFor(MaxId, [&RemapNodes, &Options, &PruneOptions](int32 Index)
 		{
 			FSetElementId SetId = FSetElementId::FromInteger(Index);
-			if (!KeepDependsNodes.IsValidId(SetId))
+			if (!RemapNodes.IsValidId(SetId))
 			{
 				return;
 			}
-			TPair<FDependsNode*, int8>& KeepPair = KeepDependsNodes.Get(SetId);
+			TPair<FDependsNode*, FDependsNode*>& RemapPair = RemapNodes.Get(SetId);
 
-			FDependsNode* DependsNode = KeepPair.Key;
-			int8& KeepDependsNode = KeepPair.Value;
+			FDependsNode* DependsNode = RemapPair.Key;
+			FDependsNode*& RemapNode = RemapPair.Value;
+			bool bKeepDependsNode = DependsNode == RemapNode; // kept DependsNodes map to themselves.
 			const FAssetIdentifier& Id = DependsNode->GetIdentifier();
 
 			// bKeepDependsNode is initially the value false if the dependsnode does not correspond to an AssetData
@@ -1093,8 +1105,7 @@ void FAssetRegistryState::Prune(const FAssetRegistryPruneOptions& PruneOptions)
 			{
 				// SearchableName; KeepDependsNode is always initially false, set it instead to remove the
 				// SearchableName depends node if and only if we're removing all SearchableName depends nodes.
-				bool bKeepDependsNode = !Options.bFilterSearchableNames;
-				KeepDependsNode = bKeepDependsNode ? 1 : 0;
+				bKeepDependsNode = !Options.bFilterSearchableNames;
 			}
 			else if (Id.PrimaryAssetType.IsValid())
 			{
@@ -1102,7 +1113,7 @@ void FAssetRegistryState::Prune(const FAssetRegistryPruneOptions& PruneOptions)
 				if (!PrimaryAssetId.IsValid())
 				{
 					UE_LOG(LogAssetRegistry, Warning, TEXT("FAssetRegistryState::Prune: Unknown type of DependsNode: %s"), *Id.ToString());
-					KeepDependsNode = 0;
+					bKeepDependsNode = false;
 				}
 				else
 				{
@@ -1110,68 +1121,144 @@ void FAssetRegistryState::Prune(const FAssetRegistryPruneOptions& PruneOptions)
 					// discard that AssetData above. Override this to additionally keep the asset if we are not
 					// bRemoveDependenciesWithoutPackages or if it is in
 					// RemoveDependenciesWithoutPackagesKeepPrimaryAssetTypes.
-					bool bKeepDependsNode = !PruneOptions.bRemoveDependenciesWithoutPackages ||
+					bKeepDependsNode |= !PruneOptions.bRemoveDependenciesWithoutPackages ||
 						PruneOptions.RemoveDependenciesWithoutPackagesKeepPrimaryAssetTypes.Contains(
 							PrimaryAssetId.PrimaryAssetType);
-					KeepDependsNode |= bKeepDependsNode ? 1 : 0;
 				}
 			}
 			else if (Id.IsPackage())
 			{
 				// Package. bKeepDependsNode is false if the package does not have an AssetData, or if we decided to
 				// discard above. Override for script packages, which do not have AssetDatas and which we always keep.
-				bool bKeepDependsNode = FPackageName::IsScriptPackage(WriteToString<256>(Id.PackageName));
-				KeepDependsNode |= bKeepDependsNode ? 1 : 0;
+				bKeepDependsNode |= FPackageName::IsScriptPackage(WriteToString<256>(Id.PackageName));
 			}
 			else if (!Id.PackageName.IsNone() && !Id.ObjectName.IsNone())
 			{
 				// A DependsNode to an object within a package. Invalid, ignore it.
-				KeepDependsNode = 0;
+				bKeepDependsNode = false;
 			}
 			else
 			{
 				// Invalid depends node with all fields == NAME_None; remove it
-				KeepDependsNode = 0;
+				bKeepDependsNode = false;
 			}
+
+			RemapNode = bKeepDependsNode ? DependsNode : nullptr;
 		});
 
-	ParallelFor(MaxId, [&KeepDependsNodes, &Options, &PruneOptions](int32 Index)
+	// For pruned nodes, check for redirectors. If the redirection target is not pruned, mark that the pruned node maps
+	// to the redirection target. Redirectors have to be searched for before removing any pruned asset datas, because
+	// we use that AssetData information to check whether each asset is a redirector.
+	TMap<FDependsNode*, FDependsNode*> RedirectCache;
+	for (TPair<FDependsNode*, FDependsNode*>& RemapPair : RemapNodes)
+	{
+		if (!RemapPair.Value)
+		{
+			FDependsNode* DependsNode = RemapPair.Key;
+			FDependsNode* TargetNode = ResolveRedirector(DependsNode, CachedAssets, RedirectCache);
+			if (TargetNode != DependsNode)
+			{
+				bool bKeepTargetNode = (RemapNodes.FindChecked(TargetNode) == TargetNode);
+				if (bKeepTargetNode)
+				{
+					RemapPair.Value = TargetNode;
+				}
+			}
+		}
+	}
+
+	// Now that we have finished all our use for the pruned AssetDatas, remove them.
+	TArray<FAssetData*> RemovedAssetDatas;
+	RemovedAssetDatas.Reserve(NumRemovedAssetDatas);
+	for (int32 AssetDataIndex = 0; AssetDataIndex < NumAssetDatas; ++AssetDataIndex)
+	{
+		if (!KeepAssetDataIndex[AssetDataIndex])
+		{
+			RemovedAssetDatas.Add(AllAssetData[AssetDataIndex]);
+		}
+	}
+	TBitArray<> UnusedRemoveAssetData;
+	TBitArray<> UnusedRemovedDependencyData;
+	RemoveAssetDatas(RemovedAssetDatas, false /* bRemoveDependencyData*/, UnusedRemoveAssetData, UnusedRemovedDependencyData);
+
+	TMpscQueue<TPair<FDependsNode*, FDependsNode*>> ReferencersToAdd;
+	TMpscQueue<FDependsNode*> PotentialOrphans;
+
+	ParallelFor(MaxId, [&RemapNodes, &ReferencersToAdd, &PotentialOrphans, &Options, &PruneOptions](int32 Index)
 		{
 			FSetElementId SetId = FSetElementId::FromInteger(Index);
-			if (!KeepDependsNodes.IsValidId(SetId))
+			if (!RemapNodes.IsValidId(SetId))
 			{
 				return;
 			}
-			TPair<FDependsNode*, int8>& KeepPair = KeepDependsNodes.Get(SetId);
-			if (KeepPair.Value)
-			{
-				FDependsNode* DependsNode = KeepPair.Key;
-				DependsNode->RemoveLinks([&KeepDependsNodes](const FDependsNode* ExistingDependsNode)
+			TPair<FDependsNode*, FDependsNode*>& RemapPair = RemapNodes.Get(SetId);
+			bool bKeepDependsNode = RemapPair.Key == RemapPair.Value; // kept DependsNodes map to themselves.
+			if (bKeepDependsNode) 
+			{
+				FDependsNode* DependsNode = RemapPair.Key;
+				DependsNode->RemapDependencies(
+					[DependsNode, &RemapNodes, &ReferencersToAdd, &PotentialOrphans]
+					(FDependsNode* ExistingDependsNode)
 					{
-						return 0 == KeepDependsNodes.FindChecked(ExistingDependsNode);
+						FDependsNode* Remapped = RemapNodes.FindChecked(ExistingDependsNode);
+
+						// Remapped == Existing: it's a keeper.
+						// Remapped == nullptr: pruned and either not redirected or target is also pruned; remove it.
+						if (Remapped != ExistingDependsNode && Remapped != nullptr)
+						{
+							// Remapped == some other node: Apply the remap.
+							// 
+							// Since we are adding Remapped to DependsNode->Dependencies, we also need to add
+							// DependsNode to Remapped->Referencers; we have to do that after leaving the
+							// ParallelFor to avoid data race on Remapped->Referencers.
+							// 
+							// We don't worry about removing the Referencer from ExistingDependsNode, because we
+							// are deleting that node and do not plan to read its Referencers anymore.
+							ReferencersToAdd.Enqueue(TPair<FDependsNode*, FDependsNode*>{ Remapped, DependsNode });
+						}
+						return Remapped;
+					});
+				DependsNode->RemapReferencers([&RemapNodes](FDependsNode* ExistingDependsNode)
+					{
+						// Drop references from deleted nodes, no need to handle remapping them.
+						FDependsNode* Remapped = RemapNodes.FindChecked(ExistingDependsNode);
+						if (Remapped == ExistingDependsNode)
+						{
+							return ExistingDependsNode;
+						}
+						return (FDependsNode*)nullptr;
 					});
 
-				// Mark for remove any orphaned depends nodes.
+				// Mark for remove any currently orphaned depends nodes. Wait until after ReferencersToAdd is
+				// processed to decide whether they are orphaned, since ReferencersToAdd may add links to them.
 				if (DependsNode->GetConnectionCount() == 0)
 				{
-					// We are modifying a value in KeepDependsNodes, which on first consideration might be a problem
-					// because it is being read at the same time by the RemoveLinks call for other DependsNodes in this
-					// ParallelFor. But that data race is not possible, because we're only modifying a value for a
-					// node with connectionCount == 0 after removing links to nodes with keepPair.Value == false,
-					// which means none of the other DependsNode with keepPair.Value == true, which are the only ones
-					// that call RemoveLinks, could encounter it in their RemoveLinks call.
-					KeepPair.Value = 0;
+					PotentialOrphans.Enqueue(DependsNode);
 				}
 			}
 		});
+	TPair<FDependsNode*, FDependsNode*> ReferencerPair;
+	while (ReferencersToAdd.Dequeue(ReferencerPair))
+	{
+		ReferencerPair.Key->AddReferencer(ReferencerPair.Value);
+	}
+	FDependsNode* Orphan;
+	while (PotentialOrphans.Dequeue(Orphan))
+	{
+		if (Orphan->GetConnectionCount() == 0)
+		{
+			RemapNodes.FindChecked(Orphan) = nullptr;
+		}
+	}
 
 	// Remove all of the marked-for-removal dependsnodes; we have already removed links to them so we don't need to call
 	// the full RemoveDependsNode function.
-	for (TPair<FDependsNode*, int8>& KeepPair : KeepDependsNodes)
+	for (TPair<FDependsNode*, FDependsNode*>& RemapPair : RemapNodes)
 	{
-		FDependsNode* DependsNode = KeepPair.Key;
-		if (0 == KeepPair.Value)
+		bool bKeepDependsNode = RemapPair.Key == RemapPair.Value;
+		if (!bKeepDependsNode)
 		{
+			FDependsNode* DependsNode = RemapPair.Key;
 			CachedDependsNodes.Remove(DependsNode->GetIdentifier());
 			NumDependsNodes--;
 			// if the depends nodes were preallocated in a block, we can't delete them one at a time,
@@ -1180,6 +1267,11 @@ void FAssetRegistryState::Prune(const FAssetRegistryPruneOptions& PruneOptions)
 			{
 				delete DependsNode;
 			}
+			else
+			{
+				DependsNode->ClearDependencies();
+				DependsNode->ClearReferencers();
+			}
 		}
 	}
 }
@@ -2464,7 +2556,6 @@ bool FAssetRegistryState::Save(FArchive& OriginalAr, const FAssetRegistrySeriali
 	}
 	else
 	{
-		TMap<FDependsNode*, FDependsNode*> RedirectCache;
 		TArray<FDependsNode*> Dependencies;
 
 		// Scan dependency nodes, we won't save all of them if we filter out certain types
@@ -2494,12 +2585,8 @@ bool FAssetRegistryState::Save(FArchive& OriginalAr, const FAssetRegistrySeriali
 		}
 
 		TUniqueFunction<int32(FDependsNode*, bool bAsReferencer)> GetSerializeIndexFromNode =
-			[this, &RedirectCache, &DependsIndexMap](FDependsNode* InDependency, bool bAsReferencer)
+			[this, &DependsIndexMap](FDependsNode* InDependency, bool bAsReferencer)
 		{
-			if (!bAsReferencer)
-			{
-				InDependency = ResolveRedirector(InDependency, CachedAssets, RedirectCache);
-			}
 			if (!InDependency)
 			{
 				return -1;
@@ -3066,7 +3153,7 @@ SIZE_T FAssetRegistryState::GetAllocatedSize(bool bLogDetailed) const
 
 FDependsNode* FAssetRegistryState::ResolveRedirector(FDependsNode* InDependency,
 													 const FAssetDataMap& InAllowedAssets,
-													 TMap<FDependsNode*, FDependsNode*>& InCache)
+													 TMap<FDependsNode*, FDependsNode*>& InCache) const
 {
 	if (InCache.Contains(InDependency))
 	{
diff --git a/Engine/Source/Runtime/AssetRegistry/Private/DependsNode.cpp b/Engine/Source/Runtime/AssetRegistry/Private/DependsNode.cpp
--- a/Engine/Source/Runtime/AssetRegistry/Private/DependsNode.cpp
+++ b/Engine/Source/Runtime/AssetRegistry/Private/DependsNode.cpp
@@ -706,7 +706,7 @@ void FDependsNode::RemoveLinks(const TUniqueFunction<bool(const FDependsNode*)>&
 }
 
 template <uint32 FlagWidth>
-void RemapLinks(const TMap<FDependsNode*, FDependsNode*>& OldNodeToNewNode,
+void RemapLinks(const TUniqueFunction<FDependsNode* (FDependsNode*)>& Remap,
 	TArray<FDependsNode*>& Dependencies,
 	TBitArray<>* FlagBits,
 	bool IsSorted,
@@ -714,8 +714,7 @@ void RemapLinks(const TMap<FDependsNode*, FDependsNode*>& OldNodeToNewNode,
 {
 	for (FDependsNode*& OldNode : Dependencies)
 	{
-		FDependsNode*const* NewNodePtr = OldNodeToNewNode.Find(OldNode);
-		OldNode = NewNodePtr ? *NewNodePtr : nullptr;
+		OldNode = Remap(OldNode);
 	}
 	::RemoveAll<FlagWidth>([](const FDependsNode* Node)
 		{
@@ -728,13 +727,24 @@ void RemapLinks(const TMap<FDependsNode*, FDependsNode*>& OldNodeToNewNode,
 	}
 }
 
-void FDependsNode::RemapLinks(const TMap<FDependsNode*, FDependsNode*>& OldNodeToNewNode)
+void FDependsNode::RemapLinks(const TUniqueFunction<FDependsNode* (FDependsNode*)>& Remap)
+{
+	RemapDependencies(Remap);
+	RemapReferencers(Remap);
+}
+
+void FDependsNode::RemapDependencies(const TUniqueFunction<FDependsNode* (FDependsNode*)>& Remap)
+{
+	const EAllowShrinking ShrinkPolicy = bAllowShrinking ? EAllowShrinking::Yes : EAllowShrinking::No;
+	::RemapLinks<PackageFlagWidth>(Remap, PackageDependencies, &PackageFlags, PackageIsSorted, ShrinkPolicy);
+	::RemapLinks<SearchableNameFlagWidth>(Remap, NameDependencies, nullptr, SearchableNameIsSorted, ShrinkPolicy);
+	::RemapLinks<ManageFlagWidth>(Remap, ManageDependencies, &ManageFlags, ManageIsSorted, ShrinkPolicy);
+}
+
+void FDependsNode::RemapReferencers(const TUniqueFunction<FDependsNode* (FDependsNode*)>& Remap)
 {
 	const EAllowShrinking ShrinkPolicy = bAllowShrinking ? EAllowShrinking::Yes : EAllowShrinking::No;
-	::RemapLinks<PackageFlagWidth>(OldNodeToNewNode, PackageDependencies, &PackageFlags, PackageIsSorted, ShrinkPolicy);
-	::RemapLinks<SearchableNameFlagWidth>(OldNodeToNewNode, NameDependencies, nullptr, SearchableNameIsSorted, ShrinkPolicy);
-	::RemapLinks<ManageFlagWidth>(OldNodeToNewNode, ManageDependencies, &ManageFlags, ManageIsSorted, ShrinkPolicy);
-	::RemapLinks<0>(OldNodeToNewNode, Referencers, nullptr, ReferencersIsSorted, ShrinkPolicy);
+	::RemapLinks<0>(Remap, Referencers, nullptr, ReferencersIsSorted, ShrinkPolicy);
 }
 
 bool FDependsNode::ContainsDependency(const FDependsNode* InDependency,
diff --git a/Engine/Source/Runtime/AssetRegistry/Private/DependsNode.h b/Engine/Source/Runtime/AssetRegistry/Private/DependsNode.h
--- a/Engine/Source/Runtime/AssetRegistry/Private/DependsNode.h
+++ b/Engine/Source/Runtime/AssetRegistry/Private/DependsNode.h
@@ -91,11 +91,21 @@ public:
 	void RemoveManageReferencesToNode();
 	/** Remove all nodes from referencers and dependencies for which ShouldRemove returns true */
 	void RemoveLinks(const TUniqueFunction<bool(const FDependsNode*)>& ShouldRemove);
-	/**
-	 * Replace all nodes in referencers and dependencies with the newnode for them in OldNodeToNewNode, and remove any
-	 * that have no entry or that remap to null.
-	 */
-	void RemapLinks(const TMap<FDependsNode*, FDependsNode*>& OldNodeToNewNode);
+	/** Replace nodes in referencers and dependencies with OldNodeToNewNode[Target]; remove any that remap to null. */
+	void RemapLinks(const TMap<FDependsNode*, FDependsNode*>& OldNodeToNewNode)
+	{
+		RemapLinks([&OldNodeToNewNode](FDependsNode* Target)
+			{
+				FDependsNode*const* NewNodePtr = OldNodeToNewNode.Find(Target);
+				return NewNodePtr ? *NewNodePtr : nullptr;
+			});
+	}
+	/** Replace nodes in referencers and dependencies with the Remapper[Target]; remove any that Remap to null. */
+	void RemapLinks(const TUniqueFunction<FDependsNode* (FDependsNode*)>& Remapper);
+	/** Replace nodes in dependencies with the Remapper[Target]; remove any that Remap to null. */
+	void RemapDependencies(const TUniqueFunction<FDependsNode* (FDependsNode*)>& Remapper);
+	/** Replace nodes in referencers with the Remapper[Target]; remove any that Remap to null. */
+	void RemapReferencers(const TUniqueFunction<FDependsNode* (FDependsNode*)>& Remapper);
 
 	/** Returns number of connections this node has, both references and dependencies */
 	int32 GetConnectionCount() const;
diff --git a/Engine/Source/Runtime/AssetRegistry/Private/DependsNodeTest.cpp b/Engine/Source/Runtime/AssetRegistry/Private/DependsNodeTest.cpp
--- a/Engine/Source/Runtime/AssetRegistry/Private/DependsNodeTest.cpp
+++ b/Engine/Source/Runtime/AssetRegistry/Private/DependsNodeTest.cpp
@@ -1226,6 +1226,81 @@ bool FDependsNodeTest::RunTest(const FString& Parameters)
 			TestTrue(TEXT("RefreshReferencers"), CombinationEquals<FDependsNode* const>(Referencers, TArrayView<FDependsNode* const>({ &Nodes[0], &Nodes[6], &Nodes[0], &Nodes[6] })));
 		}
 	}
+
+	// RemapLinks
+	{
+		FDependsNode A;
+		FDependsNode B;
+		FDependsNode C;
+		FDependsNode D;
+		FDependsNode RemapA;
+		FDependsNode RemapC;
+		A.AddDependency(&B, EDependencyCategory::Package, EDependencyProperty::None);
+		A.AddDependency(&B, EDependencyCategory::Manage, EDependencyProperty::None);
+		A.AddDependency(&B, EDependencyCategory::SearchableName, EDependencyProperty::None);
+		A.AddDependency(&C, EDependencyCategory::Package, EDependencyProperty::None);
+		A.AddDependency(&D, EDependencyCategory::Package, EDependencyProperty::None);
+		B.AddDependency(&C, EDependencyCategory::Package, EDependencyProperty::None);
+		B.AddReferencer(&A);
+		C.AddReferencer(&A);
+		C.AddReferencer(&B);
+		D.AddReferencer(&A);
+
+		TUniqueFunction<FDependsNode* (FDependsNode*)> Remapper = 
+			[&A, &B, &C, &RemapA, &RemapC](FDependsNode* Target)
+			{
+				if (Target == &A)
+				{
+					return &RemapA;
+				}
+				if (Target == &B)
+				{
+					return (FDependsNode*)nullptr;
+				}
+				if (Target == &C)
+				{
+					return &RemapC;
+				}
+				return Target;
+			};
+		A.RemapLinks(Remapper);
+		B.RemapLinks(Remapper);
+		C.RemapLinks(Remapper);
+		D.RemapLinks(Remapper);
+
+		Dependencies.Reset();
+		Referencers.Reset();
+		A.GetDependencies(Dependencies);
+		A.GetReferencers(Referencers);
+		TestTrue(TEXT("RemapLinks.1"), CombinationEquals<FDependsNode*>(Dependencies, TArrayView<FDependsNode* const>({ &RemapC, &D })));
+		TestTrue(TEXT("RemapLinks.2"), CombinationEquals<FDependsNode*>(Referencers, TArrayView<FDependsNode* const>()));
+
+		Dependencies.Reset();
+		Referencers.Reset();
+		B.GetDependencies(Dependencies);
+		B.GetReferencers(Referencers);
+		TestTrue(TEXT("RemapLinks.3"), CombinationEquals<FDependsNode*>(Dependencies, TArrayView<FDependsNode* const>({ &RemapC })));
+		TestTrue(TEXT("RemapLinks.4"), CombinationEquals<FDependsNode*>(Referencers, TArrayView<FDependsNode* const>({ &RemapA })));
+
+		Dependencies.Reset();
+		Referencers.Reset();
+		C.GetDependencies(Dependencies);
+		C.GetReferencers(Referencers);
+		TestTrue(TEXT("RemapLinks.5"), CombinationEquals<FDependsNode*>(Dependencies, TArrayView<FDependsNode* const>()));
+		TestTrue(TEXT("RemapLinks.6"), CombinationEquals<FDependsNode*>(Referencers, TArrayView<FDependsNode* const>({ &RemapA })));
+
+		Dependencies.Reset();
+		Referencers.Reset();
+		D.GetDependencies(Dependencies);
+		D.GetReferencers(Referencers);
+		TestTrue(TEXT("RemapLinks.7"), CombinationEquals<FDependsNode*>(Dependencies, TArrayView<FDependsNode* const>()));
+		TestTrue(TEXT("RemapLinks.8"), CombinationEquals<FDependsNode*>(Referencers, TArrayView<FDependsNode* const>({ &RemapA })));
+
+		// Remap links does not handle adding the reverse links, so RemapA and RemapC are unmodified.
+		TestTrue(TEXT("RemapLinks.9"), RemapA.GetConnectionCount() == 0);
+		TestTrue(TEXT("RemapLinks.10"), RemapC.GetConnectionCount() == 0);
+	}
+
 	return true;
 }
 
diff --git a/Engine/Source/Runtime/AssetRegistry/Public/AssetRegistry/AssetRegistryState.h b/Engine/Source/Runtime/AssetRegistry/Public/AssetRegistry/AssetRegistryState.h
--- a/Engine/Source/Runtime/AssetRegistry/Public/AssetRegistry/AssetRegistryState.h
+++ b/Engine/Source/Runtime/AssetRegistry/Public/AssetRegistry/AssetRegistryState.h
@@ -662,7 +662,7 @@ private:
 
 	/** Find the first non-redirector dependency node starting from InDependency. */
 	FDependsNode* ResolveRedirector(FDependsNode* InDependency, const FAssetDataMap& InAllowedAssets,
-		TMap<FDependsNode*, FDependsNode*>& InCache);
+		TMap<FDependsNode*, FDependsNode*>& InCache) const;
 
 	/** Finds an existing node for the given package and returns it, or returns null if one isn't found */
 	FDependsNode* FindDependsNode(const FAssetIdentifier& Identifier) const;
