AssetRegistry GathererCache: Do not allow loading of cachefile versions older than FAssetRegistryVersion::LatestVersion - 1; we do not currenltly support loading those older versions.
#rnx
#Synced-CL 51945718

diff --git a/Engine/Source/Runtime/AssetRegistry/Private/AssetDataGatherer.cpp b/Engine/Source/Runtime/AssetRegistry/Private/AssetDataGatherer.cpp
--- a/Engine/Source/Runtime/AssetRegistry/Private/AssetDataGatherer.cpp
+++ b/Engine/Source/Runtime/AssetRegistry/Private/AssetDataGatherer.cpp
@@ -5539,6 +5539,13 @@ TArray<FCachePayload> LoadCacheFiles(TConstArrayView<FString> InCacheFilenames,
 						TEXT("There was an error loading the asset registry cache from %s."),
 						*CacheFilename);
 				}
+				else
+				{
+					// We do not support backwards compatibility after FAssetRegistryVersion::MemoryMappedTagDataStore
+					UE_LOG(LogAssetRegistry, Display,
+						TEXT("Asset registry older version (%" UINT64_FMT ") cache is too old to read; discarding %s and regathering its data."),
+						Version, *CacheFilename);
+				}
 			}
 			else // check if it was a past version
 			{
@@ -5553,16 +5560,26 @@ TArray<FCachePayload> LoadCacheFiles(TConstArrayView<FString> InCacheFilenames,
 					if (VersionOpt.IsSet() && static_cast<FAssetRegistryVersion::Type>(VersionOpt.GetValue()) < FAssetRegistryVersion::MemoryMappedTagDataStore)
 					{
 						FAssetRegistryVersion::Type Version = static_cast<FAssetRegistryVersion::Type>(VersionOpt.GetValue());
-						AssetRegistryBackwardCompat::FChecksumViewReader ChecksummingReader(MoveTemp(FileReader), *CacheFilename);
-						Payload = DoLoad(ChecksummingReader, nullptr, Version);	// cannot pass the mmaped file, as we cannot map the data in old format
-						UE_CLOG(Payload.bSucceeded, LogAssetRegistry, Display,
-							TEXT("Asset registry older version (%d) cache read as %.1f MiB from %s."),
-							Version,
-							static_cast<float>(FileReader.GetTotalSize()) / 1024.f / 1024.f, *CacheFilename);
-						UE_CLOG(!Payload.bSucceeded, LogAssetRegistry, Warning,
-							TEXT("There was an error loading the asset registry cache in old format (version %d) using memory mapping from %s."),
-							Version,
-							*CacheFilename);
+						// We only support backwards compatibility for a single version
+						if (Version >= FAssetRegistryVersion::ManageDependenciesCookRule)
+						{
+							AssetRegistryBackwardCompat::FChecksumViewReader ChecksummingReader(MoveTemp(FileReader), *CacheFilename);
+							Payload = DoLoad(ChecksummingReader, nullptr, Version);	// cannot pass the mmaped file, as we cannot map the data in old format
+							UE_CLOG(Payload.bSucceeded, LogAssetRegistry, Display,
+								TEXT("Asset registry older version (%d) cache read as %.1f MiB from %s."),
+								Version,
+								static_cast<float>(FileReader.GetTotalSize()) / 1024.f / 1024.f, *CacheFilename);
+							UE_CLOG(!Payload.bSucceeded, LogAssetRegistry, Warning,
+								TEXT("There was an error loading the asset registry cache in old format (version %d) using memory mapping from %s."),
+								Version,
+								*CacheFilename);
+						}
+						else
+						{
+							UE_LOG(LogAssetRegistry, Display,
+								TEXT("Asset registry older version (%d) cache is too old to read; discarding %s and regathering its data."),
+								Version, *CacheFilename);
+						}
 					}
 					else
 					{
@@ -5599,6 +5616,12 @@ TArray<FCachePayload> LoadCacheFiles(TConstArrayView<FString> InCacheFilenames,
 						TEXT("There was an error loading the asset registry cache from %s."),
 						*CacheFilename);
 				}
+				else
+				{
+					UE_LOG(LogAssetRegistry, Display,
+						TEXT("Asset registry older version (%" UINT64_FMT ") cache is too old to read; discarding %s and regathering its data."),
+						Version, *CacheFilename);
+				}
 			}
 
 			Results[Index] = MoveTemp(Payload);
