AssetRegistry: Bugfix for crash in EnumerateAllAssets when EEnumerateAssetsFlags::Parallel.
#rnx
#Synced-CL 51945718

diff --git a/Engine/Source/Runtime/AssetRegistry/Private/AssetRegistry.cpp b/Engine/Source/Runtime/AssetRegistry/Private/AssetRegistry.cpp
--- a/Engine/Source/Runtime/AssetRegistry/Private/AssetRegistry.cpp
+++ b/Engine/Source/Runtime/AssetRegistry/Private/AssetRegistry.cpp
@@ -18,6 +18,7 @@
 #include "Async/ParallelFor.h"
 #include "AutoRTFM.h"
 #include "Blueprint/BlueprintSupport.h"
+#include "Containers/MpscQueue.h"
 #include "DependsNode.h"
 #include "GenericPlatform/GenericPlatformChunkInstall.h"
 #include "GenericPlatform/GenericPlatformFile.h"
@@ -2945,21 +2946,27 @@ bool UAssetRegistryImpl::EnumerateAllAssets(TFunctionRef<bool(const FAssetData&)
 		}
 	}
 
-	// We have to call the callback on a copy rather than a reference since the callback may reenter the lock
-	TArray<FAssetData, TInlineAllocator<128>> OnDiskAssetDatas;
-	{
-		UE::AssetRegistry::FInterfaceReadScopeLock InterfaceScopeLock(InterfaceLock);
-		OnDiskAssetDatas.Reserve(GuardedData.State.GetNumAssets());
-		GuardedData.EnumerateAllDiskAssets(PackageNamesToSkip,
-			[&OnDiskAssetDatas](const FAssetData& AssetData)
-			{
-				OnDiskAssetDatas.Add(AssetData);
-				return true;
-			}, InEnumerateFlags);
-	}
-
 	if (EnumHasAnyFlags(InEnumerateFlags, EEnumerateAssetsFlags::Parallel))
 	{
+		// We have to call the callback on a copy rather than a reference since the callback may reenter the lock
+		TMpscQueue<FAssetData> OnDiskAssetDatasQueue;
+		TArray<FAssetData, TInlineAllocator<128>> OnDiskAssetDatas;
+		{
+			UE::AssetRegistry::FInterfaceReadScopeLock InterfaceScopeLock(InterfaceLock);
+			GuardedData.EnumerateAllDiskAssets(PackageNamesToSkip,
+				[&OnDiskAssetDatasQueue](const FAssetData& AssetData)
+				{
+					OnDiskAssetDatasQueue.Enqueue(AssetData);
+					return true;
+				}, InEnumerateFlags);
+			OnDiskAssetDatas.Reserve(GuardedData.State.GetNumAssets());
+		}
+		FAssetData AssetData;
+		while (OnDiskAssetDatasQueue.Dequeue(AssetData))
+		{
+			OnDiskAssetDatas.Add(MoveTemp(AssetData));
+		}
+
 		ParallelFor(OnDiskAssetDatas.Num(), [&OnDiskAssetDatas, &Callback](int32 Index)
 			{
 				const FAssetData& AssetData = OnDiskAssetDatas[Index];
@@ -2975,6 +2982,19 @@ bool UAssetRegistryImpl::EnumerateAllAssets(TFunctionRef<bool(const FAssetData&)
 	}
 	else
 	{
+		// We have to call the callback on a copy rather than a reference since the callback may reenter the lock
+		TArray<FAssetData, TInlineAllocator<128>> OnDiskAssetDatas;
+		{
+			UE::AssetRegistry::FInterfaceReadScopeLock InterfaceScopeLock(InterfaceLock);
+			OnDiskAssetDatas.Reserve(GuardedData.State.GetNumAssets());
+			GuardedData.EnumerateAllDiskAssets(PackageNamesToSkip,
+				[&OnDiskAssetDatas](const FAssetData& AssetData)
+				{
+					OnDiskAssetDatas.Add(AssetData);
+					return true;
+				}, InEnumerateFlags);
+		}
+
 		for (const FAssetData& AssetData : OnDiskAssetDatas)
 		{
 			if (!Callback(AssetData))
