AssetRegistry GathererCache: Do not allow loading of cachefile versions older than FAssetRegistryVersion::LatestVersion - 1; we do not currenltly support loading those older versions.

- Also disallow all backward compat for the cooked editor.

#rnx
#Synced-CL 51945718

diff --git a/Engine/Source/Runtime/AssetRegistry/Private/AssetDataGatherer.cpp b/Engine/Source/Runtime/AssetRegistry/Private/AssetDataGatherer.cpp
--- a/Engine/Source/Runtime/AssetRegistry/Private/AssetDataGatherer.cpp
+++ b/Engine/Source/Runtime/AssetRegistry/Private/AssetDataGatherer.cpp
@@ -5560,7 +5560,9 @@ TArray<FCachePayload> LoadCacheFiles(TConstArrayView<FString> InCacheFilenames,
 					if (VersionOpt.IsSet() && static_cast<FAssetRegistryVersion::Type>(VersionOpt.GetValue()) < FAssetRegistryVersion::MemoryMappedTagDataStore)
 					{
 						FAssetRegistryVersion::Type Version = static_cast<FAssetRegistryVersion::Type>(VersionOpt.GetValue());
-						// We only support backwards compatibility for a single version
+						// For a cooked editor, we don't support any older caches
+						#if UE_EDITOR && !UE_IS_COOKED_EDITOR
+						// For uncooked editor, we only support backwards compatibility for a single version
 						if (Version >= FAssetRegistryVersion::ManageDependenciesCookRule)
 						{
 							AssetRegistryBackwardCompat::FChecksumViewReader ChecksummingReader(MoveTemp(FileReader), *CacheFilename);
@@ -5575,6 +5577,7 @@ TArray<FCachePayload> LoadCacheFiles(TConstArrayView<FString> InCacheFilenames,
 								*CacheFilename);
 						}
 						else
+						#endif // UE_EDITOR && !UE_IS_COOKED_EDITOR
 						{
 							UE_LOG(LogAssetRegistry, Display,
 								TEXT("Asset registry older version (%d) cache is too old to read; discarding %s and regathering its data."),
