AssetRegistry: Fix logging of negative times for OnFilesLoaded.Broadcast

#rnx

#Synced-CL 52055016

diff --git a/Engine/Source/Runtime/AssetRegistry/Private/AssetRegistry.cpp b/Engine/Source/Runtime/AssetRegistry/Private/AssetRegistry.cpp
--- a/Engine/Source/Runtime/AssetRegistry/Private/AssetRegistry.cpp
+++ b/Engine/Source/Runtime/AssetRegistry/Private/AssetRegistry.cpp
@@ -10886,7 +10886,7 @@ void UAssetRegistryImpl::Broadcast(UE::AssetRegistry::Impl::FEventContext& Event
 			FileLoadedEvent.Broadcast();
 			ScanEndedEvent.Broadcast();
 			UE_LOG(LogAssetRegistry, Display, TEXT("Completed OnFilesLoaded.Broadcast: %.3fs"),
-				(float) FPlatformTime::Seconds() - BroadcastTimeStart);
+				FPlatformTime::Seconds() - BroadcastTimeStart);
 			EventContext.bHasSentFileLoadedEventBroadcast = true;
 			EventContext.bFileLoadedEventBroadcast = false;
 		}
