AssetRegistry: DiscoveryCache: Fix spurious error when attempting to save old cache file for diagnostics. It would be triggered twice in one read attempt if the USN journal wrapped on two drives used by the project. and the second attempt would fail because we renamed the old cache file and could no longer read it at the original location. We now attempt to save it only once.
Also, turn off the save of the file for diagnostics when the USN journal wraps; the journal wrapping is an expected occurrence and even if it is occurring unexpectedly, the data in the journal will not help diagnose it.
#rnx
#Synced-CL 52123874

diff --git a/Engine/Source/Runtime/AssetRegistry/Private/AssetDataGathererDiscoveryCache.cpp b/Engine/Source/Runtime/AssetRegistry/Private/AssetDataGathererDiscoveryCache.cpp
--- a/Engine/Source/Runtime/AssetRegistry/Private/AssetDataGathererDiscoveryCache.cpp
+++ b/Engine/Source/Runtime/AssetRegistry/Private/AssetDataGathererDiscoveryCache.cpp
@@ -120,6 +120,7 @@ void FAssetDataDiscoveryCache::LoadAndUpdateCache()
 		return;
 	}
 
+	bool bReportCacheErrorDiagnostics = false;
 	for (TPair<FString, FCachedVolumeInfo>& VolumePair : CachedVolumes)
 	{
 		FString& VolumeName = VolumePair.Key;
@@ -191,6 +192,7 @@ void FAssetDataDiscoveryCache::LoadAndUpdateCache()
 					TEXT("PlatformFileJournal is not available for volume '%s'. AssetRegistry discovery of files on this volume will be uncached. Unavailability reason:")
 					TEXT("\n\t%s"),
 					*VolumeName, *VolumeInfo.LastError);
+				bReportCacheErrorDiagnostics = true;
 
 				UE::Telemetry::AssetRegistry::FFileJournalErrorTelemetry MissingJournalTelemetry;
 				MissingJournalTelemetry.Directory = VolumeName;
@@ -206,19 +208,6 @@ void FAssetDataDiscoveryCache::LoadAndUpdateCache()
 			VolumeInfo.JournalId = VolumeInfo.JournalIdOnDisk;
 			VolumeInfo.NextJournalEntryToScan = VolumeInfo.NextJournalEntryOnDisk;
 			VolumeInfo.Dirs.Empty();
-
-#if WITH_EDITOR
-			// Rename the discovery cache so it may remain on disk for inspection
-			FString CachePath = GetCacheFileName();
-			FStringView CacheDir, CacheFileName, CacheExtension;
-			FPathViews::Split(CachePath, CacheDir, CacheFileName, CacheExtension);
-			FString NewCacheFileName = FString::Printf(TEXT("%.*s-%s.%.*s"), CacheFileName.Len(), CacheFileName.GetData(), *FDateTime::Now().ToString(), CacheExtension.Len(), CacheExtension.GetData());
-			FString NewCachePath = FPaths::Combine(CacheDir, NewCacheFileName);
-			if (!IFileManager::Get().Move(*NewCachePath, *CachePath, true /*Replace*/, true /*EvenIfReadOnly*/))
-			{
-				UE_LOG(LogAssetRegistry, Warning, TEXT("Failed creating a copy of %s at %s"), *CachePath, *NewCachePath);
-			}
-#endif
 		}
 		else
 		{
@@ -241,6 +230,23 @@ void FAssetDataDiscoveryCache::LoadAndUpdateCache()
 			}
 		}
 	}
+
+	if (bReportCacheErrorDiagnostics)
+	{
+#if WITH_EDITOR
+		// Rename the discovery cache so it may remain on disk for inspection
+		FString CachePath = GetCacheFileName();
+		FStringView CacheDir, CacheFileName, CacheExtension;
+		FPathViews::Split(CachePath, CacheDir, CacheFileName, CacheExtension);
+		FString NewCacheFileName = FString::Printf(TEXT("%.*s-%s.%.*s"), CacheFileName.Len(), CacheFileName.GetData(), *FDateTime::Now().ToString(), CacheExtension.Len(), CacheExtension.GetData());
+		FString NewCachePath = FPaths::Combine(CacheDir, NewCacheFileName);
+		UE_LOG(LogAssetRegistry, Warning, TEXT("Saving a copy of %s to %s for debug diagnostics."), *CachePath, *NewCachePath);
+		if (!IFileManager::Get().Move(*NewCachePath, *CachePath, true /*Replace*/, true /*EvenIfReadOnly*/))
+		{
+			UE_LOG(LogAssetRegistry, Warning, TEXT("Failed creating a copy of %s at %s"), *CachePath, *NewCachePath);
+		}
+#endif
+	}
 }
 
 void FAssetDataDiscoveryCache::SaveCache()
