AssetRegistry: AssetDataGatherer: check for an invalid AssetDataTag and skip saving the cache if present.
This is a temporary workaround while we try to find the source of the invalid tag; it is supposed to be impossible to add one, and this search for one is expensive enough that we should skip it when no longer needed.
#Synced-CL 52271479

diff --git a/Engine/Source/Runtime/AssetRegistry/Private/AssetDataGatherer.cpp b/Engine/Source/Runtime/AssetRegistry/Private/AssetDataGatherer.cpp
--- a/Engine/Source/Runtime/AssetRegistry/Private/AssetDataGatherer.cpp
+++ b/Engine/Source/Runtime/AssetRegistry/Private/AssetDataGatherer.cpp
@@ -4162,6 +4162,31 @@ bool FAssetDataGatherer::TrySaveCacheFile(const TArray<TPair<FName,FDiskCachedAs
 			uint64 Hash = CityHash64((const char*)Buffer.GetData(), Buffer.Len() * sizeof(TCHAR));
 			int32 Index = (Hash & ShardMask);
 			DataPerShard[Index].Add(Entry);
+			// UE-372209: Work-around a not-yet-diagnosed AssetRegistry corruption bug that is allowing empty strings
+			// to be present in some AssetRegistry tag values. Skip saving the cache if that has occurred; we don't know
+			// the extent of the corruption.
+			for (const FAssetData& AssetData : Entry.Value->AssetDataList)
+			{
+				for (const TPair<FName, FAssetTagValueRef>& TagAndValue : AssetData.TagsAndValues)
+				{
+					if (TagAndValue.Value.GetValue().IsEmpty())
+					{
+						ensureMsgf(false,
+							TEXT("AssetRegistry corruption detected: the TagValue on an AssetData is empty string, which is invalid. The cache save will be skipped.")
+							TEXT("\n\tAsset: '%s'")
+							TEXT("\n\tKey: '%s'")
+							TEXT("\n\tValue: '%s'"),
+							*AssetData.ToSoftObjectPath().ToString(), *TagAndValue.Key.ToString(), *TagAndValue.Value.GetValue());
+
+						FScopedGatheringPause ScopedPause(*this);
+						FGathererScopeLock TickScopeLock(&TickLock);
+						check(bIsSavingAsyncCache);
+						bIsSavingAsyncCache = false;
+
+						return true;
+					}
+				}
+			}
 		}
 	}
 
