AssetRegistry: AssetDataGatherer: check for an invalid AssetDataTag and clear the cache and skip saving the new cache if present.
This is a temporary workaround while we try to find the source of the invalid tag; it is supposed to be impossible to add one, and this search for one is expensive enough that we should skip it when no longer needed.
Fix race condition causing cache not to be saved when ClearCache is requested on another thread, due to legacy read of bCacheWriteEnabled outside of the lock; that read is is no longer necessary because the flag is checked earlier (inside the lock) for all cases that reach the SaveCacheFileInternal function.
#Synced-CL 52434073

diff --git a/Engine/Source/Runtime/AssetRegistry/Private/AssetDataGatherer.cpp b/Engine/Source/Runtime/AssetRegistry/Private/AssetDataGatherer.cpp
--- a/Engine/Source/Runtime/AssetRegistry/Private/AssetDataGatherer.cpp
+++ b/Engine/Source/Runtime/AssetRegistry/Private/AssetDataGatherer.cpp
@@ -378,7 +378,7 @@ struct FCachePayload
 void SerializeCacheSave(FAssetRegistryWriter& Ar, const TArray<TPair<FName, FDiskCachedAssetData*>>& AssetsToSave);
 FCachePayload SerializeCacheLoad(FAssetRegistryReader& Ar);
 TArray<FCachePayload> LoadCacheFiles(TConstArrayView<FString> CacheFilenames, bool bCacheWriteEnabled);
-
+void DeleteCacheFilesWithinLock();
 
 /** InOutResult = Value, but without shrinking the string to fit. */
 void AssignStringWithoutShrinking(FString& InOutResult, FStringView Value)
@@ -4138,6 +4138,8 @@ bool FAssetDataGatherer::TrySaveCacheFile(const TArray<TPair<FName,FDiskCachedAs
 	const int32 ShardMask = CacheShards - 1;
 	std::atomic<int64> TotalCacheSize{0};
 	bool bSuccess = true;
+	FString CorruptTagText;
+	bool bCorruptTagDetected = false;
 
 	TArray<TArray<TPair<FName, FDiskCachedAssetData*>>> DataPerShard;
 	{
@@ -4155,13 +4157,39 @@ bool FAssetDataGatherer::TrySaveCacheFile(const TArray<TPair<FName,FDiskCachedAs
 			Buffer << Entry.Key;
 			// Force to lower case to maintain ordering across runs where fname casing changes because of load order
 			// Hash doesn't need to be persistent across utf-8 transition so no need to re-encode strings
-			for (int32 i=0; i < Buffer.Len(); ++i)
+			for (int32 i = 0; i < Buffer.Len(); ++i)
 			{
 				Buffer.GetData()[i] = FChar::ToLower(Buffer.GetData()[i]);
 			}
 			uint64 Hash = CityHash64((const char*)Buffer.GetData(), Buffer.Len() * sizeof(TCHAR));
 			int32 Index = (Hash & ShardMask);
 			DataPerShard[Index].Add(Entry);
+			// UE-372209: Work-around a not-yet-diagnosed AssetRegistry corruption bug that is allowing empty strings
+			// to be present in some AssetRegistry tag values. Delete the old cache and skip saving the cache if that
+			// has occurred; we don't know the extent of the corruption.
+			for (const FAssetData& AssetData : Entry.Value->AssetDataList)
+			{
+				if (bCorruptTagDetected)
+				{
+					break;
+				}
+
+				for (const TPair<FName, FAssetTagValueRef>& TagAndValue : AssetData.TagsAndValues)
+				{
+					if (TagAndValue.Value.GetStorageString().IsEmpty())
+					{
+						bCorruptTagDetected = true;
+						CorruptTagText = FString::Printf(
+							TEXT("AssetRegistry corruption detected: the TagValue on an AssetData is empty string, which is invalid. ")
+							TEXT("The existing cache will be deleted and the new cache save will be skipped.")
+							TEXT("\n\tAsset: '%s'")
+							TEXT("\n\tKey: '%s', Value: '%s', IsSet: %s, IsFixed: %s"),
+							*AssetData.ToSoftObjectPath().ToString(), *TagAndValue.Key.ToString(), *TagAndValue.Value.GetStorageString(),
+							*LexToString(TagAndValue.Value.IsSet()), *LexToString(TagAndValue.Value.IsFixed()));
+						break;
+					}
+				}
+			}
 		}
 	}
 
@@ -4174,9 +4202,9 @@ bool FAssetDataGatherer::TrySaveCacheFile(const TArray<TPair<FName,FDiskCachedAs
 	}
 	double CurrentTime = FPlatformTime::Seconds();
 	constexpr float WaitTimeBeforeReopen = 0.5f;
-	if (CurrentTime < LastCacheWriteTime + WaitTimeBeforeReopen)
+	if (CurrentTime < LocalLastCacheWriteTime + WaitTimeBeforeReopen)
 	{
-		FPlatformProcess::Sleep(static_cast<float>(LastCacheWriteTime + WaitTimeBeforeReopen - CurrentTime));
+		FPlatformProcess::Sleep(static_cast<float>(LocalLastCacheWriteTime + WaitTimeBeforeReopen - CurrentTime));
 	}
 
 #if CACHEFILE_SYSTEMWIDE_WRITELOCK
@@ -4196,38 +4224,48 @@ bool FAssetDataGatherer::TrySaveCacheFile(const TArray<TPair<FName,FDiskCachedAs
 	else
 #endif
 	{
-		ParallelFor(CacheShards, [this, &AssetsToSave, &TotalCacheSize, &DataPerShard, ShardMask](int32 Shard)
+		if (ensureMsgf(!bCorruptTagDetected, TEXT("%s"), *CorruptTagText))
 		{
-			TArray<TPair<FName, FDiskCachedAssetData*>>& ShardData = DataPerShard[Shard];
-			FString Filename = FString::Printf(TEXT("%s_%d.bin"), *GGatherSettings.GetCacheBaseFilename(), Shard);
-			int64 CacheSize = SaveCacheFileInternal(Filename, ShardData);
-			TotalCacheSize += CacheSize;
-		}, EParallelForFlags::BackgroundPriority);
-
-		UE_LOG(LogAssetRegistry, Display, TEXT("Asset registry cache written as %.1f MiB to %s_*.bin"),
-			static_cast<float>(TotalCacheSize) / 1024.f / 1024.f, *GGatherSettings.GetCacheBaseFilename());
-		// Delete old name of monolithic cache file and old non-monolithic cache directory if they exist
-		IFileManager::Get().Delete(*GGatherSettings.GetLegacyCacheFilename(), false /* bRequireExists */,
-			true /* EvenReadOnly */, true /* Quiet */);
-		IFileManager::Get().DeleteDirectory(*GGatherSettings.GetLegacyNonMonolithicCacheDirectory(),
-			false /* bRequireExists */, true /* Tree */);
-
-		// Delete any other shards if number of shards was reduced
-		TArray<FString> CacheFiles = GGatherSettings.FindShardedCacheFiles();
-		for (const FString& CacheFile : CacheFiles)
-		{
-			FStringView BaseName = FPathViews::GetBaseFilename(CacheFile);
-			if (int32 Index; BaseName.FindLastChar('_', Index))
+			ParallelFor(CacheShards, [this, &AssetsToSave, &TotalCacheSize, &DataPerShard, ShardMask](int32 Shard)
+				{
+					TArray<TPair<FName, FDiskCachedAssetData*>>& ShardData = DataPerShard[Shard];
+					FString Filename = FString::Printf(TEXT("%s_%d.bin"), *GGatherSettings.GetCacheBaseFilename(), Shard);
+					int64 CacheSize = SaveCacheFileInternal(Filename, ShardData);
+					TotalCacheSize += CacheSize;
+				}, EParallelForFlags::BackgroundPriority);
+
+			UE_LOG(LogAssetRegistry, Display, TEXT("Asset registry cache written as %.1f MiB to %s_*.bin"),
+				static_cast<float>(TotalCacheSize) / 1024.f / 1024.f, *GGatherSettings.GetCacheBaseFilename());
+			// Delete old name of monolithic cache file and old non-monolithic cache directory if they exist
+			IFileManager::Get().Delete(*GGatherSettings.GetLegacyCacheFilename(), false /* bRequireExists */,
+				true /* EvenReadOnly */, true /* Quiet */);
+			IFileManager::Get().DeleteDirectory(*GGatherSettings.GetLegacyNonMonolithicCacheDirectory(),
+				false /* bRequireExists */, true /* Tree */);
+
+			// Delete any other shards if number of shards was reduced
+			TArray<FString> CacheFiles = GGatherSettings.FindShardedCacheFiles();
+			for (const FString& CacheFile : CacheFiles)
 			{
-				BaseName.MidInline(Index + 1);
-				int32 Suffix = 0;
-				LexFromString(Suffix, BaseName);
-				if (Suffix >= CacheShards)
+				FStringView BaseName = FPathViews::GetBaseFilename(CacheFile);
+				if (int32 Index; BaseName.FindLastChar('_', Index))
 				{
-					IFileManager::Get().Delete(*CacheFile);
+					BaseName.MidInline(Index + 1);
+					int32 Suffix = 0;
+					LexFromString(Suffix, BaseName);
+					if (Suffix >= CacheShards)
+					{
+						IFileManager::Get().Delete(*CacheFile);
+					}
 				}
 			}
 		}
+		else
+		{
+			// Corruption detected, delete the cache files
+			UE::AssetDataGather::Private::DeleteCacheFilesWithinLock();
+			// Report this as a "successful" write, so that we stop attempting to write the cache files.
+			bSuccess = true;
+		}
 
 		FGathererScopeLock ResultsScopeLock(&ResultsLock);
 		LastCacheWriteTime = FPlatformTime::Seconds();
@@ -4965,6 +5003,10 @@ void FAssetDataGatherer::ClearCache()
 
 	{
 		FGathererScopeLock TickScopeLock(&TickLock);
+		// bIsSavingAsyncCache should now be false. We finished and set it to false if we were asked to save the cache,
+		// and no other thread should be able to set it to true because we set bCacheWriteEnabled=false in the same critical
+		// section that we used to test for bShouldSaveCache.,
+		check(!bIsSavingAsyncCache);
 		NewCachedAssetDataMap.Empty();
 		DiskCachedAssetDataMap.Empty();
 
@@ -5284,7 +5326,7 @@ UE_TRACE_EVENT_END()
 
 int64 FAssetDataGatherer::SaveCacheFileInternal(const FString& CacheFilename, const TArray<TPair<FName,FDiskCachedAssetData*>>& AssetsToSave)
 {
-	if (CacheFilename.IsEmpty() || !bCacheWriteEnabled)
+	if (CacheFilename.IsEmpty())
 	{
 		return 0;
 	}
@@ -5634,6 +5676,15 @@ TArray<FCachePayload> LoadCacheFiles(TConstArrayView<FString> InCacheFilenames,
 	return MoveTemp(Results);
 }
 
+void DeleteCacheFilesWithinLock()
+{
+	TArray<FString> CacheFiles = GGatherSettings.FindShardedCacheFiles();
+	for (const FString& CacheFile : CacheFiles)
+	{
+		IFileManager::Get().Delete(*CacheFile);
+	}
+}
+
 } // namespace UE::AssetDataGather::Private
 
 SIZE_T FAssetDataGatherer::GetAllocatedSize() const
diff --git a/Engine/Source/Runtime/CoreUObject/Private/AssetRegistry/AssetDataTagMap.cpp b/Engine/Source/Runtime/CoreUObject/Private/AssetRegistry/AssetDataTagMap.cpp
--- a/Engine/Source/Runtime/CoreUObject/Private/AssetRegistry/AssetDataTagMap.cpp
+++ b/Engine/Source/Runtime/CoreUObject/Private/AssetRegistry/AssetDataTagMap.cpp
@@ -987,7 +987,7 @@ namespace FixedTagPrivate
 			// Zero-length Tag Values are invalid. They are supposed to be prevented by early exits in
 			// UE::AssetRegistry::ReadPackageDataMain and FAssetData::SetTagsAndAssetBundles.
 			checkf(String.Len() > 0, TEXT("Invalid Zero-length tag: Key: \"%s\", Value: \"%s\" IsFixed: %s."),
-				*Key.ToString(), *String, *LexToString(Value.IsFixed() != 0));
+				*Key.ToString(), *String, *LexToString(Value.IsFixed()));
 			if (FCString::IsPureAnsi(*String, String.Len()))
 			{
 				return FValueId{EValueType::AnsiString,	 AnsiStrings.Index(MoveTemp(String))};
diff --git a/Engine/Source/Runtime/CoreUObject/Public/AssetRegistry/AssetDataTagMap.h b/Engine/Source/Runtime/CoreUObject/Public/AssetRegistry/AssetDataTagMap.h
--- a/Engine/Source/Runtime/CoreUObject/Public/AssetRegistry/AssetDataTagMap.h
+++ b/Engine/Source/Runtime/CoreUObject/Public/AssetRegistry/AssetDataTagMap.h
@@ -206,7 +206,6 @@ class FAssetTagValueRef
 		uint64 Bits = 0;
 	};
 
-	uint64 IsFixed() const { return Fixed.IsFixed(); }
 	FixedTagPrivate::FValueHandle AsFixed() const;
 	const FString& AsLoose() const;
 
@@ -221,6 +220,7 @@ public:
 	FAssetTagValueRef& operator=(FAssetTagValueRef&&) = default;
 
 	bool										IsSet() const { return Bits != 0; }
+	bool										IsFixed() const { return Fixed.IsFixed() != 0; }
 
 	COREUOBJECT_API FString						AsString() const;
 	COREUOBJECT_API FName						AsName() const;
