BACKUP: Add support for more package sections to PackageReader for UnrealPackageTool.
- SearchableNames
- MetaData
- ImportTypeHierarchies
- CellImports
- CellExports
- WorldTileInfo
- PackageTrailer
- DataResourceMap
#backup of 48614661

#Synced-CL 52538444
#Synced-CL 52538444

diff --git a/Engine/Source/Runtime/AssetRegistry/Internal/AssetRegistry/PackageReader.h b/Engine/Source/Runtime/AssetRegistry/Internal/AssetRegistry/PackageReader.h
--- a/Engine/Source/Runtime/AssetRegistry/Internal/AssetRegistry/PackageReader.h
+++ b/Engine/Source/Runtime/AssetRegistry/Internal/AssetRegistry/PackageReader.h
@@ -11,11 +11,16 @@
 struct FAssetData;
 class FAssetPackageData;
 class FLinkerTables;
+class FMetaData;
 class FPackageDependencyData;
+class FWorldTileInfo;
+namespace UE { class FPackageTrailer; }
 namespace UE::AssetRegistry { enum class EExtraDependencyFlags : uint32; }
+namespace UE::Serialization::Private { class FImportTypeHierarchy; }
 struct FGatherableTextData;
 struct FObjectFullNameAndThumbnail;
 
+
 class FPackageReader : public FArchiveUObject
 {
 public:
@@ -88,11 +93,39 @@ public:
 	/** Reads the list of thumbnails if not already read and returns a copy. */
 	ASSETREGISTRY_API bool GetThumbnails(TArray<FObjectFullNameAndThumbnail>& OutThumbnails);
 
+	/** Reads the searchable names and returns them. Not cached, reparsed each call. */
+	ASSETREGISTRY_API bool ReadSearchableNames(TMap<FTopLevelAssetPath, TArray<FName>>& OutSearchableNames);
 	/** Reads information from the AR section and converts it to FAssetData. Not cached, reparsed each call. */
 	ASSETREGISTRY_API bool ReadAssetRegistryData(TArray<FAssetData*>& AssetDataList, bool& bOutIsCookedWithoutAssetData);
 
+#if WITH_METADATA
+	/** 
+	 * Reads information from the metadata section. Not cached, reparsed each call. 
+	 * Note that metadata redirects are not applied here, what's in the file is returned as-is.
+	 */
+	ASSETREGISTRY_API bool ReadMetaData(FMetaData& OutMetaData);
+#endif // WITH_METADATA
+
 	/** Reads Imports if not already read and parses the names of imported classes out of the Imports. */
 	ASSETREGISTRY_API bool ReadImportedClasses(TArray<FName>& OutClassNames);
+	
+	/** Reads import type hierarchy information. Not cached, reparsed each call. */ 
+	ASSETREGISTRY_API bool ReadImportTypeHierarchies(TMap<FSoftObjectPath, UE::Serialization::Private::FImportTypeHierarchy>& OutTypeHierarchies);
+
+	/** Reads VCell exports. Not cached, reparsed each call. */
+	ASSETREGISTRY_API bool ReadCellExports(TArray<FCellExport>& OutExports);
+
+	/** Reads VCell imports. Not cached, reparsed each call. */
+	ASSETREGISTRY_API bool ReadCellImports(TArray<FCellImport>& OutImports);
+
+	/** Reads FWorldTileInfo. Not cached, reparsed each call. */
+	ASSETREGISTRY_API bool ReadWorldTileInfo(FWorldTileInfo& OutInfo);
+	
+	/** Reads the package trailer. Not cached, reparsed each call. */
+	ASSETREGISTRY_API bool ReadPackageTrailer(UE::FPackageTrailer& OutTrailer);
+
+	/** Reads the data resource map. Parses export indices in FObjectDataResource into FSoftObjectPath. Not cached, reparsed each call. */
+	ASSETREGISTRY_API bool ReadDataResourceMap(TArray<TTuple<FSoftObjectPath, FObjectDataResource>>& OutDataResources);
 
 	/**
 	 * Reads LinkerTables if not already read, reads the Editor-Only flags (not cached, reparsed each call), 
@@ -122,6 +155,7 @@ public:
 	ASSETREGISTRY_API virtual int64 Tell() override;
 	ASSETREGISTRY_API virtual int64 TotalSize() override;
 	ASSETREGISTRY_API virtual FArchive& operator<<( FName& Name ) override;
+	ASSETREGISTRY_API virtual FArchive& operator<<(FSoftObjectPath& Name) override;
 	virtual FString GetArchiveName() const override
 	{
 		return PackageFilename;
@@ -150,6 +184,12 @@ private:
 		TArray<TPair<FName, UE::AssetRegistry::EExtraDependencyFlags>>& OutExtraPackageDependencies);
 	bool SerializePackageTrailer(FAssetPackageData& PackageData);
 
+	/** 
+	 * Convert an index into the linker tables to a full object path. Requires ImportMap and ExportMap to have been serialized.
+	 * Less efficient than ConvertLinkerTableToPaths.
+	 */
+	TOptional<FSoftObjectPath> LinkerIndexToSoftObjectPath(FPackageIndex Index);
+
 	void ApplyRelocationToImportMapAndSoftPackageReferenceList(FStringView LoadedPackageName,
 		TArray<FName>& InOutSoftPackageReferenceList, 
 		TArray<TPair<FName, UE::AssetRegistry::EExtraDependencyFlags>>& InOutExtraPackageDependencies);
diff --git a/Engine/Source/Runtime/AssetRegistry/Private/PackageReader.cpp b/Engine/Source/Runtime/AssetRegistry/Private/PackageReader.cpp
--- a/Engine/Source/Runtime/AssetRegistry/Private/PackageReader.cpp
+++ b/Engine/Source/Runtime/AssetRegistry/Private/PackageReader.cpp
@@ -2,6 +2,7 @@
 
 #include "AssetRegistry/PackageReader.h"
 
+#include "Algo/Reverse.h"
 #include "AssetRegistry.h"
 #include "AssetRegistryPrivate.h"
 #include "AssetRegistry/AssetData.h"
@@ -15,9 +16,11 @@
 #include "Misc/PackageName.h"
 #include "Misc/PackagePath.h"
 #include "Misc/ScopeExit.h"
+#include "Serialization/ImportTypeHierarchy.h"
 #include "UObject/Class.h"
 #include "UObject/Linker.h"
 #include "UObject/LinkerLoad.h"
+#include "UObject/MetaData.h"
 #include "UObject/PackageRelocation.h"
 #include "UObject/PackageTrailer.h"
 
@@ -436,6 +439,279 @@ bool FPackageReader::ReadImportedClasses(TArray<FName>& OutClassNames)
 	return true;
 }
 
+bool FPackageReader::ReadImportTypeHierarchies(TMap<FSoftObjectPath, UE::Serialization::Private::FImportTypeHierarchy>& OutTypeHierarchies)
+{
+	// Skip export map as there should be no paths including exports in import type hierarchies. 
+	// If there are and the export map has not been serialized, warnings will be logged.
+	if (!SerializeNameMap() || !SerializeImportMap())
+	{
+		return false;
+	}
+
+	if (!StartSerializeSection(PackageFileSummary.ImportTypeHierarchiesOffset))
+	{
+		UE_LOGFMT_LOC(LogAssetRegistry, Warning,"ImportTypeHierarchies_BadOffset",
+			"Package is unloadable: {FileName}. Reason: ImportTypeHierarchiesOffset {Offset} was invalid. ",
+			("FileName", PackageFilename), ("Offset", PackageFileSummary.ImportTypeHierarchiesOffset));
+		return false;
+	}
+
+	int32 Count = 0;
+	*this << Count;
+	if (IsError())
+	{
+		UE_LOGFMT_LOC(LogAssetRegistry, Warning,"ImportTypeHierarchies_CountSerializeFail",
+			"Package is unloadable: {FileName}. Reason: Failed to serialize import type hierarchy count. ",
+			("FileName", PackageFilename));
+		return false;
+	}
+	if (Count != PackageFileSummary.ImportTypeHierarchiesCount)
+	{
+		UE_LOGFMT_LOC(LogAssetRegistry, Warning,"ImportTypeHierarchies_CountMismatch",
+			"Package is unloadable: {FileName}. Reason: ImportTypeHierarchiesCount {HeaderCount} did not match embedded count {EmbeddedCount}. ",
+			("FileName", PackageFilename), ("HeaderCount", PackageFileSummary.ImportTypeHierarchiesCount), ("EmbeddedCount", Count));
+		return false;
+	}
+
+	bool bSuccess = true;
+	for (int32 i=0; i < PackageFileSummary.ImportTypeHierarchiesCount; ++i)
+	{
+		TPair<FPackageIndex, UE::Serialization::Private::FImportTypeHierarchy> Pair;
+		*this << Pair;
+		if (IsError())
+		{
+			UE_LOGFMT_LOC(LogAssetRegistry, Warning,"ImportTypeHierarchies_SerializeFail",
+				"Package is unloadable: {FileName}. Reason: Failed to serialize import type hierarchy at index {Index}. ",
+				("FileName", PackageFilename), ("Index", i));
+			return false;
+		}
+
+		TOptional<FSoftObjectPath> Path = LinkerIndexToSoftObjectPath(Pair.Key);
+		if (Path.IsSet())
+		{
+			OutTypeHierarchies.Add(Path.GetValue(), MoveTemp(Pair.Value));
+		}
+		else
+		{
+			UE_LOGFMT_LOC(LogAssetRegistry, Warning,"ImportTypeHierarchies_BadPath",
+				"Package is unloadable: {FileName}. Reason: Package index {Index} could not be resolved to a path. ",
+				("FileName", PackageFilename), ("Index", Pair.Key.ForDebugging()));
+			bSuccess = false;
+		}
+	}
+
+	return bSuccess; 
+}
+
+bool FPackageReader::ReadCellExports(TArray<FCellExport>& OutExports)
+{
+	if (PackageFileSummary.CellExportCount == 0)
+	{
+		return true;
+	}
+
+	if (!StartSerializeSection(PackageFileSummary.CellExportOffset))
+	{
+		UE_LOGFMT_LOC(LogAssetRegistry, Warning,"CellExports_BadOffset",
+			"Package is unloadable: {FileName}. Reason: CellExportOffset {Offset} was invalid. ",
+			("FileName", PackageFilename), ("Offset", PackageFileSummary.CellExportOffset));
+		return false;
+	}
+
+	OutExports.Reserve(PackageFileSummary.CellExportCount);
+	for (int32 i=0; i < PackageFileSummary.CellExportCount; ++i)
+	{
+		FCellExport& Export = OutExports.Emplace_GetRef();
+		*this << Export;
+
+		if (IsError())
+		{
+			UE_LOGFMT_LOC(LogAssetRegistry, Warning,"CellExport_SerializeFailed",
+				"Package is unloadable: {FileName}. Reason: Failed to serialize cell export at index {Index}. ",
+				("FileName", PackageFilename), ("Index", i));
+			return false;
+		}
+	}
+
+	return true;
+}
+
+bool FPackageReader::ReadCellImports(TArray<FCellImport>& OutImports)
+{
+	if (PackageFileSummary.CellImportCount == 0)
+	{
+		return true;
+	}
+
+	if (!StartSerializeSection(PackageFileSummary.CellImportOffset))
+	{
+		UE_LOGFMT_LOC(LogAssetRegistry, Warning,"CellImports_BadOffset",
+			"Package is unloadable: {FileName}. Reason: CellImportOffset {Offset} was invalid. ",
+			("FileName", PackageFilename), ("Offset", PackageFileSummary.CellImportOffset));
+		return false;
+	}
+
+	OutImports.Reserve(PackageFileSummary.CellImportCount);
+	for (int32 i=0; i < PackageFileSummary.CellImportCount; ++i)
+	{
+		FCellImport& Import = OutImports.Emplace_GetRef();
+		*this << Import;
+
+		if (IsError())
+		{
+			UE_LOGFMT_LOC(LogAssetRegistry, Warning,"CellImport_SerializeFailed",
+				"Package is unloadable: {FileName}. Reason: Failed to serialize cell import at index {Index}. ",
+				("FileName", PackageFilename), ("Index", i));
+			return false;
+		}
+	}
+
+	return true;
+}
+
+bool FPackageReader::ReadSearchableNames(TMap<FTopLevelAssetPath, TArray<FName>>& OutSearchableNames)
+{
+	if (!SerializeNameMap() || !SerializeImportMap() || !SerializeExportMap())
+	{
+		return false;
+	}
+
+	FLinkerTables SearchableNames;
+	if (!SerializeSearchableNamesMap(SearchableNames))
+	{
+		return false;
+	}
+
+	bool bSuccess = true;
+	for (const TPair<FPackageIndex, TArray<FName>>& Pair : SearchableNames.SearchableNamesMap)
+	{
+		FPackageIndex Index = Pair.Key;
+		TArray<FName> PathElements;
+		while (!Index.IsNull())
+		{
+			FObjectResource* Obj = nullptr;
+			if (Index.IsImport() && ImportMap.IsValidIndex(Index.ToImport()))
+			{
+				Obj = &ImportMap[Index.ToImport()];
+			}
+			else if (Index.IsExport() && ExportMap.IsValidIndex(Index.ToExport()))
+			{
+				Obj = &ExportMap[Index.ToExport()];
+			}  
+			if (Obj)
+			{
+				PathElements.Add(Obj->ObjectName);
+				Index = Obj->OuterIndex;
+			}
+			else 
+			{
+				UE_LOGFMT_LOC(LogAssetRegistry, Warning,"SearchableNamePathInvalid",
+					"Package is unloadable: {FileName}. Reason: SearchableName path was not valid. Index: {Index} ",
+					("FileName", PackageFilename), ("Index", Pair.Key.ForDebugging()));
+				bSuccess = false;
+				break;
+			}
+		}
+
+		if (!Index.IsNull())
+		{
+			// Exited loop early above
+		}
+		else if (PathElements.Num() == 2)
+		{
+			OutSearchableNames.Add(FTopLevelAssetPath(PathElements[1], PathElements[0]), Pair.Value);
+		}
+		else if (PathElements.Num() == 1)
+		{
+			OutSearchableNames.Add(FTopLevelAssetPath(PathElements[0], FName()), Pair.Value);
+		}
+		else 
+		{
+			UE_LOGFMT_LOC(LogAssetRegistry, Warning,"SearchableNamePathTooLong",
+				"Package is unloadable: {FileName}. Reason: SearchableName object was not a top level asset. Path: {Path} ",
+				("FileName", PackageFilename), ("Path", PathElements));
+			bSuccess = false;
+		}
+	}
+
+	return bSuccess;
+}
+
+bool FPackageReader::ReadWorldTileInfo(FWorldTileInfo& OutInfo)
+{
+	if (PackageFileSummary.WorldTileInfoDataOffset == 0)
+	{
+		return true;
+	}
+
+	if (!StartSerializeSection(PackageFileSummary.WorldTileInfoDataOffset))
+	{
+		UE_LOGFMT_LOC(LogAssetRegistry, Warning,"WorldTileInfo_BadOffset",
+			"Package is unloadable: {FileName}. Reason: WorldTileInfoDataOffset {Offset} was invalid. ",
+			("FileName", PackageFilename), ("Offset", PackageFileSummary.WorldTileInfoDataOffset));
+		return false;
+	}
+
+	*this << OutInfo;
+	if (IsError())
+	{
+		UE_LOGFMT_LOC(LogAssetRegistry, Warning,"WorldTileInfo_SerializeFail",
+			"Package is unloadable: {FileName}. Reason: Failed to serialize WorldTileInfo. ",
+			("FileName", PackageFilename));
+		return false;
+	}
+
+	return true;
+}
+
+bool FPackageReader::ReadDataResourceMap(TArray<TTuple<FSoftObjectPath, FObjectDataResource>>& OutDataResources)
+{
+	// Data resources should never be associated with imports, but there may be a case in which an export is
+	// an inner of an import, e.g. external packages
+	if (!SerializeNameMap() || !SerializeExportMap() || !SerializeImportMap())
+	{
+		return false;
+	}
+
+	if (!StartSerializeSection(PackageFileSummary.DataResourceOffset))
+	{
+		UE_LOGFMT_LOC(LogAssetRegistry, Warning,"DataResourceMap_BadOffset",
+			"Package is unloadable: {FileName}. Reason: DataResourceOffset {Offset} was invalid. ",
+			("FileName", PackageFilename), ("Offset", PackageFileSummary.DataResourceOffset));
+		return false;
+	}
+
+	TArray<FObjectDataResource> DataResources;
+	FObjectDataResource::Serialize(*this, DataResources);
+	if (IsError())
+	{
+		UE_LOGFMT_LOC(LogAssetRegistry, Warning,"DataResourceMap_SerializeFail",
+			"Package is unloadable: {FileName}. Reason: Failed to serialize DataResourceMap. ",
+			("FileName", PackageFilename));
+		return false;
+	}
+
+	bool bSuccess = true;
+	for (int32 i=0; i < DataResources.Num(); ++i)
+	{
+		const FObjectDataResource& DataResource = DataResources[i];
+		TOptional<FSoftObjectPath> Path = LinkerIndexToSoftObjectPath(DataResource.OuterIndex);
+		if (Path.IsSet())
+		{
+			OutDataResources.Emplace(*Path, DataResource);
+		}
+		else
+		{
+			UE_LOGFMT_LOC(LogAssetRegistry, Warning,"DataResourceMap_BadObjectPath",
+				"Package is unloadable: {FileName}. Reason: Failed to get object path for data resource {Index}. ",
+				("FileName", PackageFilename), ("Index", i));
+			bSuccess = false;
+		}
+	}
+
+	return bSuccess;
+}
+
 bool FPackageReader::ReadAssetRegistryData(TArray<FAssetData*>& AssetDataList, bool& bOutIsCookedWithoutAssetData)
 {
 	bOutIsCookedWithoutAssetData = false;
@@ -494,6 +770,73 @@ bool FPackageReader::ReadAssetRegistryData(TArray<FAssetData*>& AssetDataList, b
 	return true;
 }
 
+#if WITH_METADATA
+bool FPackageReader::ReadMetaData(FMetaData& OutMetaData)
+{
+	if (PackageFileSummary.MetaDataOffset == 0)
+	{
+		// Empty
+		return true;
+	}
+
+	if (!SerializeNameMap() || !SerializeSoftObjectPathMap())
+	{
+		return false;
+	}
+
+	if (!StartSerializeSection(PackageFileSummary.MetaDataOffset))
+	{
+		UE_LOGFMT_LOC(LogAssetRegistry, Warning,"MetaDataInvalidOffset",
+			"Package is unloadable: {FileName}. Reason: Invalid metadata offset {Offset}",
+			("FileName", PackageFilename), ("Offset", PackageFileSummary.MetaDataOffset));
+		return false;
+	}
+
+	int32 NumObjectMetaDataMap = 0;
+	int32 NumRootMetaDataMap = 0;
+
+	*this << NumObjectMetaDataMap;
+	*this << NumRootMetaDataMap;
+	if (IsError())
+	{
+		UE_LOGFMT_LOC(LogAssetRegistry, Warning,"MetaDataFailedToReadCounts",
+			"Package is unloadable: {FileName}. Reason: Failed to read metadata counts",
+			("FileName", PackageFilename));
+		return false;
+	}
+
+	for (int32 i=0; i < NumObjectMetaDataMap; ++i)
+	{
+		TPair<FSoftObjectPath, TMap<FName, FString>> ObjectMetaData;
+		*this << ObjectMetaData;
+		if (IsError())
+		{
+			UE_LOGFMT_LOC(LogAssetRegistry, Warning,"MetaDataFailedToReadObjectMetaData",
+				"Package is unloadable: {FileName}. Reason: Failed to read object metadata. Index: {Index}",
+				("FileName", PackageFilename), ("Index", i));
+			return false;
+		}
+		OutMetaData.ObjectMetaDataMap.Add(ObjectMetaData.Key, ObjectMetaData.Value);
+	}
+
+	for (int32 i=0; i < NumRootMetaDataMap; ++i)
+	{
+		TPair<FName, FString> RootMetaData;
+		*this << RootMetaData;
+		if (IsError())
+		{
+			UE_LOGFMT_LOC(LogAssetRegistry, Warning,"MetaDataFailedToReadRootMetaData",
+				"Package is unloadable: {FileName}. Reason: Failed to read root metadata. Index: {Index}",
+				("FileName", PackageFilename), ("Index", i));
+			return false;
+		}
+		OutMetaData.RootMetaDataMap.Add(RootMetaData.Key, RootMetaData.Value);
+	}
+
+	return true;
+}
+#endif // WITH_METADATA
+
 bool FPackageReader::ReadLinkerObjects(TMap<FSoftObjectPath, FObjectData>& OutExports,
 	TMap<FSoftObjectPath, FObjectData>& OutImports, TMap<FName, bool>& OutSoftPackageReferences)
 {
@@ -588,6 +931,32 @@ bool FPackageReader::SerializeAssetRegistryDependencyData(TBitArray<>& OutImport
 	return true;
 }
 
+bool FPackageReader::ReadPackageTrailer(UE::FPackageTrailer& OutTrailer)
+{
+	if (PackageFileSummary.PayloadTocOffset == INDEX_NONE)
+	{
+		return true;
+	}
+
+	if (!StartSerializeSection(PackageFileSummary.PayloadTocOffset))
+	{
+		UE_LOGFMT_LOC(LogAssetRegistry, Warning,"PackageTrailer_BadOffset",
+			"Package is unloadable: {FileName}. Reason: Bad offset for PayloadTocOffset {Offset}. ",
+			("FileName", PackageFilename), ("Offset", PackageFileSummary.PayloadTocOffset));
+		return false;
+	}
+
+	if (!OutTrailer.TryLoad(*this))
+	{
+		UE_LOGFMT_LOC(LogAssetRegistry, Warning,"PackageTrailer_SerializeFail",
+			"Package is unloadable: {FileName}. Reason: Failed to deserialize trailer. ",
+			("FileName", PackageFilename));
+		return false;
+	}
+
+	return true;
+}
+
 bool FPackageReader::SerializePackageTrailer(FAssetPackageData& PackageData)
 {
 	if (!StartSerializeSection(PackageFileSummary.PayloadTocOffset))
@@ -1446,6 +1815,41 @@ FArchive& FPackageReader::operator<<( FName& Name )
 	return *this;
 }
 
+// Potentially look up by index into the soft object path table as with LinkerLoad
+FArchive& FPackageReader::operator<<(FSoftObjectPath& Value)
+{
+	FArchive& Ar = *this;
+	// Comment taken from FLinkerLoad:
+	// if we have items in the soft object path list consider soft object path saved as index into it.
+	// Otherwise deserialize them as usual, cooking for example might not want soft object path serialized as index
+	if (SoftObjectPathMap.Num() > 0)
+	{
+		int32 SoftObjectPathIndex = INDEX_NONE;
+		Ar << SoftObjectPathIndex;
+		if (SoftObjectPathMap.IsValidIndex(SoftObjectPathIndex))
+		{
+			Value = SoftObjectPathMap[SoftObjectPathIndex];
+
+#if WITH_EDITOR
+			Value.PostLoadPath(this);
+#endif // WITH_EDITOR
+		}
+		else
+		{
+			UE_LOGFMT_LOC(LogAssetRegistry, Warning,"SoftObjectPathInvalidIndex",
+				"Package is unloadable: {FileName}. Reason: Bad soft object path index. Index: {Index}/{Count}",
+				("FileName", PackageFilename), ("Index", SoftObjectPathIndex), ("Count", SoftObjectPathMap.Num()));
+			Value = FSoftObjectPath();
+			SetError();
+		}
+	}
+	else
+	{
+		FArchiveUObject::operator<<(Value);
+	}
+	return Ar;
+}
+
 void FPackageReader::ConvertLinkerTableToPaths(FName InPackageName, TArray<FObjectExport>& InExportMap, TArray<FObjectImport>& InImportMap,
 	TArray<FSoftObjectPath>& OutExports, TArray<FSoftObjectPath>& OutImports)
 {
@@ -1896,4 +2300,97 @@ namespace UE::AssetRegistry
 	}
 }
 
+TOptional<FSoftObjectPath> FPackageReader::LinkerIndexToSoftObjectPath(FPackageIndex Index)
+{
+	if (Index.IsNull())
+	{
+		return FSoftObjectPath();
+	}
+
+	// Most paths will be 2-3 elements but there may be some deeply nested subobjects in certain asset types.
+	TArray<FName, TInlineAllocator<8>> PathElements; 
+
+	auto GetObjectResource = [this](FPackageIndex I) -> FObjectResource*
+	{
+		if (I.IsImport())
+		{
+			if (ImportMap.IsValidIndex(I.ToImport()))
+			{
+				return &ImportMap[I.ToImport()];
+			}
+		}
+		else 
+		{
+			if (ExportMap.IsValidIndex(I.ToExport()))
+			{
+				return &ExportMap[I.ToExport()];
+			}
+		}
+		UE_LOGFMT_LOC(LogAssetRegistry, Warning,"ConstructObjectPath_PackageIndexOutOfRange",
+			"Package is unloadable: {FileName}. Reason: Bad object index {Index}. ",
+			("FileName", PackageFilename), ("Index", I.ForDebugging()));
+		return nullptr;
+	};
+
+	// Simultaneously check for loops with another pointer moving at twice the rate
+	FPackageIndex P1 = Index;
+	FPackageIndex P2 = Index;
+	while (!P1.IsNull())
+	{
+		FObjectResource* Obj = GetObjectResource(P1);
+		if (!Obj)
+		{
+			return {};
+		}
+		PathElements.Add(Obj->ObjectName);
+		P1 = Obj->OuterIndex;
+		if (P1.IsNull())
+		{
+			// If P1 terminates we know there is no loop and we have all path elements that we want.
+			break;
+		}
+
+		// Advance P2 by two steps, if it terminates we know there is no loop so we can stop advancing it
+		for (int32 i=0; i <2; ++i)
+		{
+			if (!P2.IsNull())
+			{
+				FObjectResource* Obj2 = GetObjectResource(P2);
+				if (!ensureMsgf(Obj2, TEXT("Unexpected null object resource traversing object path")))
+				{
+					return {};
+				}
+				P2 = Obj2->OuterIndex;
+			}
+		}
+
+		// If P1 did not terminate and P2 is at the same place as P1, we have a loop
+		if (P1 == P2)
+		{
+			UE_LOGFMT_LOC(LogAssetRegistry, Warning,"PackageIndexOuterLoop",
+				"Package is unloadable: {FileName}. Reason: Loop detected in outer chain of object index {Index}. ",
+				("FileName", PackageFilename), ("Index", Index.ForDebugging()));
+			return {};
+		}
+	}
+
+	// PathElements was populated from leaf to root, reverse it before handling any of these cases for clarity
+	Algo::Reverse(PathElements);
+	if (PathElements.Num() == 2)
+	{
+		return FSoftObjectPath(FTopLevelAssetPath(PathElements[0], PathElements[1]));
+	}
+	else if (PathElements.Num() == 1)
+	{
+		return FSoftObjectPath(FTopLevelAssetPath(PathElements[0], FName()));
+	}
+	else 
+	{
+		// Technically some delimiters ought to be : but we can't do that without using type information
+		TStringBuilder<512> Builder;
+		Builder.Join(PathElements, '.'); 
+		return FSoftObjectPath(Builder.ToView());
+	}
+}
+
 #undef LOCTEXT_NAMESPACE
