AssetRegistry - Remove deprecated function FAssetDataGatherer::GetAssetsToSave(TArrayView<const FString> SaveCacheLongPackageNameDirs, ...)
#rnx
#Synced-CL 52602039

diff --git a/Engine/Source/Runtime/AssetRegistry/Private/AssetDataGatherer.cpp b/Engine/Source/Runtime/AssetRegistry/Private/AssetDataGatherer.cpp
--- a/Engine/Source/Runtime/AssetRegistry/Private/AssetDataGatherer.cpp
+++ b/Engine/Source/Runtime/AssetRegistry/Private/AssetDataGatherer.cpp
@@ -5257,36 +5257,6 @@ void FAssetDataGatherer::TryReserveSaveCache(bool& bOutShouldSave, TArray<TPair<
 	ClearRequestAsyncCacheSave();
 }
 
-void FAssetDataGatherer::GetAssetsToSave(TArrayView<const FString> SaveCacheLongPackageNameDirs, TArray<TPair<FName,FDiskCachedAssetData*>>& OutAssetsToSave)
-{
-	CHECK_IS_LOCKED_CURRENT_THREAD(TickLock);
-
-	OutAssetsToSave.Reset();
-	if (SaveCacheLongPackageNameDirs.IsEmpty())
-	{
-		OutAssetsToSave.Reserve(NewCachedAssetDataMap.Num());
-		for (const TPair<FName, FDiskCachedAssetData*>& Pair : NewCachedAssetDataMap)
-		{
-			OutAssetsToSave.Add(Pair);
-		}
-	}
-	else
-	{
-		for (const TPair<FName, FDiskCachedAssetData*>& Pair : NewCachedAssetDataMap)
-		{
-			TStringBuilder<128> PackageNameStr;
-			Pair.Key.ToString(PackageNameStr);
-			if (Algo::AnyOf(SaveCacheLongPackageNameDirs, [PackageNameSV = FStringView(PackageNameStr)](const FString& SaveCacheLongPackageNameDir)
-			{
-				return FPathViews::IsParentPathOf(SaveCacheLongPackageNameDir, PackageNameSV);
-			}))
-			{
-				OutAssetsToSave.Add(Pair);
-			}
-		}
-	}
-}
-
 void FAssetDataGatherer::GetCacheAssetsToSave(TArray<TPair<FName,FDiskCachedAssetData*>>& OutAssetsToSave)
 {
 	CHECK_IS_LOCKED_CURRENT_THREAD(TickLock);
diff --git a/Engine/Source/Runtime/AssetRegistry/Private/AssetDataGatherer.h b/Engine/Source/Runtime/AssetRegistry/Private/AssetDataGatherer.h
--- a/Engine/Source/Runtime/AssetRegistry/Private/AssetDataGatherer.h
+++ b/Engine/Source/Runtime/AssetRegistry/Private/AssetDataGatherer.h
@@ -427,12 +427,6 @@ private:
 	 */
 	int64 SaveCacheFileInternal(const FString& CacheFilename,
 		const TArray<TPair<FName,FDiskCachedAssetData*>>& AssetsToSave);
-	/**
-	 * Get the list of FDiskCachedAssetData* that have been loaded in the gatherer, for saving into a cachefile.
-	 * Filters the list of assets by child paths of the elements in SaveCacheLongPackageNameDirs, if it is non-empty.
-	 */
-	void GetAssetsToSave(TArrayView<const FString> SaveCacheLongPackageNameDirs,
-		TArray<TPair<FName,FDiskCachedAssetData*>>& OutAssetsToSave);
 	/**
 	 * Get the list of FDiskCachedAssetData* for saving into the cache.
 	 * Includes both assets that were loaded in the gatherer and assets which were loaded from the cache and have not been pruned.
